Fix bad loctext name
#rnx

diff --git a/Engine/Plugins/Animation/GameplayInsights/Source/GameplayInsights/Private/AnimationTickRecordsTrack.cpp b/Engine/Plugins/Animation/GameplayInsights/Source/GameplayInsights/Private/AnimationTickRecordsTrack.cpp
--- a/Engine/Plugins/Animation/GameplayInsights/Source/GameplayInsights/Private/AnimationTickRecordsTrack.cpp
+++ b/Engine/Plugins/Animation/GameplayInsights/Source/GameplayInsights/Private/AnimationTickRecordsTrack.cpp
@@ -381,7 +381,7 @@ void FAnimationTickRecordsTrack::InitTooltip(FTooltipDrawState& Tooltip, const I
 		if(InMessage.bIsBlendSpace)
 		{
 			Tooltip.AddNameValueTextLine(LOCTEXT("BlendSpacePositionX", "X").ToString(), FText::AsNumber(InMessage.BlendSpacePositionX).ToString());
-			Tooltip.AddNameValueTextLine(LOCTEXT("BlendSpacePositionX", "Y").ToString(), FText::AsNumber(InMessage.BlendSpacePositionY).ToString());
+			Tooltip.AddNameValueTextLine(LOCTEXT("BlendSpacePositionY", "Y").ToString(), FText::AsNumber(InMessage.BlendSpacePositionY).ToString());
 		}
 		Tooltip.AddNameValueTextLine(LOCTEXT("PlaybackTime", "Playback Time").ToString(), FText::AsNumber(InMessage.PlaybackTime).ToString());
 		Tooltip.AddNameValueTextLine(LOCTEXT("RootMotionWeight", "Root Motion Weight").ToString(), FText::AsNumber(InMessage.RootMotionWeight).ToString());
