Moved Gameplay Insights to new track API

diff --git a/Engine/Plugins/Animation/GameplayInsights/Source/GameplayInsights/Private/AnimationSharedData.cpp b/Engine/Plugins/Animation/GameplayInsights/Source/GameplayInsights/Private/AnimationSharedData.cpp
--- a/Engine/Plugins/Animation/GameplayInsights/Source/GameplayInsights/Private/AnimationSharedData.cpp
+++ b/Engine/Plugins/Animation/GameplayInsights/Source/GameplayInsights/Private/AnimationSharedData.cpp
@@ -63,10 +63,12 @@ void FAnimationSharedData::Tick(Insights::ITimingViewSession& InTimingViewSessio
 
 	if(AnimationProvider && GameplayProvider)
 	{
+		Trace::FAnalysisSessionReadScope SessionReadScope(GetAnalysisSession());
+
 		// Add tracks for each tracked object's animation data
 		GameplayProvider->EnumerateObjects([this, &InTimingViewSession, &InAnalysisSession, &AnimationProvider, &GameplayProvider](const FObjectInfo& InObjectInfo)
 		{
-			FObjectEventsTrack* ObjectEventsTrack = GameplaySharedData.GetObjectEventsTrackForId(InTimingViewSession, InAnalysisSession, InObjectInfo);
+			TSharedRef<FObjectEventsTrack> ObjectEventsTrack = GameplaySharedData.GetObjectEventsTrackForId(InTimingViewSession, InAnalysisSession, InObjectInfo);
 
 			AnimationProvider->ReadSkeletalMeshPoseTimeline(InObjectInfo.Id, [this, &InObjectInfo, &ObjectEventsTrack, &InTimingViewSession](const IAnimationProvider::SkeletalMeshPoseTimeline& InTimeline)
 			{
@@ -83,14 +85,15 @@ void FAnimationSharedData::Tick(Insights::ITimingViewSession& InTimingViewSessio
 					return false;
 				};
 
-				FSkeletalMeshPoseTrack* SkeletalMeshPoseTrack = static_cast<FSkeletalMeshPoseTrack*>(ObjectEventsTrack->GetGameplayTrack().FindChildTrack(InObjectInfo.Id, FindSkeletalMeshPoseTrack));
-				if(SkeletalMeshPoseTrack == nullptr)
+				TSharedPtr<FSkeletalMeshPoseTrack> SkeletalMeshPoseTrack = StaticCastSharedPtr<FSkeletalMeshPoseTrack>(ObjectEventsTrack->GetGameplayTrack().FindChildTrack(InObjectInfo.Id, FindSkeletalMeshPoseTrack));
+				if(!SkeletalMeshPoseTrack.IsValid())
 				{
-					SkeletalMeshPoseTrack = new FSkeletalMeshPoseTrack(*this, InObjectInfo.Id, InObjectInfo.Name);
+					SkeletalMeshPoseTrack = MakeShared<FSkeletalMeshPoseTrack>(*this, InObjectInfo.Id, InObjectInfo.Name);
 					SkeletalMeshPoseTrack->SetVisibilityFlag(bAnimationTracksEnabled);
-					SkeletalMeshPoseTracks.Add(SkeletalMeshPoseTrack);
+					SkeletalMeshPoseTracks.Add(SkeletalMeshPoseTrack.ToSharedRef());
 
-					InTimingViewSession.AddTimingEventsTrack(SkeletalMeshPoseTrack);
+					InTimingViewSession.AddScrollableTrack(SkeletalMeshPoseTrack);
+					GameplaySharedData.InvalidateObjectTracksOrder();
 
 					ObjectEventsTrack->GetGameplayTrack().AddChildTrack(SkeletalMeshPoseTrack->GetGameplayTrack());
 				}
@@ -112,16 +115,17 @@ void FAnimationSharedData::Tick(Insights::ITimingViewSession& InTimingViewSessio
 					return false;
 				};
 
-				FAnimationTickRecordsTrack* AnimationTickRecordsTrack = static_cast<FAnimationTickRecordsTrack*>(ObjectEventsTrack->GetGameplayTrack().FindChildTrack(InObjectInfo.Id, FindTickRecordTrackWithAssetId));
-				if(AnimationTickRecordsTrack == nullptr)
+				TSharedPtr<FAnimationTickRecordsTrack> AnimationTickRecordsTrack = StaticCastSharedPtr<FAnimationTickRecordsTrack>(ObjectEventsTrack->GetGameplayTrack().FindChildTrack(InObjectInfo.Id, FindTickRecordTrackWithAssetId));
+				if(!AnimationTickRecordsTrack.IsValid())
 				{
 					const FObjectInfo* AssetObjectInfo = GameplayProvider->FindObjectInfo(InAssetId);
 					FString AssetName = AssetObjectInfo ? AssetObjectInfo->Name : LOCTEXT("UnknownAsset", "Unknown").ToString();
-					AnimationTickRecordsTrack = new FAnimationTickRecordsTrack(*this, InObjectInfo.Id, InAssetId, *AssetName);
+					AnimationTickRecordsTrack = MakeShared<FAnimationTickRecordsTrack>(*this, InObjectInfo.Id, InAssetId, *AssetName);
 					AnimationTickRecordsTrack->SetVisibilityFlag(bAnimationTracksEnabled);
-					AnimationTickRecordsTracks.Add(AnimationTickRecordsTrack);
+					AnimationTickRecordsTracks.Add(AnimationTickRecordsTrack.ToSharedRef());
 
-					InTimingViewSession.AddTimingEventsTrack(AnimationTickRecordsTrack);
+					InTimingViewSession.AddScrollableTrack(AnimationTickRecordsTrack);
+					GameplaySharedData.InvalidateObjectTracksOrder();
 
 					ObjectEventsTrack->GetGameplayTrack().AddChildTrack(AnimationTickRecordsTrack->GetGameplayTrack());
 				}
@@ -149,12 +153,12 @@ void FAnimationSharedData::ToggleAnimationTracks()
 {
 	bAnimationTracksEnabled = !bAnimationTracksEnabled;
 
-	for(FSkeletalMeshPoseTrack* PoseTrack : SkeletalMeshPoseTracks)
+	for(TSharedRef<FSkeletalMeshPoseTrack> PoseTrack : SkeletalMeshPoseTracks)
 	{
 		PoseTrack->SetVisibilityFlag(bAnimationTracksEnabled);
 	}
 
-	for(FAnimationTickRecordsTrack* TickRecordTrack : AnimationTickRecordsTracks)
+	for(TSharedRef<FAnimationTickRecordsTrack> TickRecordTrack : AnimationTickRecordsTracks)
 	{
 		TickRecordTrack->SetVisibilityFlag(bAnimationTracksEnabled);
 	}
@@ -165,20 +169,38 @@ bool FAnimationSharedData::AreAnimationTracksEnabled() const
 	return bAnimationTracksEnabled;
 }
 
-void FAnimationSharedData::OnSelectedEventChanged(const FTimingEvent& InEvent)
+void FAnimationSharedData::OnSelectedEventChanged(const TSharedPtr<const ITimingEvent> InEvent)
 {
 	bSelectedEventValid = InEvent.IsValid();
-	SelectedEventTrack = InEvent.Track;
-	SelectedEventStartTime = InEvent.StartTime;
-	SelectedEventEndTime = InEvent.EndTime;
+	if(bSelectedEventValid)
+	{
+		SelectedEventTrack = InEvent->GetTrack();
+		SelectedEventStartTime = InEvent->GetStartTime();
+		SelectedEventEndTime = InEvent->GetEndTime();
+	}
+	else
+	{
+		SelectedEventTrack = nullptr;
+		SelectedEventStartTime = 0.0;
+		SelectedEventEndTime = 0.0;
+	}
 }
 
-void FAnimationSharedData::OnHoveredEventChanged(const FTimingEvent& InEvent)
+void FAnimationSharedData::OnHoveredEventChanged(const TSharedPtr<const ITimingEvent> InEvent)
 {
 	bHoveredEventValid = InEvent.IsValid();
-	HoveredEventTrack = InEvent.Track;
-	HoveredEventStartTime = InEvent.StartTime;
-	HoveredEventEndTime = InEvent.EndTime;
+	if(bHoveredEventValid)
+	{
+		HoveredEventTrack = InEvent->GetTrack();
+		HoveredEventStartTime = InEvent->GetStartTime();
+		HoveredEventEndTime = InEvent->GetEndTime();
+	}
+	else
+	{
+		HoveredEventTrack = nullptr;
+		HoveredEventStartTime = 0.0;
+		HoveredEventEndTime = 0.0;
+	}
 }
 
 void FAnimationSharedData::OnTimeMarkerChanged(Insights::ETimeChangedFlags InFlags, double InTimeMarker)
@@ -198,7 +220,7 @@ void FAnimationSharedData::OnSelectionChanged(Insights::ETimeChangedFlags InFlag
 
 void FAnimationSharedData::DrawPoses(ULineBatchComponent* InLineBatcher)
 {
-	for(FSkeletalMeshPoseTrack* PoseTrack : SkeletalMeshPoseTracks)
+	for(TSharedRef<FSkeletalMeshPoseTrack> PoseTrack : SkeletalMeshPoseTracks)
 	{
 		if(PoseTrack->IsVisible())
 		{
diff --git a/Engine/Plugins/Animation/GameplayInsights/Source/GameplayInsights/Private/AnimationSharedData.h b/Engine/Plugins/Animation/GameplayInsights/Source/GameplayInsights/Private/AnimationSharedData.h
--- a/Engine/Plugins/Animation/GameplayInsights/Source/GameplayInsights/Private/AnimationSharedData.h
+++ b/Engine/Plugins/Animation/GameplayInsights/Source/GameplayInsights/Private/AnimationSharedData.h
@@ -39,8 +39,8 @@ public:
 private:
 	// UI handlers
 	void ToggleAnimationTracks();
-	void OnSelectedEventChanged(const FTimingEvent& InEvent);
-	void OnHoveredEventChanged(const FTimingEvent& InEvent);
+	void OnSelectedEventChanged(const TSharedPtr<const ITimingEvent> InEvent);
+	void OnHoveredEventChanged(const TSharedPtr<const ITimingEvent> InEvent);
 	void OnTimeMarkerChanged(Insights::ETimeChangedFlags InFlags, double InTimeMarker);
 	void OnSelectionChanged(Insights::ETimeChangedFlags InFlags, double InStartTime, double InEndTime);
 
@@ -52,8 +52,8 @@ private:
 	const Trace::IAnalysisSession* AnalysisSession;
 
 	// All the tracks we manage
-	TArray<FSkeletalMeshPoseTrack*> SkeletalMeshPoseTracks;
-	TArray<FAnimationTickRecordsTrack*> AnimationTickRecordsTracks;
+	TArray<TSharedRef<FSkeletalMeshPoseTrack>> SkeletalMeshPoseTracks;
+	TArray<TSharedRef<FAnimationTickRecordsTrack>> AnimationTickRecordsTracks;
 
 	// Delegate handles for hooks into the timing view
 	FDelegateHandle SelectedEventChangedHandle;
@@ -62,8 +62,8 @@ private:
 	FDelegateHandle SelectionChangedHandle;
 
 	/** Selected/hovered tracks */
-	const FTimingEventsTrack* SelectedEventTrack;
-	const FTimingEventsTrack* HoveredEventTrack;
+	TSharedPtr<const FBaseTimingTrack> SelectedEventTrack;
+	TSharedPtr<const FBaseTimingTrack> HoveredEventTrack;
 
 	/** Various times and ranges */
 	double SelectedEventStartTime;
diff --git a/Engine/Plugins/Animation/GameplayInsights/Source/GameplayInsights/Private/AnimationTickRecordsTrack.cpp b/Engine/Plugins/Animation/GameplayInsights/Source/GameplayInsights/Private/AnimationTickRecordsTrack.cpp
--- a/Engine/Plugins/Animation/GameplayInsights/Source/GameplayInsights/Private/AnimationTickRecordsTrack.cpp
+++ b/Engine/Plugins/Animation/GameplayInsights/Source/GameplayInsights/Private/AnimationTickRecordsTrack.cpp
@@ -1,7 +1,6 @@
 // Copyright 1998-2019 Epic Games, Inc. All Rights Reserved.
 
 #include "AnimationTickRecordsTrack.h"
-#include "Insights/ITimingViewDrawHelper.h"
 #include "GameplayProvider.h"
 #include "AnimationProvider.h"
 #include "Insights/ViewModels/TimingTrackViewport.h"
@@ -9,6 +8,7 @@
 #include "Insights/ViewModels/TooltipDrawState.h"
 #include "AnimationSharedData.h"
 #include "Insights/ViewModels/TimingEventSearch.h"
+#include "Insights/ViewModels/TooltipDrawState.h"
 
 #define LOCTEXT_NAMESPACE "AnimationTickRecordsTrack"
 
@@ -22,45 +22,38 @@ FAnimationTickRecordsTrack::FAnimationTickRecordsTrack(const FAnimationSharedDat
 {
 }
 
-void FAnimationTickRecordsTrack::Draw(ITimingViewDrawHelper& Helper) const
+void FAnimationTickRecordsTrack::BuildDrawState(ITimingEventsTrackDrawStateBuilder& Builder, const ITimingTrackUpdateContext& Context)
 {
-	FAnimationTickRecordsTrack& Track = *const_cast<FAnimationTickRecordsTrack*>(this);
+	const FGameplayProvider* GameplayProvider = SharedData.GetAnalysisSession().ReadProvider<FGameplayProvider>(FGameplayProvider::ProviderName);
+	const FAnimationProvider* AnimationProvider = SharedData.GetAnalysisSession().ReadProvider<FAnimationProvider>(FAnimationProvider::ProviderName);
 
-	if (Helper.BeginTimeline(Track))
+	if(GameplayProvider && AnimationProvider)
 	{
-		const FGameplayProvider* GameplayProvider = SharedData.GetAnalysisSession().ReadProvider<FGameplayProvider>(FGameplayProvider::ProviderName);
-		const FAnimationProvider* AnimationProvider = SharedData.GetAnalysisSession().ReadProvider<FAnimationProvider>(FAnimationProvider::ProviderName);
+		Trace::FAnalysisSessionReadScope SessionReadScope(SharedData.GetAnalysisSession());
 
-		if(GameplayProvider && AnimationProvider)
+		AnimationProvider->ReadTickRecordTimeline(GetGameplayTrack().GetObjectId(), GetAssetId(), [this, &GameplayProvider, &AnimationProvider, &Context, &Builder](const FAnimationProvider::TickRecordTimeline& InTimeline)
 		{
-			Trace::FAnalysisSessionReadScope SessionReadScope(SharedData.GetAnalysisSession());
-
-			AnimationProvider->ReadTickRecordTimeline(GetGameplayTrack().GetObjectId(), GetAssetId(), [this, &GameplayProvider, &AnimationProvider, &Helper](const FAnimationProvider::TickRecordTimeline& InTimeline)
+			auto DrawEvents = [this, &Builder, &GameplayProvider](double InStartTime, double InEndTime, uint32 InDepth, const FTickRecordMessage& InMessage)
 			{
-				auto DrawEvents = [this, &Helper, &GameplayProvider](double InStartTime, double InEndTime, uint32 InDepth, const FTickRecordMessage& InMessage)
-				{
-					Helper.AddEvent(InStartTime, InEndTime, 0, *GetName());
-				};
+				Builder.AddEvent(InStartTime, InEndTime, 0, *GetName());
+			};
 
-				if (FTimingEventsTrack::bUseDownSampling)
-				{
-					const double SecondsPerPixel = 1.0 / Helper.GetViewport().GetScaleX();
-					InTimeline.EnumerateEventsDownSampled(Helper.GetViewport().GetStartTime(), Helper.GetViewport().GetEndTime(), SecondsPerPixel, DrawEvents);
-				}
-				else
-				{
-					InTimeline.EnumerateEvents(Helper.GetViewport().GetStartTime(), Helper.GetViewport().GetEndTime(), DrawEvents);
-				}
-			});
-		}
-
-		Helper.EndTimeline(Track);
+			if (FTimingEventsTrack::bUseDownSampling)
+			{
+				const double SecondsPerPixel = 1.0 / Context.GetViewport().GetScaleX();
+				InTimeline.EnumerateEventsDownSampled(Context.GetViewport().GetStartTime(), Context.GetViewport().GetEndTime(), SecondsPerPixel, DrawEvents);
+			}
+			else
+			{
+				InTimeline.EnumerateEvents(Context.GetViewport().GetStartTime(), Context.GetViewport().GetEndTime(), DrawEvents);
+			}
+		});
 	}
 }
 
-void FAnimationTickRecordsTrack::InitTooltip(FTooltipDrawState& Tooltip, const FTimingEvent& HoveredTimingEvent) const
+void FAnimationTickRecordsTrack::InitTooltip(FTooltipDrawState& Tooltip, const ITimingEvent& HoveredTimingEvent) const
 {
-	FTimingEventSearchParameters SearchParameters(HoveredTimingEvent.StartTime, HoveredTimingEvent.EndTime, ETimingEventSearchFlags::StopAtFirstMatch);
+	FTimingEventSearchParameters SearchParameters(HoveredTimingEvent.GetStartTime(), HoveredTimingEvent.GetEndTime(), ETimingEventSearchFlags::StopAtFirstMatch);
 
 	FindTickRecordMessage(SearchParameters, [this, &Tooltip](double InFoundStartTime, double InFoundEndTime, uint32 InFoundDepth, const FTickRecordMessage& InMessage)
 	{
@@ -79,17 +72,21 @@ void FAnimationTickRecordsTrack::InitTooltip(FTooltipDrawState& Tooltip, const F
 	});
 }
 
-bool FAnimationTickRecordsTrack::SearchTimingEvent(const FTimingEventSearchParameters& InSearchParameters, FTimingEvent& InOutTimingEvent) const
+const TSharedPtr<const ITimingEvent> FAnimationTickRecordsTrack::SearchEvent(const FTimingEventSearchParameters& InSearchParameters) const
 {
-	return FindTickRecordMessage(InSearchParameters, [this, &InOutTimingEvent](double InFoundStartTime, double InFoundEndTime, uint32 InFoundDepth, const FTickRecordMessage& InFoundMessage)
+	TSharedPtr<const ITimingEvent> FoundEvent;
+
+	FindTickRecordMessage(InSearchParameters, [this, &FoundEvent](double InFoundStartTime, double InFoundEndTime, uint32 InFoundDepth, const FTickRecordMessage& InFoundMessage)
 	{
-		InOutTimingEvent = FTimingEvent(this, InFoundStartTime, InFoundEndTime, InFoundDepth);
+		FoundEvent = MakeShared<const FTimingEvent>(SharedThis(this), InFoundStartTime, InFoundEndTime, InFoundDepth);
 	});
+
+	return FoundEvent;
 }
 
-bool FAnimationTickRecordsTrack::FindTickRecordMessage(const FTimingEventSearchParameters& InParameters, TFunctionRef<void(double, double, uint32, const FTickRecordMessage&)> InFoundPredicate) const
+void FAnimationTickRecordsTrack::FindTickRecordMessage(const FTimingEventSearchParameters& InParameters, TFunctionRef<void(double, double, uint32, const FTickRecordMessage&)> InFoundPredicate) const
 {
-	return TTimingEventSearch<FTickRecordMessage>::Search(
+	TTimingEventSearch<FTickRecordMessage>::Search(
 		InParameters,
 
 		[this](TTimingEventSearch<FTickRecordMessage>::FContext& InContext)
diff --git a/Engine/Plugins/Animation/GameplayInsights/Source/GameplayInsights/Private/AnimationTickRecordsTrack.h b/Engine/Plugins/Animation/GameplayInsights/Source/GameplayInsights/Private/AnimationTickRecordsTrack.h
--- a/Engine/Plugins/Animation/GameplayInsights/Source/GameplayInsights/Private/AnimationTickRecordsTrack.h
+++ b/Engine/Plugins/Animation/GameplayInsights/Source/GameplayInsights/Private/AnimationTickRecordsTrack.h
@@ -17,16 +17,16 @@ public:
 
 	FAnimationTickRecordsTrack(const FAnimationSharedData& InSharedData, uint64 InObjectId, uint64 InAssetId, const TCHAR* InName);
 
-	virtual void Draw(ITimingViewDrawHelper& Helper) const override;
-	virtual void InitTooltip(FTooltipDrawState& Tooltip, const FTimingEvent& HoveredTimingEvent) const override;
-	virtual bool SearchTimingEvent(const FTimingEventSearchParameters& InSearchParameters, FTimingEvent& InOutTimingEvent) const override;
+	virtual void BuildDrawState(ITimingEventsTrackDrawStateBuilder& Builder, const ITimingTrackUpdateContext& Context) override;
+	virtual void InitTooltip(FTooltipDrawState& Tooltip, const ITimingEvent& HoveredTimingEvent) const override;
+	virtual const TSharedPtr<const ITimingEvent> SearchEvent(const FTimingEventSearchParameters& InSearchParameters) const override;
 
 	/** Get the asset ID that this track uses */
 	uint64 GetAssetId() const { return AssetId; }
 
 private:
 	// Helper function used to find a tick record
-	bool FindTickRecordMessage(const FTimingEventSearchParameters& InParameters, TFunctionRef<void(double, double, uint32, const FTickRecordMessage&)> InFoundPredicate) const;
+	void FindTickRecordMessage(const FTimingEventSearchParameters& InParameters, TFunctionRef<void(double, double, uint32, const FTickRecordMessage&)> InFoundPredicate) const;
 
 private:
 	/** The shared data */
diff --git a/Engine/Plugins/Animation/GameplayInsights/Source/GameplayInsights/Private/GameplayProvider.cpp b/Engine/Plugins/Animation/GameplayInsights/Source/GameplayInsights/Private/GameplayProvider.cpp
--- a/Engine/Plugins/Animation/GameplayInsights/Source/GameplayInsights/Private/GameplayProvider.cpp
+++ b/Engine/Plugins/Animation/GameplayInsights/Source/GameplayInsights/Private/GameplayProvider.cpp
@@ -41,6 +41,8 @@ bool FGameplayProvider::ReadObjectEvent(uint64 InObjectId, uint64 InMessageId, T
 
 void FGameplayProvider::EnumerateObjects(TFunctionRef<void(const FObjectInfo&)> Callback) const
 {
+	Session.ReadAccessCheck();
+
 	for(const FObjectInfo& ObjectInfo : ObjectInfos)
 	{
 		Callback(ObjectInfo);
@@ -49,6 +51,8 @@ void FGameplayProvider::EnumerateObjects(TFunctionRef<void(const FObjectInfo&)>
 
 const FClassInfo* FGameplayProvider::FindClassInfo(uint64 InClassId) const
 {
+	Session.ReadAccessCheck();
+
 	const int32* ClassIndex = ClassIdToIndexMap.Find(InClassId);
 	if(ClassIndex != nullptr)
 	{
@@ -60,6 +64,8 @@ const FClassInfo* FGameplayProvider::FindClassInfo(uint64 InClassId) const
 
 const FObjectInfo* FGameplayProvider::FindObjectInfo(uint64 InObjectId) const
 {
+	Session.ReadAccessCheck();
+
 	const int32* ObjectIndex = ObjectIdToIndexMap.Find(InObjectId);
 	if(ObjectIndex != nullptr)
 	{
diff --git a/Engine/Plugins/Animation/GameplayInsights/Source/GameplayInsights/Private/GameplaySharedData.cpp b/Engine/Plugins/Animation/GameplayInsights/Source/GameplayInsights/Private/GameplaySharedData.cpp
--- a/Engine/Plugins/Animation/GameplayInsights/Source/GameplayInsights/Private/GameplaySharedData.cpp
+++ b/Engine/Plugins/Animation/GameplayInsights/Source/GameplayInsights/Private/GameplaySharedData.cpp
@@ -12,6 +12,7 @@
 
 FGameplaySharedData::FGameplaySharedData()
 	: AnalysisSession(nullptr)
+	, bObjectTracksDirty(false)
 	, bObjectTracksEnabled(false)
 {
 }
@@ -26,36 +27,38 @@ void FGameplaySharedData::OnEndSession(Insights::ITimingViewSession& InTimingVie
 	ObjectTracks.Reset();
 }
 
-FObjectEventsTrack* FGameplaySharedData::GetObjectEventsTrackForId(Insights::ITimingViewSession& InTimingViewSession, const Trace::IAnalysisSession& InAnalysisSession, const FObjectInfo& InObjectInfo)
+TSharedRef<FObjectEventsTrack> FGameplaySharedData::GetObjectEventsTrackForId(Insights::ITimingViewSession& InTimingViewSession, const Trace::IAnalysisSession& InAnalysisSession, const FObjectInfo& InObjectInfo)
 {
 	const FGameplayProvider* GameplayProvider = InAnalysisSession.ReadProvider<FGameplayProvider>(FGameplayProvider::ProviderName);
 	check(GameplayProvider);
 
-	FObjectEventsTrack* LeafObjectEventsTrack = ObjectTracks.FindRef(InObjectInfo.Id);
-	if(LeafObjectEventsTrack == nullptr)
+	TSharedPtr<FObjectEventsTrack> LeafObjectEventsTrack = ObjectTracks.FindRef(InObjectInfo.Id);
+	if(!LeafObjectEventsTrack.IsValid())
 	{
-		LeafObjectEventsTrack = new FObjectEventsTrack(*this, InObjectInfo.Id, InObjectInfo.Name);
+		LeafObjectEventsTrack = MakeShared<FObjectEventsTrack>(*this, InObjectInfo.Id, InObjectInfo.Name);
 		LeafObjectEventsTrack->SetVisibilityFlag(bObjectTracksEnabled);
 		ObjectTracks.Add(InObjectInfo.Id, LeafObjectEventsTrack);
 
-		InTimingViewSession.AddTimingEventsTrack(LeafObjectEventsTrack);
+		InTimingViewSession.AddScrollableTrack(LeafObjectEventsTrack);
+		InvalidateObjectTracksOrder();
 	}
 
 	// Fill the outer chain
 	if(InObjectInfo.OuterId != 0)
 	{
-		FObjectEventsTrack* ObjectEventsTrack = LeafObjectEventsTrack;
+		TSharedPtr<FObjectEventsTrack> ObjectEventsTrack = LeafObjectEventsTrack;
 		uint64 OuterId = InObjectInfo.OuterId;
 		while(const FObjectInfo* OuterInfo = GameplayProvider->FindObjectInfo(OuterId))
 		{
-			FObjectEventsTrack* OuterEventsTrack = ObjectTracks.FindRef(OuterId);
-			if(OuterEventsTrack == nullptr)
+			TSharedPtr<FObjectEventsTrack> OuterEventsTrack = ObjectTracks.FindRef(OuterId);
+			if(!OuterEventsTrack.IsValid())
 			{
-				OuterEventsTrack = new FObjectEventsTrack(*this, OuterId, OuterInfo->Name);
+				OuterEventsTrack = MakeShared<FObjectEventsTrack>(*this, OuterId, OuterInfo->Name);
 				OuterEventsTrack->SetVisibilityFlag(bObjectTracksEnabled);
 				ObjectTracks.Add(OuterId, OuterEventsTrack);
 
-				InTimingViewSession.AddTimingEventsTrack(OuterEventsTrack);
+				InTimingViewSession.AddScrollableTrack(OuterEventsTrack);
+				InvalidateObjectTracksOrder();
 			}
 
 			// setup hierarchy
@@ -69,10 +72,10 @@ FObjectEventsTrack* FGameplaySharedData::GetObjectEventsTrackForId(Insights::ITi
 		}
 	}
 
-	return LeafObjectEventsTrack;
+	return LeafObjectEventsTrack.ToSharedRef();
 }
 
-static void UpdateTrackOrderRecursive(FObjectEventsTrack* InTrack, int32& InOutOrder)
+static void UpdateTrackOrderRecursive(TSharedRef<FObjectEventsTrack> InTrack, int32& InOutOrder)
 {
 	// recurse down object-track children, then non-object track leaf tracks to set 
 	// overall ordering based on depth-first traversal of the hierarchy
@@ -81,17 +84,17 @@ static void UpdateTrackOrderRecursive(FObjectEventsTrack* InTrack, int32& InOutO
 
 	for(FGameplayTrack* ChildTrack : InTrack->GetGameplayTrack().GetChildTracks())
 	{
-		if(ChildTrack->GetTimingTrack().GetType() == FObjectEventsTrack::TypeName)
+		if(ChildTrack->GetTimingTrack()->GetType() == FObjectEventsTrack::TypeName)
 		{
-			UpdateTrackOrderRecursive(static_cast<FObjectEventsTrack*>(&ChildTrack->GetTimingTrack()), InOutOrder);
+			UpdateTrackOrderRecursive(StaticCastSharedRef<FObjectEventsTrack>(ChildTrack->GetTimingTrack()), InOutOrder);
 		}
 	}
 
 	for(FGameplayTrack* ChildTrack : InTrack->GetGameplayTrack().GetChildTracks())
 	{
-		if(ChildTrack->GetTimingTrack().GetType() != FObjectEventsTrack::TypeName)
+		if(ChildTrack->GetTimingTrack()->GetType() != FObjectEventsTrack::TypeName)
 		{
-			ChildTrack->GetTimingTrack().SetOrder(InOutOrder++);
+			ChildTrack->GetTimingTrack()->SetOrder(InOutOrder++);
 		}
 	}
 }
@@ -104,10 +107,12 @@ void FGameplaySharedData::Tick(Insights::ITimingViewSession& InTimingViewSession
 
 	if(GameplayProvider)
 	{
+		Trace::FAnalysisSessionReadScope SessionReadScope(GetAnalysisSession());
+
 		// Add a track for each tracked object
 		GameplayProvider->EnumerateObjects([this, &InTimingViewSession, &InAnalysisSession, &GameplayProvider](const FObjectInfo& InObjectInfo)
 		{
-			FObjectEventsTrack* ObjectEventsTrack = GetObjectEventsTrackForId(InTimingViewSession, InAnalysisSession, InObjectInfo);
+			TSharedPtr<FObjectEventsTrack> ObjectEventsTrack = GetObjectEventsTrackForId(InTimingViewSession, InAnalysisSession, InObjectInfo);
 			ObjectEventsTrack->SetVisibilityFlag(false);
 
 			GameplayProvider->ReadObjectEventsTimeline(InObjectInfo.Id, [this, &ObjectEventsTrack](const IGameplayProvider::ObjectEventsTimeline& InTimeline)
@@ -115,6 +120,13 @@ void FGameplaySharedData::Tick(Insights::ITimingViewSession& InTimingViewSession
 				ObjectEventsTrack->SetVisibilityFlag(bObjectTracksEnabled);
 			});
 		});
+
+		if(bObjectTracksDirty)
+		{
+			SortTracks();
+			InTimingViewSession.InvalidateScrollableTracksOrder();
+			bObjectTracksDirty = false;
+		}
 	}
 }
 
@@ -133,30 +145,32 @@ void FGameplaySharedData::ExtendFilterMenu(FMenuBuilder& InMenuBuilder)
 	);
 }
 
-void FGameplaySharedData::OnTracksChanged(int32& InOutOrder)
+void FGameplaySharedData::SortTracks()
 {
+	int32 Order = 10000;
+
 	// Find current roots
-	static TArray<FObjectEventsTrack*> Roots;
+	static TArray<TSharedRef<FObjectEventsTrack>> Roots;
 	check(Roots.Num() == 0);
 
 	for(auto ObjectTrackPair : ObjectTracks)
 	{
 		if(ObjectTrackPair.Value->GetGameplayTrack().GetParentTrack() == nullptr)
 		{
-			Roots.Add(ObjectTrackPair.Value);
+			Roots.Add(ObjectTrackPair.Value.ToSharedRef());
 		}
 	}
 
 	// Sort roots alphabetically
-	Algo::Sort(Roots, [](FObjectEventsTrack* InTrack0, FObjectEventsTrack* InTrack1)
+	Algo::Sort(Roots, [](TSharedRef<FObjectEventsTrack> InTrack0, TSharedRef<FObjectEventsTrack> InTrack1)
 	{
 		return InTrack0->GetName() < InTrack1->GetName();
 	});
 
 	// update ordering
-	for(FObjectEventsTrack* RootTrack : Roots)
+	for(TSharedRef<FObjectEventsTrack> RootTrack : Roots)
 	{
-		UpdateTrackOrderRecursive(RootTrack, InOutOrder);
+		UpdateTrackOrderRecursive(RootTrack, Order);
 	}
 
 	Roots.Reset();
@@ -166,9 +180,9 @@ void FGameplaySharedData::ToggleGameplayTracks()
 {
 	bObjectTracksEnabled = !bObjectTracksEnabled;
 
-	for(TPair<uint64, FObjectEventsTrack*> TrackPair : ObjectTracks)
+	for(auto ObjectTrackPair : ObjectTracks)
 	{
-		TrackPair.Value->SetVisibilityFlag(bObjectTracksEnabled);
+		ObjectTrackPair.Value->SetVisibilityFlag(bObjectTracksEnabled);
 	}
 }
 
diff --git a/Engine/Plugins/Animation/GameplayInsights/Source/GameplayInsights/Private/GameplaySharedData.h b/Engine/Plugins/Animation/GameplayInsights/Source/GameplayInsights/Private/GameplaySharedData.h
--- a/Engine/Plugins/Animation/GameplayInsights/Source/GameplayInsights/Private/GameplaySharedData.h
+++ b/Engine/Plugins/Animation/GameplayInsights/Source/GameplayInsights/Private/GameplaySharedData.h
@@ -20,28 +20,36 @@ public:
 	void OnEndSession(Insights::ITimingViewSession& InTimingViewSession);
 	void Tick(Insights::ITimingViewSession& InTimingViewSession, const Trace::IAnalysisSession& InAnalysisSession);
 	void ExtendFilterMenu(FMenuBuilder& InMenuBuilder);
-	void OnTracksChanged(int32& InOutOrder);
 
 	// Helper function. Builds object track hierarchy on-demand and returns a track for the supplied object info.
-	FObjectEventsTrack* GetObjectEventsTrackForId(Insights::ITimingViewSession& InTimingViewSession, const Trace::IAnalysisSession& InAnalysisSession, const FObjectInfo& InObjectInfo);
+	TSharedRef<FObjectEventsTrack> GetObjectEventsTrackForId(Insights::ITimingViewSession& InTimingViewSession, const Trace::IAnalysisSession& InAnalysisSession, const FObjectInfo& InObjectInfo);
 
 	// Check whether gameplay tacks are enabled
 	bool AreGameplayTracksEnabled() const;
 
+	// Invalidate object trcks order, so they get re-sorted next tick
+	void InvalidateObjectTracksOrder() { bObjectTracksDirty = true; }
+
 	// Get the last cached analysis session
 	const Trace::IAnalysisSession& GetAnalysisSession() const { return *AnalysisSession; }
 
 private:
+	// Re-sort tracks if trrack ordering has changed
+	void SortTracks();
+
 	// UI handlers
 	void ToggleGameplayTracks();
 
 private:
 	// Track for each tracked object, mapped from Object ID -> track
-	TMap<uint64, FObjectEventsTrack*> ObjectTracks;
+	TMap<uint64, TSharedPtr<FObjectEventsTrack>> ObjectTracks;
 
 	// Cached analysis session, set in Tick()
 	const Trace::IAnalysisSession* AnalysisSession;
 
+	// Dirty flag for adding object tracks, used to trigger re-sorting
+	bool bObjectTracksDirty;
+
 	// Whether all of our object tracks are enabled
 	bool bObjectTracksEnabled;
 };
\ No newline at end of file
diff --git a/Engine/Plugins/Animation/GameplayInsights/Source/GameplayInsights/Private/GameplayTimingViewExtender.cpp b/Engine/Plugins/Animation/GameplayInsights/Source/GameplayInsights/Private/GameplayTimingViewExtender.cpp
--- a/Engine/Plugins/Animation/GameplayInsights/Source/GameplayInsights/Private/GameplayTimingViewExtender.cpp
+++ b/Engine/Plugins/Animation/GameplayInsights/Source/GameplayInsights/Private/GameplayTimingViewExtender.cpp
@@ -61,15 +61,6 @@ void FGameplayTimingViewExtender::ExtendFilterMenu(Insights::ITimingViewSession&
 	}
 }
 
-void FGameplayTimingViewExtender::OnTracksChanged(Insights::ITimingViewSession& InSession, int32& InOutOrder)
-{
-	FPerSessionData* PerSessionData = PerSessionDataMap.Find(&InSession);
-	if(PerSessionData != nullptr)
-	{
-		PerSessionData->GameplaySharedData->OnTracksChanged(InOutOrder);
-	}
-}
-
 #if WITH_ENGINE
 void FGameplayTimingViewExtender::AddVisualizerWorld(UWorld* InWorld)
 {
diff --git a/Engine/Plugins/Animation/GameplayInsights/Source/GameplayInsights/Private/GameplayTimingViewExtender.h b/Engine/Plugins/Animation/GameplayInsights/Source/GameplayInsights/Private/GameplayTimingViewExtender.h
--- a/Engine/Plugins/Animation/GameplayInsights/Source/GameplayInsights/Private/GameplayTimingViewExtender.h
+++ b/Engine/Plugins/Animation/GameplayInsights/Source/GameplayInsights/Private/GameplayTimingViewExtender.h
@@ -17,7 +17,6 @@ public:
 	virtual void OnEndSession(Insights::ITimingViewSession& InSession) override;
 	virtual void Tick(Insights::ITimingViewSession& InSession, const Trace::IAnalysisSession& InAnalysisSession) override;
 	virtual void ExtendFilterMenu(Insights::ITimingViewSession& InSession, FMenuBuilder& InMenuBuilder) override;
-	virtual void OnTracksChanged(Insights::ITimingViewSession& InSession, int32& InOutOrder) override;
 
 #if WITH_ENGINE
 	// Add a world to be visualized
diff --git a/Engine/Plugins/Animation/GameplayInsights/Source/GameplayInsights/Private/GameplayTrack.h b/Engine/Plugins/Animation/GameplayInsights/Source/GameplayInsights/Private/GameplayTrack.h
--- a/Engine/Plugins/Animation/GameplayInsights/Source/GameplayInsights/Private/GameplayTrack.h
+++ b/Engine/Plugins/Animation/GameplayInsights/Source/GameplayInsights/Private/GameplayTrack.h
@@ -11,7 +11,7 @@
 class FGameplayTrack
 {
 public:
-	FGameplayTrack(FBaseTimingTrack* InTimingTrack, uint64 InObjectId)
+	FGameplayTrack(FBaseTimingTrack& InTimingTrack, uint64 InObjectId)
 		: TimingTrack(InTimingTrack)
 		, ObjectId(InObjectId)
 		, Parent(nullptr)
@@ -37,20 +37,20 @@ public:
 
 		Algo::Sort(Children, [](FGameplayTrack* InTrack0, FGameplayTrack* InTrack1)
 		{
-			return InTrack0->GetTimingTrack().GetName() < InTrack1->GetTimingTrack().GetName();
+			return InTrack0->GetTimingTrack()->GetName() < InTrack1->GetTimingTrack()->GetName();
 		});	
 	}
 
 	/** Find a child track using the specified callback */
-	FBaseTimingTrack* FindChildTrack(uint64 InObjectId, TFunctionRef<bool(const FBaseTimingTrack& InTrack)> Callback)
+	TSharedPtr<FBaseTimingTrack> FindChildTrack(uint64 InObjectId, TFunctionRef<bool(const FBaseTimingTrack& InTrack)> Callback)
 	{
 		for(FGameplayTrack* ChildTrack : Children)
 		{
 			if( ChildTrack != nullptr &&
 				ChildTrack->ObjectId == InObjectId && 
-				Callback(ChildTrack->GetTimingTrack()))
+				Callback(ChildTrack->GetTimingTrack().Get()))
 			{
-				return &ChildTrack->GetTimingTrack();
+				return ChildTrack->GetTimingTrack();
 			}
 		}
 
@@ -58,14 +58,11 @@ public:
 	}
 
 	/** Access the outer timing track */
-	FBaseTimingTrack& GetTimingTrack() { return *TimingTrack; }
-
-	/** Access the outer timing track */
-	const FBaseTimingTrack& GetTimingTrack() const { return *TimingTrack; }
+	TSharedRef<FBaseTimingTrack> GetTimingTrack() const { return TimingTrack.AsShared(); }
 
 private:
 	/** Outer timing track */
-	FBaseTimingTrack* TimingTrack;
+	FBaseTimingTrack& TimingTrack;
 
 	// The object ID for this track
 	uint64 ObjectId;
@@ -82,8 +79,8 @@ class TGameplayTrackMixin : public Base
 {
 public:
 	TGameplayTrackMixin(uint64 InObjectId, const FName& InType, const FName& InSubType, const FText& InName)
-		: Base(FBaseTimingTrack::GenerateId(), InType, InSubType, InName.ToString())
-		, GameplayTrack(this, InObjectId)
+		: Base(InType, InSubType, InName.ToString())
+		, GameplayTrack(*this, InObjectId)
 	{
 	}
 
diff --git a/Engine/Plugins/Animation/GameplayInsights/Source/GameplayInsights/Private/ObjectEventsTrack.cpp b/Engine/Plugins/Animation/GameplayInsights/Source/GameplayInsights/Private/ObjectEventsTrack.cpp
--- a/Engine/Plugins/Animation/GameplayInsights/Source/GameplayInsights/Private/ObjectEventsTrack.cpp
+++ b/Engine/Plugins/Animation/GameplayInsights/Source/GameplayInsights/Private/ObjectEventsTrack.cpp
@@ -1,7 +1,6 @@
 // Copyright 1998-2019 Epic Games, Inc. All Rights Reserved.
 
 #include "ObjectEventsTrack.h"
-#include "Insights/ITimingViewDrawHelper.h"
 #include "GameplayProvider.h"
 #include "Insights/ViewModels/TimingTrackViewport.h"
 #include "Insights/ViewModels/TimingEvent.h"
@@ -21,58 +20,55 @@ FObjectEventsTrack::FObjectEventsTrack(const FGameplaySharedData& InSharedData,
 {
 }
 
-void FObjectEventsTrack::Draw(ITimingViewDrawHelper& Helper) const
+void FObjectEventsTrack::BuildDrawState(ITimingEventsTrackDrawStateBuilder& Builder, const ITimingTrackUpdateContext& Context)
 {
-	FObjectEventsTrack& Track = *const_cast<FObjectEventsTrack*>(this);
+	const FGameplayProvider* GameplayProvider = SharedData.GetAnalysisSession().ReadProvider<FGameplayProvider>(FGameplayProvider::ProviderName);
 
-	if (Helper.BeginTimeline(Track))
+	if(GameplayProvider)
 	{
-		const FGameplayProvider* GameplayProvider = SharedData.GetAnalysisSession().ReadProvider<FGameplayProvider>(FGameplayProvider::ProviderName);
+		Trace::FAnalysisSessionReadScope SessionReadScope(SharedData.GetAnalysisSession());
 
-		if(GameplayProvider)
+		// object events
+		GameplayProvider->ReadObjectEventsTimeline(GetGameplayTrack().GetObjectId(), [&Context, &Builder](const FGameplayProvider::ObjectEventsTimeline& InTimeline)
 		{
-			Trace::FAnalysisSessionReadScope SessionReadScope(SharedData.GetAnalysisSession());
-
-			// object events
-			GameplayProvider->ReadObjectEventsTimeline(GetGameplayTrack().GetObjectId(), [&Helper](const FGameplayProvider::ObjectEventsTimeline& InTimeline)
+			InTimeline.EnumerateEvents(Context.GetViewport().GetStartTime(), Context.GetViewport().GetEndTime(), [&Builder](double InStartTime, double InEndTime, uint32 InDepth, const FObjectEventMessage& InMessage)
 			{
-				InTimeline.EnumerateEvents(Helper.GetViewport().GetStartTime(), Helper.GetViewport().GetEndTime(), [&Helper](double InStartTime, double InEndTime, uint32 InDepth, const FObjectEventMessage& InMessage)
-				{
-					Helper.AddEvent(InStartTime, InEndTime, 0, InMessage.Name);
-				});
+				Builder.AddEvent(InStartTime, InEndTime, 0, InMessage.Name);
 			});
-		}
-
-		Helper.EndTimeline(Track);
+		});
 	}
 }
 
-void FObjectEventsTrack::InitTooltip(FTooltipDrawState& Tooltip, const FTimingEvent& HoveredTimingEvent) const
+void FObjectEventsTrack::InitTooltip(FTooltipDrawState& Tooltip, const ITimingEvent& HoveredTimingEvent) const
 {
-	FTimingEventSearchParameters SearchParameters(HoveredTimingEvent.StartTime, HoveredTimingEvent.EndTime, ETimingEventSearchFlags::StopAtFirstMatch);
+	FTimingEventSearchParameters SearchParameters(HoveredTimingEvent.GetStartTime(), HoveredTimingEvent.GetEndTime(), ETimingEventSearchFlags::StopAtFirstMatch);
 
 	FindObjectEvent(SearchParameters, [this, &Tooltip, &HoveredTimingEvent](double InFoundStartTime, double InFoundEndTime, uint32 InFoundDepth, const FObjectEventMessage& InMessage)
 	{
 		Tooltip.ResetContent();
 
 		Tooltip.AddTitle(FText::FromString(FString(InMessage.Name)).ToString());
-		Tooltip.AddNameValueTextLine(LOCTEXT("EventTime", "Time").ToString(), FText::AsNumber(HoveredTimingEvent.StartTime).ToString());
+		Tooltip.AddNameValueTextLine(LOCTEXT("EventTime", "Time").ToString(), FText::AsNumber(HoveredTimingEvent.GetStartTime()).ToString());
 
 		Tooltip.UpdateLayout();
 	});
 }
 
-bool FObjectEventsTrack::SearchTimingEvent(const FTimingEventSearchParameters& InSearchParameters, FTimingEvent& InOutTimingEvent) const
+const TSharedPtr<const ITimingEvent> FObjectEventsTrack::SearchEvent(const FTimingEventSearchParameters& InSearchParameters) const
 {
-	return FindObjectEvent(InSearchParameters, [this, &InOutTimingEvent](double InFoundStartTime, double InFoundEndTime, uint32 InFoundDepth, const FObjectEventMessage& InFoundMessage)
+	TSharedPtr<const ITimingEvent> FoundEvent;
+
+	FindObjectEvent(InSearchParameters, [this, &FoundEvent](double InFoundStartTime, double InFoundEndTime, uint32 InFoundDepth, const FObjectEventMessage& InFoundMessage)
 	{
-		InOutTimingEvent = FTimingEvent(this, InFoundStartTime, InFoundEndTime, InFoundDepth);
+		FoundEvent = MakeShared<const FTimingEvent>(SharedThis(this), InFoundStartTime, InFoundEndTime, InFoundDepth);
 	});
+
+	return FoundEvent;
 }
 
-bool FObjectEventsTrack::FindObjectEvent(const FTimingEventSearchParameters& InParameters, TFunctionRef<void(double, double, uint32, const FObjectEventMessage&)> InFoundPredicate) const
+void FObjectEventsTrack::FindObjectEvent(const FTimingEventSearchParameters& InParameters, TFunctionRef<void(double, double, uint32, const FObjectEventMessage&)> InFoundPredicate) const
 {
-	return TTimingEventSearch<FObjectEventMessage>::Search(
+	TTimingEventSearch<FObjectEventMessage>::Search(
 		InParameters,
 
 		// Search...
diff --git a/Engine/Plugins/Animation/GameplayInsights/Source/GameplayInsights/Private/ObjectEventsTrack.h b/Engine/Plugins/Animation/GameplayInsights/Source/GameplayInsights/Private/ObjectEventsTrack.h
--- a/Engine/Plugins/Animation/GameplayInsights/Source/GameplayInsights/Private/ObjectEventsTrack.h
+++ b/Engine/Plugins/Animation/GameplayInsights/Source/GameplayInsights/Private/ObjectEventsTrack.h
@@ -17,13 +17,13 @@ public:
 
 	FObjectEventsTrack(const FGameplaySharedData& InSharedData, uint64 InObjectID, const TCHAR* InName);
 
-	virtual void Draw(ITimingViewDrawHelper& Helper) const override;
-	virtual void InitTooltip(FTooltipDrawState& Tooltip, const FTimingEvent& HoveredTimingEvent) const override;
-	virtual bool SearchTimingEvent(const FTimingEventSearchParameters& InSearchParameters, FTimingEvent& InOutTimingEvent) const override;
+	virtual void BuildDrawState(ITimingEventsTrackDrawStateBuilder& Builder, const ITimingTrackUpdateContext& Context) override;
+	virtual void InitTooltip(FTooltipDrawState& Tooltip, const ITimingEvent& HoveredTimingEvent) const override;
+	virtual const TSharedPtr<const ITimingEvent> SearchEvent(const FTimingEventSearchParameters& InSearchParameters) const override;
 
 private:
 	// Helper function used to find an object event
-	bool FindObjectEvent(const FTimingEventSearchParameters& InParameters, TFunctionRef<void(double, double, uint32, const FObjectEventMessage&)> InFoundPredicate) const;
+	void FindObjectEvent(const FTimingEventSearchParameters& InParameters, TFunctionRef<void(double, double, uint32, const FObjectEventMessage&)> InFoundPredicate) const;
 
 private:
 	const FGameplaySharedData& SharedData;
diff --git a/Engine/Plugins/Animation/GameplayInsights/Source/GameplayInsights/Private/SkeletalMeshPoseTrack.cpp b/Engine/Plugins/Animation/GameplayInsights/Source/GameplayInsights/Private/SkeletalMeshPoseTrack.cpp
--- a/Engine/Plugins/Animation/GameplayInsights/Source/GameplayInsights/Private/SkeletalMeshPoseTrack.cpp
+++ b/Engine/Plugins/Animation/GameplayInsights/Source/GameplayInsights/Private/SkeletalMeshPoseTrack.cpp
@@ -1,7 +1,6 @@
 // Copyright 1998-2019 Epic Games, Inc. All Rights Reserved.
 
 #include "SkeletalMeshPoseTrack.h"
-#include "Insights/ITimingViewDrawHelper.h"
 #include "GameplayProvider.h"
 #include "AnimationProvider.h"
 #include "Insights/ViewModels/TimingTrackViewport.h"
@@ -9,11 +8,13 @@
 #include "AnimationSharedData.h"
 #include "Framework/Multibox/MultiboxBuilder.h"
 #include "Insights/ViewModels/TimingEventSearch.h"
+#include "Insights/ViewModels/TooltipDrawState.h"
 
 #if WITH_ENGINE
 #include "Components/LineBatchComponent.h"
 #endif
 
+
 #define LOCTEXT_NAMESPACE "SkeletalMeshPoseTrack"
 
 const FName FSkeletalMeshPoseTrack::TypeName(TEXT("Animation"));
@@ -29,53 +30,39 @@ FSkeletalMeshPoseTrack::FSkeletalMeshPoseTrack(const FAnimationSharedData& InSha
 {
 }
 
-void FSkeletalMeshPoseTrack::Reset()
-{
-	FTimingEventsTrack::Reset();
-
-	SetHeight(16.0f);
-}
-
-void FSkeletalMeshPoseTrack::Draw( ITimingViewDrawHelper& Helper) const
+void FSkeletalMeshPoseTrack::BuildDrawState(ITimingEventsTrackDrawStateBuilder& Builder, const ITimingTrackUpdateContext& Context)
 {
-	FSkeletalMeshPoseTrack& Track = *const_cast<FSkeletalMeshPoseTrack*>(this);
+	const FAnimationProvider* AnimationProvider = SharedData.GetAnalysisSession().ReadProvider<FAnimationProvider>(FAnimationProvider::ProviderName);
 
-	if (Helper.BeginTimeline(Track))
+	if(AnimationProvider)
 	{
-		const FAnimationProvider* AnimationProvider = SharedData.GetAnalysisSession().ReadProvider<FAnimationProvider>(FAnimationProvider::ProviderName);
+		Trace::FAnalysisSessionReadScope SessionReadScope(SharedData.GetAnalysisSession());
 
-		if(AnimationProvider)
+		AnimationProvider->ReadSkeletalMeshPoseTimeline(GetGameplayTrack().GetObjectId(), [&Context, &Builder](const FAnimationProvider::SkeletalMeshPoseTimeline& InTimeline)
 		{
-			Trace::FAnalysisSessionReadScope SessionReadScope(SharedData.GetAnalysisSession());
-
-			AnimationProvider->ReadSkeletalMeshPoseTimeline(GetGameplayTrack().GetObjectId(), [&Helper](const FAnimationProvider::SkeletalMeshPoseTimeline& InTimeline)
+			auto DrawEvent = [&Builder](double InStartTime, double InEndTime, uint32 InDepth, const FSkeletalMeshPoseMessage& InMessage)
 			{
-				auto DrawEvent = [&Helper](double InStartTime, double InEndTime, uint32 InDepth, const FSkeletalMeshPoseMessage& InMessage)
-				{
-					static TCHAR Buffer[256];
-					FCString::Snprintf(Buffer, 256, TEXT("%d Bones"), InMessage.NumTransforms);
-					Helper.AddEvent(InStartTime, InEndTime, 0, Buffer);
-				};
-
-				if (FTimingEventsTrack::bUseDownSampling)
-				{
-					const double SecondsPerPixel = 1.0 / Helper.GetViewport().GetScaleX();
-					InTimeline.EnumerateEventsDownSampled(Helper.GetViewport().GetStartTime(), Helper.GetViewport().GetEndTime(), SecondsPerPixel, DrawEvent);
-				}
-				else
-				{
-					InTimeline.EnumerateEvents(Helper.GetViewport().GetStartTime(), Helper.GetViewport().GetEndTime(), DrawEvent);
-				}
-			});
-		}
+				static TCHAR Buffer[256];
+				FCString::Snprintf(Buffer, 256, TEXT("%d Bones"), InMessage.NumTransforms);
+				Builder.AddEvent(InStartTime, InEndTime, 0, Buffer);
+			};
 
-		Helper.EndTimeline(Track);
+			if (FTimingEventsTrack::bUseDownSampling)
+			{
+				const double SecondsPerPixel = 1.0 / Context.GetViewport().GetScaleX();
+				InTimeline.EnumerateEventsDownSampled(Context.GetViewport().GetStartTime(), Context.GetViewport().GetEndTime(), SecondsPerPixel, DrawEvent);
+			}
+			else
+			{
+				InTimeline.EnumerateEvents(Context.GetViewport().GetStartTime(), Context.GetViewport().GetEndTime(), DrawEvent);
+			}
+		});
 	}
 }
 
-void FSkeletalMeshPoseTrack::InitTooltip(FTooltipDrawState& Tooltip, const FTimingEvent& HoveredTimingEvent) const
+void FSkeletalMeshPoseTrack::InitTooltip(FTooltipDrawState& Tooltip, const ITimingEvent& HoveredTimingEvent) const
 {
-	FTimingEventSearchParameters SearchParameters(HoveredTimingEvent.StartTime, HoveredTimingEvent.EndTime, ETimingEventSearchFlags::StopAtFirstMatch);
+	FTimingEventSearchParameters SearchParameters(HoveredTimingEvent.GetStartTime(), HoveredTimingEvent.GetEndTime(), ETimingEventSearchFlags::StopAtFirstMatch);
 
 	FindSkeletalMeshPoseMessage(SearchParameters, [this, &Tooltip](double InFoundStartTime, double InFoundEndTime, uint32 InFoundDepth, const FSkeletalMeshPoseMessage& InMessage)
 	{
@@ -90,17 +77,21 @@ void FSkeletalMeshPoseTrack::InitTooltip(FTooltipDrawState& Tooltip, const FTimi
 	});
 }
 
-bool FSkeletalMeshPoseTrack::SearchTimingEvent(const FTimingEventSearchParameters& InSearchParameters, FTimingEvent& InOutTimingEvent) const
+const TSharedPtr<const ITimingEvent> FSkeletalMeshPoseTrack::SearchEvent(const FTimingEventSearchParameters& InSearchParameters) const
 {
-	return FindSkeletalMeshPoseMessage(InSearchParameters, [this, &InOutTimingEvent](double InFoundStartTime, double InFoundEndTime, uint32 InFoundDepth, const FSkeletalMeshPoseMessage& InFoundMessage)
+	TSharedPtr<const ITimingEvent> FoundEvent;
+
+	FindSkeletalMeshPoseMessage(InSearchParameters, [this, &FoundEvent](double InFoundStartTime, double InFoundEndTime, uint32 InFoundDepth, const FSkeletalMeshPoseMessage& InFoundMessage)
 	{
-		InOutTimingEvent = FTimingEvent(this, InFoundStartTime, InFoundEndTime, InFoundDepth);
+		FoundEvent = MakeShared<const FTimingEvent>(SharedThis(this), InFoundStartTime, InFoundEndTime, InFoundDepth);
 	});
+
+	return FoundEvent;
 }
 
-bool FSkeletalMeshPoseTrack::FindSkeletalMeshPoseMessage(const FTimingEventSearchParameters& InParameters, TFunctionRef<void(double, double, uint32, const FSkeletalMeshPoseMessage&)> InFoundPredicate) const
+void FSkeletalMeshPoseTrack::FindSkeletalMeshPoseMessage(const FTimingEventSearchParameters& InParameters, TFunctionRef<void(double, double, uint32, const FSkeletalMeshPoseMessage&)> InFoundPredicate) const
 {
-	return TTimingEventSearch<FSkeletalMeshPoseMessage>::Search(
+	TTimingEventSearch<FSkeletalMeshPoseMessage>::Search(
 		InParameters,
 
 		[this](TTimingEventSearch<FSkeletalMeshPoseMessage>::FContext& InContext)
diff --git a/Engine/Plugins/Animation/GameplayInsights/Source/GameplayInsights/Private/SkeletalMeshPoseTrack.h b/Engine/Plugins/Animation/GameplayInsights/Source/GameplayInsights/Private/SkeletalMeshPoseTrack.h
--- a/Engine/Plugins/Animation/GameplayInsights/Source/GameplayInsights/Private/SkeletalMeshPoseTrack.h
+++ b/Engine/Plugins/Animation/GameplayInsights/Source/GameplayInsights/Private/SkeletalMeshPoseTrack.h
@@ -18,10 +18,9 @@ public:
 
 	FSkeletalMeshPoseTrack(const FAnimationSharedData& InSharedData, uint64 InObjectID, const TCHAR* InName);
 
-	virtual void Reset() override;
-	virtual void Draw(ITimingViewDrawHelper& Helper) const override;
-	virtual void InitTooltip(FTooltipDrawState& Tooltip, const FTimingEvent& HoveredTimingEvent) const override;
-	virtual bool SearchTimingEvent(const FTimingEventSearchParameters& InSearchParameters, FTimingEvent& InOutTimingEvent) const override;
+	virtual void BuildDrawState(ITimingEventsTrackDrawStateBuilder& Builder, const ITimingTrackUpdateContext& Context) override;
+	virtual void InitTooltip(FTooltipDrawState& Tooltip, const ITimingEvent& HoveredTimingEvent) const override;
+	virtual const TSharedPtr<const ITimingEvent> SearchEvent(const FTimingEventSearchParameters& InSearchParameters) const override;
 	virtual void BuildContextMenu(FMenuBuilder& MenuBuilder) override;
 
 	// Access drawing flags
@@ -37,7 +36,7 @@ public:
 
 private:
 	// Helper function used to find a skeletal mesh pose
-	bool FindSkeletalMeshPoseMessage(const FTimingEventSearchParameters& InParameters, TFunctionRef<void(double, double, uint32, const FSkeletalMeshPoseMessage&)> InFoundPredicate) const;
+	void FindSkeletalMeshPoseMessage(const FTimingEventSearchParameters& InParameters, TFunctionRef<void(double, double, uint32, const FSkeletalMeshPoseMessage&)> InFoundPredicate) const;
 
 private:
 	/** The shared data */
