Added indented 'tree' track headers to gameplay insights tracks

diff --git a/Engine/Plugins/Animation/GameplayInsights/Source/GameplayInsights/Private/AnimationTickRecordsTrack.cpp b/Engine/Plugins/Animation/GameplayInsights/Source/GameplayInsights/Private/AnimationTickRecordsTrack.cpp
--- a/Engine/Plugins/Animation/GameplayInsights/Source/GameplayInsights/Private/AnimationTickRecordsTrack.cpp
+++ b/Engine/Plugins/Animation/GameplayInsights/Source/GameplayInsights/Private/AnimationTickRecordsTrack.cpp
@@ -51,6 +51,12 @@ void FAnimationTickRecordsTrack::BuildDrawState(ITimingEventsTrackDrawStateBuild
 	}
 }
 
+void FAnimationTickRecordsTrack::Draw(const ITimingTrackDrawContext& Context) const
+{
+	DrawEvents(Context);
+	GetGameplayTrack().DrawHeaderForTimingTrack(Context, *this);
+}
+
 void FAnimationTickRecordsTrack::InitTooltip(FTooltipDrawState& Tooltip, const ITimingEvent& HoveredTimingEvent) const
 {
 	FTimingEventSearchParameters SearchParameters(HoveredTimingEvent.GetStartTime(), HoveredTimingEvent.GetEndTime(), ETimingEventSearchFlags::StopAtFirstMatch);
diff --git a/Engine/Plugins/Animation/GameplayInsights/Source/GameplayInsights/Private/AnimationTickRecordsTrack.h b/Engine/Plugins/Animation/GameplayInsights/Source/GameplayInsights/Private/AnimationTickRecordsTrack.h
--- a/Engine/Plugins/Animation/GameplayInsights/Source/GameplayInsights/Private/AnimationTickRecordsTrack.h
+++ b/Engine/Plugins/Animation/GameplayInsights/Source/GameplayInsights/Private/AnimationTickRecordsTrack.h
@@ -18,6 +18,7 @@ public:
 	FAnimationTickRecordsTrack(const FAnimationSharedData& InSharedData, uint64 InObjectId, uint64 InAssetId, const TCHAR* InName);
 
 	virtual void BuildDrawState(ITimingEventsTrackDrawStateBuilder& Builder, const ITimingTrackUpdateContext& Context) override;
+	virtual void Draw(const ITimingTrackDrawContext& Context) const override;
 	virtual void InitTooltip(FTooltipDrawState& Tooltip, const ITimingEvent& HoveredTimingEvent) const override;
 	virtual const TSharedPtr<const ITimingEvent> SearchEvent(const FTimingEventSearchParameters& InSearchParameters) const override;
 
diff --git a/Engine/Plugins/Animation/GameplayInsights/Source/GameplayInsights/Private/GameplaySharedData.cpp b/Engine/Plugins/Animation/GameplayInsights/Source/GameplayInsights/Private/GameplaySharedData.cpp
--- a/Engine/Plugins/Animation/GameplayInsights/Source/GameplayInsights/Private/GameplaySharedData.cpp
+++ b/Engine/Plugins/Animation/GameplayInsights/Source/GameplayInsights/Private/GameplaySharedData.cpp
@@ -75,18 +75,19 @@ TSharedRef<FObjectEventsTrack> FGameplaySharedData::GetObjectEventsTrackForId(In
 	return LeafObjectEventsTrack.ToSharedRef();
 }
 
-static void UpdateTrackOrderRecursive(TSharedRef<FObjectEventsTrack> InTrack, int32& InOutOrder)
+static void UpdateTrackOrderRecursive(TSharedRef<FObjectEventsTrack> InTrack, int32& InOutOrder, uint32 InDepth)
 {
 	// recurse down object-track children, then non-object track leaf tracks to set 
 	// overall ordering based on depth-first traversal of the hierarchy
 
 	InTrack->SetOrder(InOutOrder++);
+	InTrack->GetGameplayTrack().SetIndent(InDepth);
 
 	for(FGameplayTrack* ChildTrack : InTrack->GetGameplayTrack().GetChildTracks())
 	{
 		if(ChildTrack->GetTimingTrack()->GetType() == FObjectEventsTrack::TypeName)
 		{
-			UpdateTrackOrderRecursive(StaticCastSharedRef<FObjectEventsTrack>(ChildTrack->GetTimingTrack()), InOutOrder);
+			UpdateTrackOrderRecursive(StaticCastSharedRef<FObjectEventsTrack>(ChildTrack->GetTimingTrack()), InOutOrder, InDepth + 1);
 		}
 	}
 
@@ -95,6 +96,7 @@ static void UpdateTrackOrderRecursive(TSharedRef<FObjectEventsTrack> InTrack, in
 		if(ChildTrack->GetTimingTrack()->GetType() != FObjectEventsTrack::TypeName)
 		{
 			ChildTrack->GetTimingTrack()->SetOrder(InOutOrder++);
+			ChildTrack->SetIndent(InDepth + 1);
 		}
 	}
 }
@@ -170,7 +172,7 @@ void FGameplaySharedData::SortTracks()
 	// update ordering
 	for(TSharedRef<FObjectEventsTrack> RootTrack : Roots)
 	{
-		UpdateTrackOrderRecursive(RootTrack, Order);
+		UpdateTrackOrderRecursive(RootTrack, Order, 0);
 	}
 
 	Roots.Reset();
diff --git a/Engine/Plugins/Animation/GameplayInsights/Source/GameplayInsights/Private/GameplayTrack.cpp b/Engine/Plugins/Animation/GameplayInsights/Source/GameplayInsights/Private/GameplayTrack.cpp
new file mode 100644
--- /dev/null
+++ b/Engine/Plugins/Animation/GameplayInsights/Source/GameplayInsights/Private/GameplayTrack.cpp
@@ -0,0 +1,89 @@
+// Copyright 1998-2019 Epic Games, Inc. All Rights Reserved.
+
+#include "GameplayTrack.h"
+#include "Insights/ViewModels/TimingEventsTrack.h"
+#include "Algo/Sort.h"
+#include "Framework/Application/SlateApplication.h"
+#include "Fonts/FontMeasure.h"
+#include "Insights/ViewModels/TimingTrackViewport.h"
+#include "Insights/Common/PaintUtils.h"
+#include "Insights/ViewModels/ITimingViewDrawHelper.h"
+
+namespace GameplayTrackConstants
+{
+	constexpr float IndentSize = 12.0f;
+}
+
+void FGameplayTrack::AddChildTrack(FGameplayTrack& InChildTrack)
+{
+	check(InChildTrack.Parent == nullptr);
+	InChildTrack.Parent = this;
+
+	Children.Add(&InChildTrack);
+
+	Algo::Sort(Children, [](FGameplayTrack* InTrack0, FGameplayTrack* InTrack1)
+	{
+		return InTrack0->GetTimingTrack()->GetName() < InTrack1->GetTimingTrack()->GetName();
+	});	
+}
+
+TSharedPtr<FBaseTimingTrack> FGameplayTrack::FindChildTrack(uint64 InObjectId, TFunctionRef<bool(const FBaseTimingTrack& InTrack)> Callback) const
+{
+	for(const FGameplayTrack* ChildTrack : Children)
+	{
+		if( ChildTrack != nullptr &&
+			ChildTrack->ObjectId == InObjectId && 
+			Callback(ChildTrack->GetTimingTrack().Get()))
+		{
+			return ChildTrack->GetTimingTrack();
+		}
+	}
+
+	return nullptr;
+}
+
+static FORCEINLINE bool IntervalsIntersect(float Min1, float Max1, float Min2, float Max2)
+{
+	return !(Max2 < Min1 || Max1 < Min2);
+}
+
+void FGameplayTrack::DrawHeaderForTimingTrack(const ITimingTrackDrawContext& InContext, const FTimingEventsTrack& InTrack) const
+{
+	const float X = (float)Indent * GameplayTrackConstants::IndentSize;
+	const float Y = InTrack.GetPosY();
+	const float H = InTrack.GetHeight();
+	const float TrackNameH = H > 7.0f ? 12.0f : H;
+
+	if (H > 0.0f &&
+		Y + H > InContext.GetViewport().GetTopOffset() &&
+		Y < InContext.GetViewport().GetHeight() - InContext.GetViewport().GetBottomOffset())
+	{
+		// Draw a horizontal line between timelines.
+		InContext.GetDrawContext().DrawBox(InContext.GetHelper().GetHeaderBackgroundLayerId(), X, Y, InContext.GetViewport().GetWidth(), 1.0f, InContext.GetHelper().GetWhiteBrush(), InContext.GetHelper().GetEdgeColor());
+
+		if(H > 7.0f)
+		{
+			const TSharedRef<FSlateFontMeasure> FontMeasureService = FSlateApplication::Get().GetRenderer()->GetFontMeasureService();
+			const float NameWidth = FontMeasureService->Measure(InTrack.GetName(), InContext.GetHelper().GetEventFont()).X;
+			InContext.GetDrawContext().DrawBox(InContext.GetHelper().GetHeaderBackgroundLayerId(), X, Y + 1.0f, NameWidth + 4.0f, TrackNameH, InContext.GetHelper().GetWhiteBrush(), InContext.GetHelper().GetEdgeColor());
+			InContext.GetDrawContext().DrawText(InContext.GetHelper().GetHeaderTextLayerId(), X + 2.0f, Y, InTrack.GetName(), InContext.GetHelper().GetEventFont(), InContext.GetHelper().GetTrackNameTextColor(InTrack));
+		}
+		else
+		{
+			InContext.GetDrawContext().DrawBox(InContext.GetHelper().GetHeaderBackgroundLayerId(), X, Y + 1.0f, H, H, InContext.GetHelper().GetWhiteBrush(), InContext.GetHelper().GetEdgeColor());
+		}
+	}
+
+	// Draw lines connecting to parent
+	if(Parent && Parent->GetTimingTrack()->IsVisible())
+	{
+		const float ParentX = (float)Parent->GetIndent() * GameplayTrackConstants::IndentSize;
+		const float ParentY = FMath::Max(InContext.GetViewport().GetTopOffset(), Parent->GetTimingTrack()->GetPosY());
+
+		if (IntervalsIntersect(ParentY, Y, InContext.GetViewport().GetTopOffset(), InContext.GetViewport().GetHeight() - InContext.GetViewport().GetBottomOffset()))
+		{
+			InContext.GetDrawContext().DrawBox(InContext.GetHelper().GetHeaderBackgroundLayerId(), ParentX, Y + (TrackNameH * 0.5f), X - ParentX, 1.0f, InContext.GetHelper().GetWhiteBrush(), InContext.GetHelper().GetEdgeColor());
+			InContext.GetDrawContext().DrawBox(InContext.GetHelper().GetHeaderBackgroundLayerId(), ParentX, ParentY, 1.0f, (Y - ParentY) + (TrackNameH * 0.5f), InContext.GetHelper().GetWhiteBrush(), InContext.GetHelper().GetEdgeColor());
+		}
+	}
+}
diff --git a/Engine/Plugins/Animation/GameplayInsights/Source/GameplayInsights/Private/GameplayTrack.h b/Engine/Plugins/Animation/GameplayInsights/Source/GameplayInsights/Private/GameplayTrack.h
--- a/Engine/Plugins/Animation/GameplayInsights/Source/GameplayInsights/Private/GameplayTrack.h
+++ b/Engine/Plugins/Animation/GameplayInsights/Source/GameplayInsights/Private/GameplayTrack.h
@@ -3,8 +3,6 @@
 #pragma once
 
 #include "Insights/ViewModels/TimingEventsTrack.h"
-#include "Insights/ViewModels/GraphTrack.h"
-#include "Algo/Sort.h"
 
 // Provides parent/child hierarchy structure and the owning object Id
 // Designed as a compositional member of outer timing tracks
@@ -15,6 +13,7 @@ public:
 		: TimingTrack(InTimingTrack)
 		, ObjectId(InObjectId)
 		, Parent(nullptr)
+		, Indent(0)
 	{
 	}
 
@@ -28,38 +27,23 @@ public:
 	TArrayView<FGameplayTrack*> GetChildTracks() { return Children; }
 
 	/** Add a child track */
-	void AddChildTrack(FGameplayTrack& InChildTrack)
-	{
-		check(InChildTrack.Parent == nullptr);
-		InChildTrack.Parent = this;
-
-		Children.Add(&InChildTrack);
-
-		Algo::Sort(Children, [](FGameplayTrack* InTrack0, FGameplayTrack* InTrack1)
-		{
-			return InTrack0->GetTimingTrack()->GetName() < InTrack1->GetTimingTrack()->GetName();
-		});	
-	}
+	void AddChildTrack(FGameplayTrack& InChildTrack);
 
 	/** Find a child track using the specified callback */
-	TSharedPtr<FBaseTimingTrack> FindChildTrack(uint64 InObjectId, TFunctionRef<bool(const FBaseTimingTrack& InTrack)> Callback)
-	{
-		for(FGameplayTrack* ChildTrack : Children)
-		{
-			if( ChildTrack != nullptr &&
-				ChildTrack->ObjectId == InObjectId && 
-				Callback(ChildTrack->GetTimingTrack().Get()))
-			{
-				return ChildTrack->GetTimingTrack();
-			}
-		}
-
-		return nullptr;
-	}
+	TSharedPtr<FBaseTimingTrack> FindChildTrack(uint64 InObjectId, TFunctionRef<bool(const FBaseTimingTrack& InTrack)> Callback) const;
+
+	/** Helper to draw the name for a timing track (uses indentation etc.) */
+	void DrawHeaderForTimingTrack(const ITimingTrackDrawContext& InContext, const FTimingEventsTrack& InTrack) const;
 
 	/** Access the outer timing track */
 	TSharedRef<FBaseTimingTrack> GetTimingTrack() const { return TimingTrack.AsShared(); }
 
+	/** Access indent */
+	void SetIndent(uint32 InIndent) { Indent = InIndent; }
+
+	/** Access indent */
+	uint32 GetIndent() const { return Indent; }
+
 private:
 	/** Outer timing track */
 	FBaseTimingTrack& TimingTrack;
@@ -72,6 +56,9 @@ private:
 
 	/** Our child tracks */
 	TArray<FGameplayTrack*> Children;
+
+	/** Our indent when drawn in a tree */
+	uint32 Indent;
 };
 
 template <class Base>
diff --git a/Engine/Plugins/Animation/GameplayInsights/Source/GameplayInsights/Private/ObjectEventsTrack.cpp b/Engine/Plugins/Animation/GameplayInsights/Source/GameplayInsights/Private/ObjectEventsTrack.cpp
--- a/Engine/Plugins/Animation/GameplayInsights/Source/GameplayInsights/Private/ObjectEventsTrack.cpp
+++ b/Engine/Plugins/Animation/GameplayInsights/Source/GameplayInsights/Private/ObjectEventsTrack.cpp
@@ -39,6 +39,12 @@ void FObjectEventsTrack::BuildDrawState(ITimingEventsTrackDrawStateBuilder& Buil
 	}
 }
 
+void FObjectEventsTrack::Draw(const ITimingTrackDrawContext& Context) const
+{
+	DrawEvents(Context);
+	GetGameplayTrack().DrawHeaderForTimingTrack(Context, *this);
+}
+
 void FObjectEventsTrack::InitTooltip(FTooltipDrawState& Tooltip, const ITimingEvent& HoveredTimingEvent) const
 {
 	FTimingEventSearchParameters SearchParameters(HoveredTimingEvent.GetStartTime(), HoveredTimingEvent.GetEndTime(), ETimingEventSearchFlags::StopAtFirstMatch);
diff --git a/Engine/Plugins/Animation/GameplayInsights/Source/GameplayInsights/Private/ObjectEventsTrack.h b/Engine/Plugins/Animation/GameplayInsights/Source/GameplayInsights/Private/ObjectEventsTrack.h
--- a/Engine/Plugins/Animation/GameplayInsights/Source/GameplayInsights/Private/ObjectEventsTrack.h
+++ b/Engine/Plugins/Animation/GameplayInsights/Source/GameplayInsights/Private/ObjectEventsTrack.h
@@ -18,6 +18,7 @@ public:
 	FObjectEventsTrack(const FGameplaySharedData& InSharedData, uint64 InObjectID, const TCHAR* InName);
 
 	virtual void BuildDrawState(ITimingEventsTrackDrawStateBuilder& Builder, const ITimingTrackUpdateContext& Context) override;
+	virtual void Draw(const ITimingTrackDrawContext& Context) const override;
 	virtual void InitTooltip(FTooltipDrawState& Tooltip, const ITimingEvent& HoveredTimingEvent) const override;
 	virtual const TSharedPtr<const ITimingEvent> SearchEvent(const FTimingEventSearchParameters& InSearchParameters) const override;
 
diff --git a/Engine/Plugins/Animation/GameplayInsights/Source/GameplayInsights/Private/SkeletalMeshPoseTrack.cpp b/Engine/Plugins/Animation/GameplayInsights/Source/GameplayInsights/Private/SkeletalMeshPoseTrack.cpp
--- a/Engine/Plugins/Animation/GameplayInsights/Source/GameplayInsights/Private/SkeletalMeshPoseTrack.cpp
+++ b/Engine/Plugins/Animation/GameplayInsights/Source/GameplayInsights/Private/SkeletalMeshPoseTrack.cpp
@@ -21,7 +21,7 @@ const FName FSkeletalMeshPoseTrack::TypeName(TEXT("Animation"));
 const FName FSkeletalMeshPoseTrack::SubTypeName(TEXT("SkeletalMeshPose"));
 
 FSkeletalMeshPoseTrack::FSkeletalMeshPoseTrack(const FAnimationSharedData& InSharedData, uint64 InObjectID, const TCHAR* InName)
-	: TGameplayTrackMixin<FTimingEventsTrack>(InObjectID, FSkeletalMeshPoseTrack::TypeName, FSkeletalMeshPoseTrack::SubTypeName, FText::Format(LOCTEXT("SkelMeshPoseName", "{0} Pose"), FText::FromString(FString(InName))))
+	: TGameplayTrackMixin<FTimingEventsTrack>(InObjectID, FSkeletalMeshPoseTrack::TypeName, FSkeletalMeshPoseTrack::SubTypeName, FText::Format(LOCTEXT("TrackNameFormat", "Pose - {0}"), FText::FromString(FString(InName))))
 	, SharedData(InSharedData)
 	, bDrawMarkerTime(false)
 	, bDrawSelectedEvent(false)
@@ -60,6 +60,12 @@ void FSkeletalMeshPoseTrack::BuildDrawState(ITimingEventsTrackDrawStateBuilder&
 	}
 }
 
+void FSkeletalMeshPoseTrack::Draw(const ITimingTrackDrawContext& Context) const
+{
+	DrawEvents(Context);
+	GetGameplayTrack().DrawHeaderForTimingTrack(Context, *this);
+}
+
 void FSkeletalMeshPoseTrack::InitTooltip(FTooltipDrawState& Tooltip, const ITimingEvent& HoveredTimingEvent) const
 {
 	FTimingEventSearchParameters SearchParameters(HoveredTimingEvent.GetStartTime(), HoveredTimingEvent.GetEndTime(), ETimingEventSearchFlags::StopAtFirstMatch);
diff --git a/Engine/Plugins/Animation/GameplayInsights/Source/GameplayInsights/Private/SkeletalMeshPoseTrack.h b/Engine/Plugins/Animation/GameplayInsights/Source/GameplayInsights/Private/SkeletalMeshPoseTrack.h
--- a/Engine/Plugins/Animation/GameplayInsights/Source/GameplayInsights/Private/SkeletalMeshPoseTrack.h
+++ b/Engine/Plugins/Animation/GameplayInsights/Source/GameplayInsights/Private/SkeletalMeshPoseTrack.h
@@ -19,6 +19,7 @@ public:
 	FSkeletalMeshPoseTrack(const FAnimationSharedData& InSharedData, uint64 InObjectID, const TCHAR* InName);
 
 	virtual void BuildDrawState(ITimingEventsTrackDrawStateBuilder& Builder, const ITimingTrackUpdateContext& Context) override;
+	virtual void Draw(const ITimingTrackDrawContext& Context) const override;
 	virtual void InitTooltip(FTooltipDrawState& Tooltip, const ITimingEvent& HoveredTimingEvent) const override;
 	virtual const TSharedPtr<const ITimingEvent> SearchEvent(const FTimingEventSearchParameters& InSearchParameters) const override;
 	virtual void BuildContextMenu(FMenuBuilder& MenuBuilder) override;
