Added mesh rendering of skeletal mesh pose track data

#rnx

diff --git a/Engine/Plugins/Animation/GameplayInsights/Source/GameplayInsights/GameplayInsights.Build.cs b/Engine/Plugins/Animation/GameplayInsights/Source/GameplayInsights/GameplayInsights.Build.cs
--- a/Engine/Plugins/Animation/GameplayInsights/Source/GameplayInsights/GameplayInsights.Build.cs
+++ b/Engine/Plugins/Animation/GameplayInsights/Source/GameplayInsights/GameplayInsights.Build.cs
@@ -10,7 +10,6 @@ namespace UnrealBuildTool.Rules
 			new string[]
 			{
 				"Core",
-				"CoreUObject",
 				"InputCore",
 				"SlateCore",
 				"Slate",
@@ -21,7 +20,21 @@ namespace UnrealBuildTool.Rules
 
 			if (Target.bCompileAgainstEngine)
 			{
-				PrivateDependencyModuleNames.Add("Engine");
+				PrivateDependencyModuleNames.AddRange(
+				new string[]
+				{
+					"Engine",
+					"CoreUObject",
+				});
+			}
+
+			if (Target.bBuildEditor)
+			{
+				PrivateDependencyModuleNames.AddRange(
+				new string[]
+				{
+					"UnrealEd",
+				});
 			}
 		}
 	}
diff --git a/Engine/Plugins/Animation/GameplayInsights/Source/GameplayInsights/Private/AnimationAnalyzer.cpp b/Engine/Plugins/Animation/GameplayInsights/Source/GameplayInsights/Private/AnimationAnalyzer.cpp
--- a/Engine/Plugins/Animation/GameplayInsights/Source/GameplayInsights/Private/AnimationAnalyzer.cpp
+++ b/Engine/Plugins/Animation/GameplayInsights/Source/GameplayInsights/Private/AnimationAnalyzer.cpp
@@ -17,7 +17,8 @@ void FAnimationAnalyzer::OnAnalysisBegin(const FOnAnalysisContext& Context)
 
 	Builder.RouteEvent(RouteId_TickRecord, "Animation", "TickRecord");
 	Builder.RouteEvent(RouteId_SkeletalMesh, "Animation", "SkeletalMesh");
-	Builder.RouteEvent(RouteId_SkeletalMeshPose, "Animation", "SkeletalMeshPose");
+	Builder.RouteEvent(RouteId_SkeletalMeshComponent, "Animation", "SkeletalMeshComponent");
+	Builder.RouteEvent(RouteId_Name, "Animation", "Name");
 }
 
 bool FAnimationAnalyzer::OnEvent(uint16 RouteId, const FOnEventContext& Context)
@@ -50,16 +51,24 @@ bool FAnimationAnalyzer::OnEvent(uint16 RouteId, const FOnEventContext& Context)
 		AnimationProvider.AppendSkeletalMesh(Id, ParentIndices);
 		break;
 	}
-	case RouteId_SkeletalMeshPose:
+	case RouteId_SkeletalMeshComponent:
 	{
 		uint64 Cycle = EventData.GetValue<uint64>("Cycle");
 		uint64 ComponentId = EventData.GetValue<uint64>("ComponentId");
 		uint64 MeshId = EventData.GetValue<uint64>("MeshId");
 		uint32 BoneCount = EventData.GetValue<uint32>("BoneCount");
+		uint32 CurveCount = EventData.GetValue<uint32>("CurveCount");
 		uint16 FrameCounter = EventData.GetValue<uint16>("FrameCounter");
 		uint16 LodIndex = EventData.GetValue<uint16>("LodIndex");
 		TArrayView<const FTransform> Pose(reinterpret_cast<const FTransform*>(EventData.GetAttachment()), BoneCount);
-		AnimationProvider.AppendSkeletalMeshPose(ComponentId, MeshId, Context.SessionContext.TimestampFromCycle(Cycle), LodIndex, FrameCounter, Pose);
+		TArrayView<const FSkeletalMeshNamedCurve> Curves(reinterpret_cast<const FSkeletalMeshNamedCurve*>(EventData.GetAttachment() + (sizeof(FTransform) * BoneCount)), CurveCount);
+		AnimationProvider.AppendSkeletalMeshComponent(ComponentId, MeshId, Context.SessionContext.TimestampFromCycle(Cycle), LodIndex, FrameCounter, Pose, Curves);
+		break;
+	}
+	case RouteId_Name:
+	{
+		uint32 Id = EventData.GetValue<uint32>("Id");
+		AnimationProvider.AppendName(Id, reinterpret_cast<const TCHAR*>(EventData.GetAttachment()));
 		break;
 	}
 	}
diff --git a/Engine/Plugins/Animation/GameplayInsights/Source/GameplayInsights/Private/AnimationAnalyzer.h b/Engine/Plugins/Animation/GameplayInsights/Source/GameplayInsights/Private/AnimationAnalyzer.h
--- a/Engine/Plugins/Animation/GameplayInsights/Source/GameplayInsights/Private/AnimationAnalyzer.h
+++ b/Engine/Plugins/Animation/GameplayInsights/Source/GameplayInsights/Private/AnimationAnalyzer.h
@@ -21,7 +21,8 @@ private:
 	{
 		RouteId_TickRecord,
 		RouteId_SkeletalMesh,
-		RouteId_SkeletalMeshPose,
+		RouteId_SkeletalMeshComponent,
+		RouteId_Name,
 	};
 
 	Trace::IAnalysisSession& Session;
diff --git a/Engine/Plugins/Animation/GameplayInsights/Source/GameplayInsights/Private/AnimationProvider.cpp b/Engine/Plugins/Animation/GameplayInsights/Source/GameplayInsights/Private/AnimationProvider.cpp
--- a/Engine/Plugins/Animation/GameplayInsights/Source/GameplayInsights/Private/AnimationProvider.cpp
+++ b/Engine/Plugins/Animation/GameplayInsights/Source/GameplayInsights/Private/AnimationProvider.cpp
@@ -6,11 +6,14 @@
 
 FName FAnimationProvider::ProviderName("AnimationProvider");
 
-FAnimationProvider::FAnimationProvider(Trace::IAnalysisSession& InSession)
+FAnimationProvider::FAnimationProvider(Trace::IAnalysisSession& InSession, FGameplayProvider& InGameplayProvider)
 	: Session(InSession)
+	, GameplayProvider(InGameplayProvider)
 	, SkeletalMeshPoseTransforms(InSession.GetLinearAllocator(), 256)
+	, SkeletalMeshCurves(InSession.GetLinearAllocator(), 256)
 	, SkeletalMeshParentIndices(InSession.GetLinearAllocator(), 256)
 {
+	GameplayProvider.OnObjectEndPlay().AddRaw(this, &FAnimationProvider::HandleObjectEndPlay);
 }
 
 bool FAnimationProvider::ReadSkeletalMeshPoseTimeline(uint64 InObjectId, TFunctionRef<void(const SkeletalMeshPoseTimeline&)> Callback) const
@@ -41,24 +44,24 @@ bool FAnimationProvider::ReadSkeletalMeshPoseMessage(uint64 InObjectId, uint64 I
 	});
 }
 
-void FAnimationProvider::EnumerateSkeletalMeshPose(const FSkeletalMeshPoseMessage& InMessage, const FSkeletalMeshInfo& InMeshInfo, TFunctionRef<void(const FTransform&, const FTransform&)> Callback) const
+void FAnimationProvider::GetSkeletalMeshComponentSpacePose(const FSkeletalMeshPoseMessage& InMessage, const FSkeletalMeshInfo& InMeshInfo, FTransform& OutComponentToWorld, TArray<FTransform>& OutTransforms) const
 {
 	Session.ReadAccessCheck();
 
-	if(InMeshInfo.BoneCount == InMessage.NumTransforms - 1)
+	if(InMeshInfo.BoneCount == InMessage.NumTransforms)
 	{
+		// Pre-alloc array
+		OutTransforms.SetNumUninitialized(InMessage.NumTransforms);
+
 		// First transform is always component to world
-		const FTransform& ComponentToWorld = SkeletalMeshPoseTransforms[InMessage.TransformStartIndex];
-		uint64 StartTransformIndex = InMessage.TransformStartIndex + 1;
-		uint64 RootBoneIndex = InMeshInfo.ParentIndicesStartIndex;
+		OutComponentToWorld = InMessage.ComponentToWorld;
+		uint64 StartTransformIndex = InMessage.TransformStartIndex;
 		uint64 EndTransformIndex = InMessage.TransformStartIndex + InMessage.NumTransforms;
-		uint64 TransformIndex;
-		uint64 BoneIndex;
-		for(TransformIndex = StartTransformIndex, BoneIndex = RootBoneIndex; TransformIndex < EndTransformIndex; ++TransformIndex, ++BoneIndex)
+		uint64 SourceTransformIndex;
+		int32 TargetTransformIndex = 0;
+		for(SourceTransformIndex = StartTransformIndex; SourceTransformIndex < EndTransformIndex; ++SourceTransformIndex, ++TargetTransformIndex)
 		{
-			int32 ParentIndex = SkeletalMeshParentIndices[BoneIndex];
-			uint64 ParentTransformIndex = ParentIndex == INDEX_NONE ? TransformIndex : (StartTransformIndex + (uint64)ParentIndex);
-			Callback(SkeletalMeshPoseTransforms[TransformIndex] * ComponentToWorld, SkeletalMeshPoseTransforms[ParentTransformIndex] * ComponentToWorld);
+			OutTransforms[TargetTransformIndex] = SkeletalMeshPoseTransforms[SourceTransformIndex];
 		}
 	}
 }
@@ -125,11 +128,16 @@ const FSkeletalMeshInfo* FAnimationProvider::FindSkeletalMeshInfo(uint64 InObjec
 	return nullptr;
 }
 
+const TCHAR* FAnimationProvider::GetName(uint32 InId) const
+{
+	return NameMap.FindRef(InId);
+}
+
 void FAnimationProvider::AppendTickRecord(uint64 InObjectId, uint64 InSubObjectId, double InTime, uint64 InAssetId, float InBlendWeight, float InPlaybackTime, float InRootMotionWeight, float InPlayRate, uint16 InFrameCounter, bool bInLooping)
 {
 	Session.WriteAccessCheck();
 
-	TSharedPtr<Trace::TMonotonicTimeline<FTickRecordMessage>> Timeline;
+	TSharedPtr<Trace::TPointTimeline<FTickRecordMessage>> Timeline;
 	TSharedPtr<FPerObjectTimelineStorage> TimelineStorage;
 	uint32* TimelineStorageIndexPtr = ObjectIdToTickRecordTimelineStorage.Find(InObjectId);
 	if(TimelineStorageIndexPtr != nullptr)
@@ -151,27 +159,11 @@ void FAnimationProvider::AppendTickRecord(uint64 InObjectId, uint64 InSubObjectI
 	}
 	else
 	{
-		Timeline = MakeShared<Trace::TMonotonicTimeline<FTickRecordMessage>>(Session.GetLinearAllocator());
+		Timeline = MakeShared<Trace::TPointTimeline<FTickRecordMessage>>(Session.GetLinearAllocator());
 		TimelineStorage->AssetIdToTickRecordTimeline.Add(InAssetId, TimelineStorage->Timelines.Num());
 		TimelineStorage->Timelines.Add(Timeline.ToSharedRef());
 	}
 
-	// terminate existing scopes - check for continuity
-	uint64 NumEvents = Timeline->GetEventCount();
-	if(NumEvents > 0)
-	{
-		if(InFrameCounter - 1 != Timeline->GetEvent(NumEvents - 1).FrameCounter)
-		{
-			// discontinuous frame count, so add an end event at last time
-			Timeline->AppendEndEvent(Timeline->GetEndTime());
-		}
-		else
-		{
-			// continuous frame time, add end event at current time
-			Timeline->AppendEndEvent(InTime);
-		}
-	}
-
 	FTickRecordMessage Message;
 	Message.ComponentId = InObjectId;
 	Message.AnimInstanceId = InSubObjectId;
@@ -183,7 +175,7 @@ void FAnimationProvider::AppendTickRecord(uint64 InObjectId, uint64 InSubObjectI
 	Message.FrameCounter = InFrameCounter;
 	Message.Looping = bInLooping;
 
-	Timeline->AppendBeginEvent(InTime, Message);
+	Timeline->AppendEvent(InTime, Message);
 
 	Session.UpdateDurationSeconds(InTime);
 }
@@ -209,11 +201,11 @@ void FAnimationProvider::AppendSkeletalMesh(uint64 InObjectId, const TArrayView<
 	}
 }
 
-void FAnimationProvider::AppendSkeletalMeshPose(uint64 InObjectId, uint64 InMeshId, double InTime, uint16 InLodIndex, uint16 InFrameCounter, const TArrayView<const FTransform>& InPose)
+void FAnimationProvider::AppendSkeletalMeshComponent(uint64 InObjectId, uint64 InMeshId, double InTime, uint16 InLodIndex, uint16 InFrameCounter, const TArrayView<const FTransform>& InPose, const TArrayView<const FSkeletalMeshNamedCurve>& InCurves)
 {
 	Session.WriteAccessCheck();
 
-	TSharedPtr<Trace::TMonotonicTimeline<FSkeletalMeshPoseMessage>> Timeline;
+	TSharedPtr<Trace::TIntervalTimeline<FSkeletalMeshPoseMessage>> Timeline;
 	uint32* IndexPtr = ObjectIdToSkeletalMeshPoseTimelines.Find(InObjectId);
 	if(IndexPtr != nullptr)
 	{
@@ -221,41 +213,63 @@ void FAnimationProvider::AppendSkeletalMeshPose(uint64 InObjectId, uint64 InMesh
 	}
 	else
 	{
-		Timeline = MakeShared<Trace::TMonotonicTimeline<FSkeletalMeshPoseMessage>>(Session.GetLinearAllocator());
+		Timeline = MakeShared<Trace::TIntervalTimeline<FSkeletalMeshPoseMessage>>(Session.GetLinearAllocator());
 		ObjectIdToSkeletalMeshPoseTimelines.Add(InObjectId, SkeletalMeshPoseTimelines.Num());
 		SkeletalMeshPoseTimelines.Add(Timeline.ToSharedRef());
 	}
 
-	// terminate existing scopes - check for continuity
+	// terminate existing scopes
 	uint64 NumEvents = Timeline->GetEventCount();
 	if(NumEvents > 0)
 	{
-		if(InFrameCounter - 1 != Timeline->GetEvent(NumEvents - 1).FrameCounter)
-		{
-			// discontinuous frame count, so add an end event at last time
-			Timeline->AppendEndEvent(Timeline->GetEndTime());
-		}
-		else
-		{
-			// continuous frame time, add end event at current time
-			Timeline->AppendEndEvent(InTime);
-		}
+		// Add end event at current time
+		Timeline->EndEvent(NumEvents - 1, InTime);
 	}
 
 	FSkeletalMeshPoseMessage Message;
+	Message.ComponentToWorld = InPose[0];
 	Message.TransformStartIndex = SkeletalMeshPoseTransforms.Num();
+	Message.CurveStartIndex = SkeletalMeshCurves.Num();
 	Message.ComponentId = InObjectId;
 	Message.MeshId = InMeshId;
-	Message.NumTransforms = (uint16)InPose.Num();
+	Message.NumTransforms = (uint16)InPose.Num() - 1;
+	Message.NumCurves = (uint16)InCurves.Num();
 	Message.LodIndex = InLodIndex;
 	Message.FrameCounter = InFrameCounter;
 
 	Timeline->AppendBeginEvent(InTime, Message);
 
-	for(const FTransform& Transform : InPose)
+	for(int32 TransformIndex = 1; TransformIndex < InPose.Num(); ++TransformIndex)
+	{
+		SkeletalMeshPoseTransforms.PushBack() = InPose[TransformIndex];
+	}
+
+	for(const FSkeletalMeshNamedCurve& Curve : InCurves)
 	{
-		SkeletalMeshPoseTransforms.PushBack() = Transform;
+		SkeletalMeshCurves.PushBack() = Curve;
 	}
 
 	Session.UpdateDurationSeconds(InTime);
 }
+
+void FAnimationProvider::AppendName(uint32 InId, const TCHAR* InName)
+{
+	Session.WriteAccessCheck();
+
+	NameMap.Add(InId, Session.StoreString(InName));
+}
+
+void FAnimationProvider::HandleObjectEndPlay(uint64 InObjectId, double InTime, const FObjectInfo& InObjectInfo)
+{
+	// terminate all existing scopes for this object
+	uint32* SkelMeshPoseIndexPtr = ObjectIdToSkeletalMeshPoseTimelines.Find(InObjectId);
+	if(SkelMeshPoseIndexPtr != nullptr)
+	{
+		TSharedPtr<Trace::TIntervalTimeline<FSkeletalMeshPoseMessage>> Timeline = SkeletalMeshPoseTimelines[*SkelMeshPoseIndexPtr];
+		uint64 NumEvents = Timeline->GetEventCount();
+		if(NumEvents > 0)
+		{
+			Timeline->EndEvent(NumEvents - 1, InTime);
+		}
+	}
+}
\ No newline at end of file
diff --git a/Engine/Plugins/Animation/GameplayInsights/Source/GameplayInsights/Private/AnimationProvider.h b/Engine/Plugins/Animation/GameplayInsights/Source/GameplayInsights/Private/AnimationProvider.h
--- a/Engine/Plugins/Animation/GameplayInsights/Source/GameplayInsights/Private/AnimationProvider.h
+++ b/Engine/Plugins/Animation/GameplayInsights/Source/GameplayInsights/Private/AnimationProvider.h
@@ -4,26 +4,30 @@
 
 #include "IAnimationProvider.h"
 #include "Common/PagedArray.h"
-#include "Model/MonotonicTimeline.h"
+#include "Model/IntervalTimeline.h"
+#include "Model/PointTimeline.h"
 #include "Containers/ArrayView.h"
 
 namespace Trace { class IAnalysisSession; }
+class FGameplayProvider;
+struct FObjectInfo;
 
 class FAnimationProvider : public IAnimationProvider
 {
 public:
 	static FName ProviderName;
 
-	FAnimationProvider(Trace::IAnalysisSession& InSession);
+	FAnimationProvider(Trace::IAnalysisSession& InSession, FGameplayProvider& InGameplayProvider);
 
 	/** IAnimationProvider interface */
 	virtual bool ReadSkeletalMeshPoseTimeline(uint64 InObjectId, TFunctionRef<void(const SkeletalMeshPoseTimeline&)> Callback) const override;
 	virtual bool ReadSkeletalMeshPoseMessage(uint64 InObjectId, uint64 InMessageId, TFunctionRef<void(const FSkeletalMeshPoseMessage&)> Callback) const override;
-	virtual void EnumerateSkeletalMeshPose(const FSkeletalMeshPoseMessage& InMessage, const FSkeletalMeshInfo& InMeshInfo, TFunctionRef<void(const FTransform&, const FTransform&)> Callback) const override;
+	virtual void GetSkeletalMeshComponentSpacePose(const FSkeletalMeshPoseMessage& InMessage, const FSkeletalMeshInfo& InMeshInfo, FTransform& OutComponentToWorld, TArray<FTransform>& OutTransforms) const override;
 	virtual void EnumerateTickRecordTimelines(uint64 InObjectId, TFunctionRef<void(uint64, const TickRecordTimeline&)> Callback) const override;
 	virtual bool ReadTickRecordTimeline(uint64 InObjectId, uint64 InAssetId, TFunctionRef<void(const TickRecordTimeline&)> Callback) const override;
 	virtual bool ReadTickRecordMessage(uint64 InObjectId, uint64 InAssetId, uint64 InMessageId, TFunctionRef<void(const FTickRecordMessage&)> Callback) const override;
 	virtual const FSkeletalMeshInfo* FindSkeletalMeshInfo(uint64 InObjectId) const override;
+	virtual const TCHAR* GetName(uint32 InId) const override;
 
 	/** Add a tick record */
 	void AppendTickRecord(uint64 InObjectId, uint64 InSubObjectId, double InTime, uint64 InAssetId, float InBlendWeight, float InPlaybackTime, float InRootMotionWeight, float InPlayRate, uint16 InFrameCounter, bool bInLooping);
@@ -31,11 +35,22 @@ public:
 	/** Add a skeletal mesh */
 	void AppendSkeletalMesh(uint64 InObjectId, const TArrayView<const int32>& ParentIndices);
 
-	/** Add a skeletal mesh pose */
-	void AppendSkeletalMeshPose(uint64 InObjectId, uint64 InMeshId, double InTime, uint16 InLodIndex, uint16 InFrameCounter, const TArrayView<const FTransform>& InTransforms);
+	/** Add a skeletal mesh pose/curves etc. */
+	void AppendSkeletalMeshComponent(uint64 InObjectId, uint64 InMeshId, double InTime, uint16 InLodIndex, uint16 InFrameCounter, const TArrayView<const FTransform>& InTransforms, const TArrayView<const FSkeletalMeshNamedCurve>& InCurves);
+
+	/** Get a skeletal mesh for a specified path. If the mesh no longer exists, make a fake one */
+	USkeletalMesh* GetSkeletalMesh(const TCHAR* InPath);
+
+	/** Add a name, referenced by ID (for curves etc.) */
+	void AppendName(uint32 InId, const TCHAR* InName);
+
+private:
+	/** Handle an object ending play - terminate existing scopes */
+	void HandleObjectEndPlay(uint64 InObjectId, double InTime, const FObjectInfo& InObjectInfo);
 
 private:
 	Trace::IAnalysisSession& Session;
+	FGameplayProvider& GameplayProvider;
 
 	/** Maps into timeline arrays per-object */
 	TMap<uint64, uint32> ObjectIdToTickRecordTimelineStorage;
@@ -47,15 +62,19 @@ private:
 	/** Map of skeletal mesh->index into SkeletalMeshInfos index */
 	TMap<uint64, int32> SkeletalMeshIdToIndexMap;
 
+	/** Map of name IDs to name string */
+	TMap<uint32, const TCHAR*> NameMap;
+	
 	struct FPerObjectTimelineStorage
 	{
 		TMap<uint64, uint32> AssetIdToTickRecordTimeline;
-		TArray<TSharedRef<Trace::TMonotonicTimeline<FTickRecordMessage>>> Timelines;
+		TArray<TSharedRef<Trace::TPointTimeline<FTickRecordMessage>>> Timelines;
 	};
 
 	/** Message storage */
 	TArray<TSharedRef<FPerObjectTimelineStorage>> PerObjectTimelineStorage;
-	TArray<TSharedRef<Trace::TMonotonicTimeline<FSkeletalMeshPoseMessage>>> SkeletalMeshPoseTimelines;
+	TArray<TSharedRef<Trace::TIntervalTimeline<FSkeletalMeshPoseMessage>>> SkeletalMeshPoseTimelines;
 	TPagedArray<FTransform> SkeletalMeshPoseTransforms;
+	TPagedArray<FSkeletalMeshNamedCurve> SkeletalMeshCurves;
 	TPagedArray<int32> SkeletalMeshParentIndices;
 };
\ No newline at end of file
diff --git a/Engine/Plugins/Animation/GameplayInsights/Source/GameplayInsights/Private/AnimationSharedData.cpp b/Engine/Plugins/Animation/GameplayInsights/Source/GameplayInsights/Private/AnimationSharedData.cpp
--- a/Engine/Plugins/Animation/GameplayInsights/Source/GameplayInsights/Private/AnimationSharedData.cpp
+++ b/Engine/Plugins/Animation/GameplayInsights/Source/GameplayInsights/Private/AnimationSharedData.cpp
@@ -17,16 +17,7 @@
 FAnimationSharedData::FAnimationSharedData(FGameplaySharedData& InGameplaySharedData)
 	: GameplaySharedData(InGameplaySharedData)
 	, AnalysisSession(nullptr)
-	, SelectedEventStartTime(0.0)
-	, SelectedEventEndTime(0.0)
-	, HoveredEventStartTime(0.0)
-	, HoveredEventEndTime(0.0)
-	, SelectionStartTime(0.0)
-	, SelectionEndTime(0.0)
 	, MarkerTime(0.0)
-	, bSelectedEventValid(false)
-	, bHoveredEventValid(false)
-	, bSelectionValid(false)
 	, bTimeMarkerValid(false)
 	, bAnimationTracksEnabled(true)
 {
@@ -37,10 +28,7 @@ void FAnimationSharedData::OnBeginSession(Insights::ITimingViewSession& InTiming
 	SkeletalMeshPoseTracks.Reset();
 	AnimationTickRecordsTracks.Reset();
 
-	SelectedEventChangedHandle = InTimingViewSession.OnSelectedEventChanged().AddRaw(this, &FAnimationSharedData::OnSelectedEventChanged);
-	HoveredEventChangedHandle = InTimingViewSession.OnHoveredEventChanged().AddRaw(this, &FAnimationSharedData::OnHoveredEventChanged);
 	TimeMarkerChangedHandle = InTimingViewSession.OnTimeMarkerChanged().AddRaw(this, &FAnimationSharedData::OnTimeMarkerChanged);
-	SelectionChangedHandle = InTimingViewSession.OnSelectionChanged().AddRaw(this, &FAnimationSharedData::OnSelectionChanged);
 }
 
 void FAnimationSharedData::OnEndSession(Insights::ITimingViewSession& InTimingViewSession)
@@ -48,10 +36,7 @@ void FAnimationSharedData::OnEndSession(Insights::ITimingViewSession& InTimingVi
 	SkeletalMeshPoseTracks.Reset();
 	AnimationTickRecordsTracks.Reset();
 
-	InTimingViewSession.OnSelectedEventChanged().Remove(SelectedEventChangedHandle);
-	InTimingViewSession.OnHoveredEventChanged().Remove(HoveredEventChangedHandle);
 	InTimingViewSession.OnTimeMarkerChanged().Remove(TimeMarkerChangedHandle);
-	InTimingViewSession.OnSelectionChanged().Remove(SelectionChangedHandle);
 }
 
 void FAnimationSharedData::Tick(Insights::ITimingViewSession& InTimingViewSession, const Trace::IAnalysisSession& InAnalysisSession)
@@ -132,6 +117,19 @@ void FAnimationSharedData::Tick(Insights::ITimingViewSession& InTimingViewSessio
 			});
 		});
 	}
+
+	// Prevent mouse movement throttling if we are drawing stuff that can change when the mouse is dragged
+	for(TSharedRef<FSkeletalMeshPoseTrack> PoseTrack : SkeletalMeshPoseTracks)
+	{
+		if(PoseTrack->IsVisible())
+		{
+			if(PoseTrack->ShouldDrawPose())
+			{
+				InTimingViewSession.PreventThrottling();
+				break;
+			}
+		}
+	}
 }
 
 void FAnimationSharedData::ExtendFilterMenu(FMenuBuilder& InMenuBuilder)
@@ -169,79 +167,31 @@ bool FAnimationSharedData::AreAnimationTracksEnabled() const
 	return bAnimationTracksEnabled;
 }
 
-void FAnimationSharedData::OnSelectedEventChanged(const TSharedPtr<const ITimingEvent> InEvent)
-{
-	bSelectedEventValid = InEvent.IsValid();
-	if(bSelectedEventValid)
-	{
-		SelectedEventTrack = InEvent->GetTrack();
-		SelectedEventStartTime = InEvent->GetStartTime();
-		SelectedEventEndTime = InEvent->GetEndTime();
-	}
-	else
-	{
-		SelectedEventTrack = nullptr;
-		SelectedEventStartTime = 0.0;
-		SelectedEventEndTime = 0.0;
-	}
-}
-
-void FAnimationSharedData::OnHoveredEventChanged(const TSharedPtr<const ITimingEvent> InEvent)
-{
-	bHoveredEventValid = InEvent.IsValid();
-	if(bHoveredEventValid)
-	{
-		HoveredEventTrack = InEvent->GetTrack();
-		HoveredEventStartTime = InEvent->GetStartTime();
-		HoveredEventEndTime = InEvent->GetEndTime();
-	}
-	else
-	{
-		HoveredEventTrack = nullptr;
-		HoveredEventStartTime = 0.0;
-		HoveredEventEndTime = 0.0;
-	}
-}
-
 void FAnimationSharedData::OnTimeMarkerChanged(Insights::ETimeChangedFlags InFlags, double InTimeMarker)
 {
 	bTimeMarkerValid = InTimeMarker != std::numeric_limits<double>::infinity();
 	MarkerTime = InTimeMarker;
 }
 
-void FAnimationSharedData::OnSelectionChanged(Insights::ETimeChangedFlags InFlags, double InStartTime, double InEndTime)
+void FAnimationSharedData::EnumerateSkeletalMeshPoseTracks(TFunctionRef<void(const TSharedRef<FSkeletalMeshPoseTrack>&)> InCallback) const
 {
-	bSelectionValid = InStartTime < InEndTime;
-	SelectionStartTime = InStartTime;
-	SelectionEndTime = InEndTime;
+	for(const TSharedRef<FSkeletalMeshPoseTrack>& Track : SkeletalMeshPoseTracks)
+	{
+		InCallback(Track);
+	}
 }
 
 #if WITH_ENGINE
 
-void FAnimationSharedData::DrawPoses(ULineBatchComponent* InLineBatcher)
+void FAnimationSharedData::DrawPoses(UWorld* InWorld)
 {
 	for(TSharedRef<FSkeletalMeshPoseTrack> PoseTrack : SkeletalMeshPoseTracks)
 	{
 		if(PoseTrack->IsVisible())
 		{
-			if(bSelectedEventValid && PoseTrack == SelectedEventTrack && PoseTrack->ShouldDrawSelectedEvent())
-			{
-				PoseTrack->DrawPoses(InLineBatcher, SelectedEventStartTime, SelectedEventEndTime);
-			}
-		
-			if(bHoveredEventValid && PoseTrack == HoveredEventTrack && PoseTrack->ShouldDrawHoveredEvent())
-			{
-				PoseTrack->DrawPoses(InLineBatcher, HoveredEventStartTime, HoveredEventEndTime);
-			}
-
-			if(bTimeMarkerValid && PoseTrack->ShouldDrawMarkerTime())
-			{
-				PoseTrack->DrawPoses(InLineBatcher, MarkerTime, MarkerTime);
-			}
-
-			if(bSelectionValid && PoseTrack->ShouldDrawSelection())
+			if(bTimeMarkerValid && PoseTrack->ShouldDrawPose())
 			{
-				PoseTrack->DrawPoses(InLineBatcher, SelectionStartTime, SelectionEndTime);
+				PoseTrack->DrawPoses(InWorld, MarkerTime);
 			}
 		}
 	}
diff --git a/Engine/Plugins/Animation/GameplayInsights/Source/GameplayInsights/Private/AnimationSharedData.h b/Engine/Plugins/Animation/GameplayInsights/Source/GameplayInsights/Private/AnimationSharedData.h
--- a/Engine/Plugins/Animation/GameplayInsights/Source/GameplayInsights/Private/AnimationSharedData.h
+++ b/Engine/Plugins/Animation/GameplayInsights/Source/GameplayInsights/Private/AnimationSharedData.h
@@ -11,7 +11,7 @@ namespace Insights { enum class ETimeChangedFlags : int32; }
 class FSkeletalMeshPoseTrack;
 class FAnimationTickRecordsTrack;
 class FMenuBuilder;
-class ULineBatchComponent;
+class UWorld;
 
 class FAnimationSharedData
 {
@@ -24,7 +24,7 @@ public:
 	void ExtendFilterMenu(FMenuBuilder& InMenuBuilder);
 
 #if WITH_ENGINE
-	void DrawPoses(ULineBatchComponent* InLineBatcher);
+	void DrawPoses(UWorld* InWorld);
 #endif
 
 	// Check whether animation tracks are enabled
@@ -36,13 +36,18 @@ public:
 	// Check whether the analysis session is valid
 	bool IsAnalysisSessionValid() const { return AnalysisSession != nullptr; }
 
+	// Get the gameplay shared data we are linked to
+	const FGameplaySharedData& GetGameplaySharedData() const { return GameplaySharedData; }
+
+	// Enumerate skeletal mesh pose tracks
+	void EnumerateSkeletalMeshPoseTracks(TFunctionRef<void(const TSharedRef<FSkeletalMeshPoseTrack>&)> InCallback) const;
+
 private:
 	// UI handlers
 	void ToggleAnimationTracks();
-	void OnSelectedEventChanged(const TSharedPtr<const ITimingEvent> InEvent);
-	void OnHoveredEventChanged(const TSharedPtr<const ITimingEvent> InEvent);
 	void OnTimeMarkerChanged(Insights::ETimeChangedFlags InFlags, double InTimeMarker);
-	void OnSelectionChanged(Insights::ETimeChangedFlags InFlags, double InStartTime, double InEndTime);
+	void ToggleAnimationParentTracks();
+	bool AreAnimationParentTracksEnabled() const;
 
 private:
 	// The gameplay shared data we are linked to
@@ -56,28 +61,12 @@ private:
 	TArray<TSharedRef<FAnimationTickRecordsTrack>> AnimationTickRecordsTracks;
 
 	// Delegate handles for hooks into the timing view
-	FDelegateHandle SelectedEventChangedHandle;
-	FDelegateHandle HoveredEventChangedHandle;
 	FDelegateHandle TimeMarkerChangedHandle;
-	FDelegateHandle SelectionChangedHandle;
-
-	/** Selected/hovered tracks */
-	TSharedPtr<const FBaseTimingTrack> SelectedEventTrack;
-	TSharedPtr<const FBaseTimingTrack> HoveredEventTrack;
 
 	/** Various times and ranges */
-	double SelectedEventStartTime;
-	double SelectedEventEndTime;
-	double HoveredEventStartTime;
-	double HoveredEventEndTime;
-	double SelectionStartTime;
-	double SelectionEndTime;
 	double MarkerTime;
 
 	/** Validity flags for pose times/ranges */
-	bool bSelectedEventValid;
-	bool bHoveredEventValid;
-	bool bSelectionValid;
 	bool bTimeMarkerValid;
 
 	// Whether all of our animation tracks are enabled
diff --git a/Engine/Plugins/Animation/GameplayInsights/Source/GameplayInsights/Private/AnimationTickRecordsTrack.cpp b/Engine/Plugins/Animation/GameplayInsights/Source/GameplayInsights/Private/AnimationTickRecordsTrack.cpp
--- a/Engine/Plugins/Animation/GameplayInsights/Source/GameplayInsights/Private/AnimationTickRecordsTrack.cpp
+++ b/Engine/Plugins/Animation/GameplayInsights/Source/GameplayInsights/Private/AnimationTickRecordsTrack.cpp
@@ -33,20 +33,10 @@ void FAnimationTickRecordsTrack::BuildDrawState(ITimingEventsTrackDrawStateBuild
 
 		AnimationProvider->ReadTickRecordTimeline(GetGameplayTrack().GetObjectId(), GetAssetId(), [this, &GameplayProvider, &AnimationProvider, &Context, &Builder](const FAnimationProvider::TickRecordTimeline& InTimeline)
 		{
-			auto DrawEvents = [this, &Builder, &GameplayProvider](double InStartTime, double InEndTime, uint32 InDepth, const FTickRecordMessage& InMessage)
+			InTimeline.EnumerateEvents(Context.GetViewport().GetStartTime(), Context.GetViewport().GetEndTime(), [this, &Builder, &GameplayProvider](double InStartTime, double InEndTime, uint32 InDepth, const FTickRecordMessage& InMessage)
 			{
 				Builder.AddEvent(InStartTime, InEndTime, 0, *GetName());
-			};
-
-			if (FTimingEventsTrack::bUseDownSampling)
-			{
-				const double SecondsPerPixel = 1.0 / Context.GetViewport().GetScaleX();
-				InTimeline.EnumerateEventsDownSampled(Context.GetViewport().GetStartTime(), Context.GetViewport().GetEndTime(), SecondsPerPixel, DrawEvents);
-			}
-			else
-			{
-				InTimeline.EnumerateEvents(Context.GetViewport().GetStartTime(), Context.GetViewport().GetEndTime(), DrawEvents);
-			}
+			});
 		});
 	}
 }
diff --git a/Engine/Plugins/Animation/GameplayInsights/Source/GameplayInsights/Private/GameplayInsightsModule.cpp b/Engine/Plugins/Animation/GameplayInsights/Source/GameplayInsights/Private/GameplayInsightsModule.cpp
--- a/Engine/Plugins/Animation/GameplayInsights/Source/GameplayInsights/Private/GameplayInsightsModule.cpp
+++ b/Engine/Plugins/Animation/GameplayInsights/Source/GameplayInsights/Private/GameplayInsightsModule.cpp
@@ -6,7 +6,6 @@
 #include "GameplayTimingViewExtender.h"
 #include "Insights/ITimingViewExtender.h"
 #include "Containers/Ticker.h"
-#include "Misc/CoreDelegates.h"
 #include "Modules/ModuleManager.h"
 #include "Framework/Multibox/MultiboxBuilder.h"
 
@@ -23,22 +22,6 @@ public:
 		IModularFeatures::Get().RegisterModularFeature(Trace::ModuleFeatureName, &GameplayTraceModule);
 		IModularFeatures::Get().RegisterModularFeature(Insights::TimingViewExtenderFeatureName, &GameplayTimingViewExtender);
 
-
-		PostEngineInitHandle = FCoreDelegates::OnPostEngineInit.AddLambda([this]()
-		{
-#if WITH_ENGINE
-			if(GEngine)
-			{
-				for(const FWorldContext& WorldContext : GEngine->GetWorldContexts())
-				{
-					GameplayTimingViewExtender.AddVisualizerWorld(WorldContext.World());
-				}
-
-				OnWorldAddedHandle = GEngine->OnWorldAdded().AddRaw(this, &FGameplayInsightsModule::HandleWorldAdded);
-			}
-#endif
-		});
-
 		TickerHandle = FTicker::GetCoreTicker().AddTicker(TEXT("GameplayInsights"), 0.0f, [this](float DeltaTime)
 		{
 			GameplayTimingViewExtender.TickVisualizers(DeltaTime);
@@ -50,37 +33,14 @@ public:
 	{
 		FTicker::GetCoreTicker().RemoveTicker(TickerHandle);
 
-		FCoreDelegates::OnPostEngineInit.Remove(PostEngineInitHandle);
-
-#if WITH_ENGINE
-		if(GEngine)
-		{
-			GEngine->OnWorldAdded().Remove(OnWorldAddedHandle);
-		}
-#endif
-
 		IModularFeatures::Get().UnregisterModularFeature(Trace::ModuleFeatureName, &GameplayTraceModule);
 		IModularFeatures::Get().UnregisterModularFeature(Insights::TimingViewExtenderFeatureName, &GameplayTimingViewExtender);
 	}
 
-#if WITH_ENGINE
-	void HandleWorldAdded(UWorld* InWorld)
-	{
-		if (InWorld->WorldType == EWorldType::Editor ||
-			InWorld->WorldType == EWorldType::PIE ||
-			InWorld->WorldType == EWorldType::Game)
-		{
-			GameplayTimingViewExtender.AddVisualizerWorld(InWorld);
-		}
-	}
-#endif
-
 	FGameplayTraceModule GameplayTraceModule;
 	FGameplayTimingViewExtender GameplayTimingViewExtender;
 
-	FDelegateHandle OnWorldAddedHandle;
 	FDelegateHandle TickerHandle;
-	FDelegateHandle PostEngineInitHandle;
 };
 
 IMPLEMENT_MODULE(FGameplayInsightsModule, GameplayInsights);
diff --git a/Engine/Plugins/Animation/GameplayInsights/Source/GameplayInsights/Private/GameplayProvider.cpp b/Engine/Plugins/Animation/GameplayInsights/Source/GameplayInsights/Private/GameplayProvider.cpp
--- a/Engine/Plugins/Animation/GameplayInsights/Source/GameplayInsights/Private/GameplayProvider.cpp
+++ b/Engine/Plugins/Animation/GameplayInsights/Source/GameplayInsights/Private/GameplayProvider.cpp
@@ -6,6 +6,7 @@ FName FGameplayProvider::ProviderName("GameplayProvider");
 
 FGameplayProvider::FGameplayProvider(Trace::IAnalysisSession& InSession)
 	: Session(InSession)
+	, EndPlayEvent(nullptr)
 {
 }
 
@@ -120,6 +121,12 @@ void FGameplayProvider::AppendObjectEvent(uint64 InObjectId, double InTime, cons
 {
 	Session.WriteAccessCheck();
 
+	// Important events need some extra routing
+	if(EndPlayEvent == nullptr)
+	{
+		EndPlayEvent = Session.StoreString(TEXT("EndPlay"));
+	}
+
 	TSharedPtr<Trace::TPointTimeline<FObjectEventMessage>> Timeline;
 	uint32* IndexPtr = ObjectIdToEventTimelines.Find(InObjectId);
 	if(IndexPtr != nullptr)
@@ -137,6 +144,14 @@ void FGameplayProvider::AppendObjectEvent(uint64 InObjectId, double InTime, cons
 	Message.Id = InObjectId;
 	Message.Name = Session.StoreString(InEventName);
 
+	if(Message.Name == EndPlayEvent)
+	{
+		if(int32* ObjectInfoIndex = ObjectIdToIndexMap.Find(InObjectId))
+		{
+			OnObjectEndPlayDelegate.Broadcast(InObjectId, InTime, ObjectInfos[*ObjectInfoIndex]);
+		}
+	}
+
 	Timeline->AppendEvent(InTime, Message);
 
 	Session.UpdateDurationSeconds(InTime);
diff --git a/Engine/Plugins/Animation/GameplayInsights/Source/GameplayInsights/Private/GameplayProvider.h b/Engine/Plugins/Animation/GameplayInsights/Source/GameplayInsights/Private/GameplayProvider.h
--- a/Engine/Plugins/Animation/GameplayInsights/Source/GameplayInsights/Private/GameplayProvider.h
+++ b/Engine/Plugins/Animation/GameplayInsights/Source/GameplayInsights/Private/GameplayProvider.h
@@ -22,6 +22,7 @@ public:
 	virtual void EnumerateObjects(TFunctionRef<void(const FObjectInfo&)> Callback) const override;
 	virtual const FClassInfo* FindClassInfo(uint64 InClassId) const override;
 	virtual const FObjectInfo* FindObjectInfo(uint64 InObjectId) const override;
+	virtual FOnObjectEndPlay& OnObjectEndPlay() override { return OnObjectEndPlayDelegate; }
 
 	/** Add a class message */
 	void AppendClass(uint64 InClassId, uint64 InSuperId, const TCHAR* InClassName, const TCHAR* InClassPathName);
@@ -52,4 +53,10 @@ private:
 
 	/** Message storage */
 	TArray<TSharedRef<Trace::TPointTimeline<FObjectEventMessage>>> EventTimelines;
+
+	/** EndPlay event text */
+	const TCHAR* EndPlayEvent;
+
+	/** Delegate fired when an object receives an end play event */
+	FOnObjectEndPlay OnObjectEndPlayDelegate;
 };
\ No newline at end of file
diff --git a/Engine/Plugins/Animation/GameplayInsights/Source/GameplayInsights/Private/GameplaySharedData.cpp b/Engine/Plugins/Animation/GameplayInsights/Source/GameplayInsights/Private/GameplaySharedData.cpp
--- a/Engine/Plugins/Animation/GameplayInsights/Source/GameplayInsights/Private/GameplaySharedData.cpp
+++ b/Engine/Plugins/Animation/GameplayInsights/Source/GameplayInsights/Private/GameplaySharedData.cpp
@@ -193,4 +193,12 @@ bool FGameplaySharedData::AreGameplayTracksEnabled() const
 	return bObjectTracksEnabled;
 }
 
+void FGameplaySharedData::EnumerateObjectTracks(TFunctionRef<void(const TSharedRef<FObjectEventsTrack>&)> InCallback) const
+{
+	for(const auto& TrackPair : ObjectTracks)
+	{
+		InCallback(TrackPair.Value.ToSharedRef());
+	}
+}
+
 #undef LOCTEXT_NAMESPACE
\ No newline at end of file
diff --git a/Engine/Plugins/Animation/GameplayInsights/Source/GameplayInsights/Private/GameplaySharedData.h b/Engine/Plugins/Animation/GameplayInsights/Source/GameplayInsights/Private/GameplaySharedData.h
--- a/Engine/Plugins/Animation/GameplayInsights/Source/GameplayInsights/Private/GameplaySharedData.h
+++ b/Engine/Plugins/Animation/GameplayInsights/Source/GameplayInsights/Private/GameplaySharedData.h
@@ -33,6 +33,9 @@ public:
 	// Get the last cached analysis session
 	const Trace::IAnalysisSession& GetAnalysisSession() const { return *AnalysisSession; }
 
+	// Enumerate object tracks
+	void EnumerateObjectTracks(TFunctionRef<void(const TSharedRef<FObjectEventsTrack>&)> InCallback) const;
+
 private:
 	// Re-sort tracks if trrack ordering has changed
 	void SortTracks();
diff --git a/Engine/Plugins/Animation/GameplayInsights/Source/GameplayInsights/Private/GameplayTimingViewExtender.cpp b/Engine/Plugins/Animation/GameplayInsights/Source/GameplayInsights/Private/GameplayTimingViewExtender.cpp
--- a/Engine/Plugins/Animation/GameplayInsights/Source/GameplayInsights/Private/GameplayTimingViewExtender.cpp
+++ b/Engine/Plugins/Animation/GameplayInsights/Source/GameplayInsights/Private/GameplayTimingViewExtender.cpp
@@ -8,6 +8,7 @@
 
 #if WITH_ENGINE
 #include "Engine/World.h"
+#include "Editor/EditorEngine.h"
 #endif
 
 void FGameplayTimingViewExtender::OnBeginSession(Insights::ITimingViewSession& InSession)
@@ -62,27 +63,39 @@ void FGameplayTimingViewExtender::ExtendFilterMenu(Insights::ITimingViewSession&
 }
 
 #if WITH_ENGINE
-void FGameplayTimingViewExtender::AddVisualizerWorld(UWorld* InWorld)
+
+static UWorld* GetWorldToVisualize()
 {
-	Worlds.Add(InWorld);
+	UWorld* World = nullptr;
+
+#if WITH_EDITOR
+	UEditorEngine* EditorEngine = Cast<UEditorEngine>(GEngine);
+	if (GIsEditor && EditorEngine != nullptr && World == nullptr)
+	{
+		// lets use PlayWorld during PIE/Simulate and regular world from editor otherwise, to draw debug information
+		World = EditorEngine->PlayWorld != nullptr ? EditorEngine->PlayWorld : EditorEngine->GetEditorWorldContext().World();
+	}
+
+#endif
+	if (!GIsEditor && World == nullptr)
+	{
+		World = GEngine->GetWorld();
+	}
+
+	return World;
 }
+
 #endif
 
 void FGameplayTimingViewExtender::TickVisualizers(float DeltaTime)
 {
 #if WITH_ENGINE
-	// Trim invalid worlds
-	Worlds.RemoveAll([](const TWeakObjectPtr<UWorld>& InWorld){ return InWorld.Get() == nullptr; });
-
-	for(auto& PerSessionData : PerSessionDataMap)
+	UWorld* WorldToVisualize = GetWorldToVisualize();
+	if(WorldToVisualize)
 	{
-		// Draw using line batchers
-		for(TWeakObjectPtr<UWorld>& World : Worlds)
+		for(auto& PerSessionData : PerSessionDataMap)
 		{
-			if(World->LineBatcher)
-			{
-				PerSessionData.Value.AnimationSharedData->DrawPoses(World->LineBatcher);
-			}
+			PerSessionData.Value.AnimationSharedData->DrawPoses(WorldToVisualize);
 		}
 	}
 #endif
diff --git a/Engine/Plugins/Animation/GameplayInsights/Source/GameplayInsights/Private/GameplayTimingViewExtender.h b/Engine/Plugins/Animation/GameplayInsights/Source/GameplayInsights/Private/GameplayTimingViewExtender.h
--- a/Engine/Plugins/Animation/GameplayInsights/Source/GameplayInsights/Private/GameplayTimingViewExtender.h
+++ b/Engine/Plugins/Animation/GameplayInsights/Source/GameplayInsights/Private/GameplayTimingViewExtender.h
@@ -18,11 +18,6 @@ public:
 	virtual void Tick(Insights::ITimingViewSession& InSession, const Trace::IAnalysisSession& InAnalysisSession) override;
 	virtual void ExtendFilterMenu(Insights::ITimingViewSession& InSession, FMenuBuilder& InMenuBuilder) override;
 
-#if WITH_ENGINE
-	// Add a world to be visualized
-	void AddVisualizerWorld(UWorld* InWorld);
-#endif
-
 	// Tick the visualizers
 	void TickVisualizers(float DeltaTime);
 
@@ -36,9 +31,4 @@ private:
 
 	// The data we host per-session
 	TMap<Insights::ITimingViewSession*, FPerSessionData> PerSessionDataMap;
-
-#if WITH_ENGINE
-	// The worlds we will visualize
-	TArray<TWeakObjectPtr<UWorld>> Worlds;
-#endif
 };
diff --git a/Engine/Plugins/Animation/GameplayInsights/Source/GameplayInsights/Private/GameplayTraceModule.cpp b/Engine/Plugins/Animation/GameplayInsights/Source/GameplayInsights/Private/GameplayTraceModule.cpp
--- a/Engine/Plugins/Animation/GameplayInsights/Source/GameplayInsights/Private/GameplayTraceModule.cpp
+++ b/Engine/Plugins/Animation/GameplayInsights/Source/GameplayInsights/Private/GameplayTraceModule.cpp
@@ -18,7 +18,7 @@ void FGameplayTraceModule::OnAnalysisBegin(Trace::IAnalysisSession& InSession)
 {
 	FGameplayProvider* GameplayProvider = new FGameplayProvider(InSession);
 	InSession.AddProvider(FGameplayProvider::ProviderName, GameplayProvider);
-	FAnimationProvider* AnimationProvider = new FAnimationProvider(InSession);
+	FAnimationProvider* AnimationProvider = new FAnimationProvider(InSession, *GameplayProvider);
 	InSession.AddProvider(FAnimationProvider::ProviderName, AnimationProvider);
 
 	InSession.AddAnalyzer(new FAnimationAnalyzer(InSession, *AnimationProvider));
diff --git a/Engine/Plugins/Animation/GameplayInsights/Source/GameplayInsights/Private/IAnimationProvider.h b/Engine/Plugins/Animation/GameplayInsights/Source/GameplayInsights/Private/IAnimationProvider.h
--- a/Engine/Plugins/Animation/GameplayInsights/Source/GameplayInsights/Private/IAnimationProvider.h
+++ b/Engine/Plugins/Animation/GameplayInsights/Source/GameplayInsights/Private/IAnimationProvider.h
@@ -14,12 +14,21 @@ struct FSkeletalMeshInfo
 	uint32 BoneCount = 0;
 };
 
+struct FSkeletalMeshNamedCurve
+{
+	uint32 Id = 0;
+	float Value = 0.0f;
+};
+
 struct FSkeletalMeshPoseMessage
 {
+	FTransform ComponentToWorld;
 	uint64 TransformStartIndex = 0;
+	uint64 CurveStartIndex = 0;
 	uint64 ComponentId = 0;	
 	uint64 MeshId = 0;	
 	uint16 NumTransforms = 0;
+	uint16 NumCurves = 0;
 	uint16 FrameCounter = 0;
 	uint16 LodIndex = 0;
 };
@@ -45,9 +54,10 @@ public:
 
 	virtual bool ReadSkeletalMeshPoseTimeline(uint64 InObjectId, TFunctionRef<void(const SkeletalMeshPoseTimeline&)> Callback) const = 0;
 	virtual bool ReadSkeletalMeshPoseMessage(uint64 InObjectId, uint64 InMessageId, TFunctionRef<void(const FSkeletalMeshPoseMessage&)> Callback) const = 0;
-	virtual void EnumerateSkeletalMeshPose(const FSkeletalMeshPoseMessage& InMessage, const FSkeletalMeshInfo& InMeshInfo, TFunctionRef<void(const FTransform&, const FTransform&)> Callback) const = 0;
+	virtual void GetSkeletalMeshComponentSpacePose(const FSkeletalMeshPoseMessage& InMessage, const FSkeletalMeshInfo& InMeshInfo, FTransform& OutComponentToWorld, TArray<FTransform>& OutTransforms) const = 0;
 	virtual void EnumerateTickRecordTimelines(uint64 InObjectId, TFunctionRef<void(uint64, const TickRecordTimeline&)> Callback) const = 0;
 	virtual bool ReadTickRecordTimeline(uint64 InObjectId, uint64 InAssetId, TFunctionRef<void(const TickRecordTimeline&)> Callback) const = 0;
 	virtual bool ReadTickRecordMessage(uint64 InObjectId, uint64 InAssetId, uint64 InMessageId, TFunctionRef<void(const FTickRecordMessage&)> Callback) const = 0;
 	virtual const FSkeletalMeshInfo* FindSkeletalMeshInfo(uint64 InObjectId) const = 0;
+	virtual const TCHAR* GetName(uint32 InId) const = 0;
 };
\ No newline at end of file
diff --git a/Engine/Plugins/Animation/GameplayInsights/Source/GameplayInsights/Private/IGameplayProvider.h b/Engine/Plugins/Animation/GameplayInsights/Source/GameplayInsights/Private/IGameplayProvider.h
--- a/Engine/Plugins/Animation/GameplayInsights/Source/GameplayInsights/Private/IGameplayProvider.h
+++ b/Engine/Plugins/Animation/GameplayInsights/Source/GameplayInsights/Private/IGameplayProvider.h
@@ -28,6 +28,9 @@ struct FObjectEventMessage
 	const TCHAR* Name = nullptr;
 };
 
+// Delegate fired when an object receives an end play event
+DECLARE_MULTICAST_DELEGATE_ThreeParams(FOnObjectEndPlay, uint64 /*ObjectId*/, double /*Time*/, const FObjectInfo& /*ObjectInfo*/);
+
 class IGameplayProvider : public Trace::IProvider
 {
 public:
@@ -38,4 +41,5 @@ public:
 	virtual void EnumerateObjects(TFunctionRef<void(const FObjectInfo&)> Callback) const = 0;
 	virtual const FClassInfo* FindClassInfo(uint64 InClassId) const = 0;
 	virtual const FObjectInfo* FindObjectInfo(uint64 InObjectId) const = 0;
+	virtual FOnObjectEndPlay& OnObjectEndPlay() = 0;
 };
\ No newline at end of file
diff --git a/Engine/Plugins/Animation/GameplayInsights/Source/GameplayInsights/Private/InsightsSkeletalMeshComponent.cpp b/Engine/Plugins/Animation/GameplayInsights/Source/GameplayInsights/Private/InsightsSkeletalMeshComponent.cpp
new file mode 100644
--- /dev/null
+++ b/Engine/Plugins/Animation/GameplayInsights/Source/GameplayInsights/Private/InsightsSkeletalMeshComponent.cpp
@@ -0,0 +1,45 @@
+// Copyright 1998-2019 Epic Games, Inc. All Rights Reserved.
+
+#include "InsightsSkeletalMeshComponent.h"
+#include "AnimationProvider.h"
+
+void UInsightsSkeletalMeshComponent::SetPoseFromProvider(const FAnimationProvider& InProvider, const FSkeletalMeshPoseMessage& InMessage, const FSkeletalMeshInfo& SkeletalMeshInfo)
+{
+	// Grab transform and bone transforms from provider
+	FTransform ComponentToWorldTransform;
+	InProvider.GetSkeletalMeshComponentSpacePose(InMessage, SkeletalMeshInfo, ComponentToWorldTransform, GetEditableComponentSpaceTransforms());
+
+	// Set transform
+	SetRelativeTransform(ComponentToWorldTransform);
+
+	// Force LOD to message LOD
+	SetForcedLOD(InMessage.LodIndex + 1);
+
+	// Flip buffers once to copy the directly-written component space transforms
+	bNeedToFlipSpaceBaseBuffers = true;
+	bHasValidBoneTransform = false;
+	FlipEditableSpaceBases();
+	bHasValidBoneTransform = true;
+
+	InvalidateCachedBounds();
+	UpdateBounds();
+	MarkRenderTransformDirty();
+	MarkRenderDynamicDataDirty();
+}
+
+void UInsightsSkeletalMeshComponent::InitAnim(bool bForceReInit)
+{
+	if(SkeletalMesh)
+	{
+		// set up bone visibility states as this gets skipped since we allocate the component array before registration
+		for (int32 BaseIndex = 0; BaseIndex < 2; ++BaseIndex)
+		{
+			BoneVisibilityStates[BaseIndex].SetNum(SkeletalMesh->RefSkeleton.GetNum());
+
+			for (int32 BoneIndex = 0; BoneIndex < SkeletalMesh->RefSkeleton.GetNum(); BoneIndex++)
+			{
+				BoneVisibilityStates[BaseIndex][BoneIndex] = BVS_Visible;
+			}
+		}
+	}
+}
\ No newline at end of file
diff --git a/Engine/Plugins/Animation/GameplayInsights/Source/GameplayInsights/Private/InsightsSkeletalMeshComponent.h b/Engine/Plugins/Animation/GameplayInsights/Source/GameplayInsights/Private/InsightsSkeletalMeshComponent.h
new file mode 100644
--- /dev/null
+++ b/Engine/Plugins/Animation/GameplayInsights/Source/GameplayInsights/Private/InsightsSkeletalMeshComponent.h
@@ -0,0 +1,24 @@
+// Copyright 1998-2019 Epic Games, Inc. All Rights Reserved.
+
+#pragma once
+
+#include "CoreMinimal.h"
+#include "Components/SkeletalMeshComponent.h"
+#include "InsightsSkeletalMeshComponent.generated.h"
+
+class FAnimationProvider;
+struct FSkeletalMeshPoseMessage;
+struct FSkeletalMeshInfo;
+
+UCLASS()
+class UInsightsSkeletalMeshComponent : public USkeletalMeshComponent
+{
+	GENERATED_BODY()
+
+public:
+	// Set this component up from a provider & message
+	void SetPoseFromProvider(const FAnimationProvider& InProvider, const FSkeletalMeshPoseMessage& InMessage, const FSkeletalMeshInfo& SkeletalMeshInfo);
+
+	// USkeletalMeshComponent interface
+	virtual void InitAnim(bool bForceReInit) override;
+};
diff --git a/Engine/Plugins/Animation/GameplayInsights/Source/GameplayInsights/Private/SkeletalMeshPoseTrack.cpp b/Engine/Plugins/Animation/GameplayInsights/Source/GameplayInsights/Private/SkeletalMeshPoseTrack.cpp
--- a/Engine/Plugins/Animation/GameplayInsights/Source/GameplayInsights/Private/SkeletalMeshPoseTrack.cpp
+++ b/Engine/Plugins/Animation/GameplayInsights/Source/GameplayInsights/Private/SkeletalMeshPoseTrack.cpp
@@ -9,9 +9,13 @@
 #include "Framework/Multibox/MultiboxBuilder.h"
 #include "Insights/ViewModels/TimingEventSearch.h"
 #include "Insights/ViewModels/TooltipDrawState.h"
+#include "AssetRegistryModule.h"
+#include "Modules/ModuleManager.h"
 
 #if WITH_ENGINE
-#include "Components/LineBatchComponent.h"
+#include "Engine/SkeletalMesh.h"
+#include "UObject/SoftObjectPtr.h"
+#include "InsightsSkeletalMeshComponent.h"
 #endif
 
 
@@ -23,11 +27,20 @@ const FName FSkeletalMeshPoseTrack::SubTypeName(TEXT("SkeletalMeshPose"));
 FSkeletalMeshPoseTrack::FSkeletalMeshPoseTrack(const FAnimationSharedData& InSharedData, uint64 InObjectID, const TCHAR* InName)
 	: TGameplayTrackMixin<FTimingEventsTrack>(InObjectID, FSkeletalMeshPoseTrack::TypeName, FSkeletalMeshPoseTrack::SubTypeName, FText::Format(LOCTEXT("TrackNameFormat", "Pose - {0}"), FText::FromString(FString(InName))))
 	, SharedData(InSharedData)
-	, bDrawMarkerTime(false)
-	, bDrawSelectedEvent(false)
-	, bDrawHoveredEvent(false)
-	, bDrawSelection(false)
+	, Color(FLinearColor::MakeRandomColor())
+	, bDrawPose(false)
+	, bDrawSkeleton(false)
 {
+#if WITH_ENGINE
+	OnWorldDestroyedHandle = FWorldDelegates::OnWorldCleanup.AddRaw(this, &FSkeletalMeshPoseTrack::OnWorldCleanup);
+#endif
+}
+
+FSkeletalMeshPoseTrack::~FSkeletalMeshPoseTrack()
+{
+#if WITH_ENGINE
+	FWorldDelegates::OnWorldCleanup.Remove(OnWorldDestroyedHandle);
+#endif
 }
 
 void FSkeletalMeshPoseTrack::BuildDrawState(ITimingEventsTrackDrawStateBuilder& Builder, const ITimingTrackUpdateContext& Context)
@@ -40,22 +53,12 @@ void FSkeletalMeshPoseTrack::BuildDrawState(ITimingEventsTrackDrawStateBuilder&
 
 		AnimationProvider->ReadSkeletalMeshPoseTimeline(GetGameplayTrack().GetObjectId(), [&Context, &Builder](const FAnimationProvider::SkeletalMeshPoseTimeline& InTimeline)
 		{
-			auto DrawEvent = [&Builder](double InStartTime, double InEndTime, uint32 InDepth, const FSkeletalMeshPoseMessage& InMessage)
+			InTimeline.EnumerateEvents(Context.GetViewport().GetStartTime(), Context.GetViewport().GetEndTime(), [&Builder](double InStartTime, double InEndTime, uint32 InDepth, const FSkeletalMeshPoseMessage& InMessage)
 			{
 				static TCHAR Buffer[256];
 				FCString::Snprintf(Buffer, 256, TEXT("%d Bones"), InMessage.NumTransforms);
 				Builder.AddEvent(InStartTime, InEndTime, 0, Buffer);
-			};
-
-			if (FTimingEventsTrack::bUseDownSampling)
-			{
-				const double SecondsPerPixel = 1.0 / Context.GetViewport().GetScaleX();
-				InTimeline.EnumerateEventsDownSampled(Context.GetViewport().GetStartTime(), Context.GetViewport().GetEndTime(), SecondsPerPixel, DrawEvent);
-			}
-			else
-			{
-				InTimeline.EnumerateEvents(Context.GetViewport().GetStartTime(), Context.GetViewport().GetEndTime(), DrawEvent);
-			}
+			});
 		});
 	}
 }
@@ -77,7 +80,26 @@ void FSkeletalMeshPoseTrack::InitTooltip(FTooltipDrawState& Tooltip, const ITimi
 		Tooltip.AddTitle(LOCTEXT("SkeletalMeshPoseTooltipTitle", "Skeletal Mesh Pose").ToString());
 
 		Tooltip.AddNameValueTextLine(LOCTEXT("EventTime", "Time").ToString(), FText::AsNumber(InFoundStartTime).ToString());
+
+		{
+			Trace::FAnalysisSessionReadScope SessionReadScope(SharedData.GetAnalysisSession());
+
+			const FGameplayProvider* GameplayProvider = SharedData.GetAnalysisSession().ReadProvider<FGameplayProvider>(FGameplayProvider::ProviderName);
+			const FObjectInfo* SkeletalMeshObjectInfo = GameplayProvider->FindObjectInfo(InMessage.MeshId);
+			if(SkeletalMeshObjectInfo != nullptr)
+			{
+				Tooltip.AddNameValueTextLine(LOCTEXT("Mesh", "Mesh").ToString(), SkeletalMeshObjectInfo->PathName);
+
+				FAssetRegistryModule& AssetRegistryModule = FModuleManager::LoadModuleChecked<FAssetRegistryModule>(TEXT("AssetRegistry"));
+				if(!AssetRegistryModule.Get().GetAssetByObjectPath(SkeletalMeshObjectInfo->PathName).IsValid())
+				{
+					Tooltip.AddTextLine(LOCTEXT("MeshNotFound", "Mesh not found").ToString(), FLinearColor::Red);
+				}
+			}
+		}
+
 		Tooltip.AddNameValueTextLine(LOCTEXT("BoneCount", "Bone Count").ToString(), FText::AsNumber(InMessage.NumTransforms).ToString());
+		Tooltip.AddNameValueTextLine(LOCTEXT("CurveCount", "Curve Count").ToString(), FText::AsNumber(InMessage.NumCurves).ToString());
 
 		Tooltip.UpdateLayout();
 	});
@@ -126,98 +148,252 @@ void FSkeletalMeshPoseTrack::FindSkeletalMeshPoseMessage(const FTimingEventSearc
 
 void FSkeletalMeshPoseTrack::BuildContextMenu(FMenuBuilder& MenuBuilder)
 {
-	MenuBuilder.BeginSection(TEXT("ShowPoseSection"), LOCTEXT("ShowPose", "Show Pose"));
+	MenuBuilder.BeginSection(TEXT("DrawingSection"), LOCTEXT("Drawing", "Drawing (Component)"));
 	{
 		MenuBuilder.AddMenuEntry(
-			LOCTEXT("ToggleDrawMarkerTime", "Marker Time"),
-			LOCTEXT("ToggleDrawMarkerTime_Tooltip", "Draw the pose at the current marker time"),
+			LOCTEXT("ToggleDrawPose", "Draw Pose"),
+			LOCTEXT("ToggleDrawPose_Tooltip", "Draw the poses in this track"),
 			FSlateIcon(),
 			FUIAction(
-				FExecuteAction::CreateLambda([this](){ bDrawMarkerTime = !bDrawMarkerTime; }),
+				FExecuteAction::CreateLambda([this](){ bDrawPose = !bDrawPose; }),
 				FCanExecuteAction(),
-				FIsActionChecked::CreateLambda([this](){ return bDrawMarkerTime; })),
+				FIsActionChecked::CreateLambda([this](){ return bDrawPose; })),
 			NAME_None,
 			EUserInterfaceActionType::ToggleButton
 		);
 
 		MenuBuilder.AddMenuEntry(
-			LOCTEXT("ToggleDrawSelection", "Selection"),
-			LOCTEXT("ToggleDrawSelection_Tooltip", "Draw the pose for the currently selected event"),
+			LOCTEXT("ToggleDrawSkeleton", "Draw Skeleton"),
+			LOCTEXT("ToggleDrawSkeleton_Tooltip", "Draw the skeleton for poses in this track (when pose drawing is also enabled)"),
 			FSlateIcon(),
 			FUIAction(
-				FExecuteAction::CreateLambda([this](){ bDrawSelectedEvent = !bDrawSelectedEvent; }),
+				FExecuteAction::CreateLambda([this](){ bDrawSkeleton = !bDrawSkeleton; }),
 				FCanExecuteAction(),
-				FIsActionChecked::CreateLambda([this](){ return bDrawSelectedEvent; })),
+				FIsActionChecked::CreateLambda([this](){ return bDrawSkeleton; })),
 			NAME_None,
 			EUserInterfaceActionType::ToggleButton
 		);
+	}
+	MenuBuilder.EndSection();
 
-		MenuBuilder.AddMenuEntry(
-			LOCTEXT("ToggleDrawHovered", "Hovered"),
-			LOCTEXT("ToggleDrawSelection_Tooltip", "Draw the pose for the currently hovered event"),
-			FSlateIcon(),
-			FUIAction(
-				FExecuteAction::CreateLambda([this](){ bDrawHoveredEvent = !bDrawHoveredEvent; }),
-				FCanExecuteAction(),
-				FIsActionChecked::CreateLambda([this](){ return bDrawHoveredEvent; })),
-			NAME_None,
-			EUserInterfaceActionType::ToggleButton
-		);
+	const FGameplayProvider* GameplayProvider = SharedData.GetAnalysisSession().ReadProvider<FGameplayProvider>(FGameplayProvider::ProviderName);
 
-		MenuBuilder.AddMenuEntry(
-			LOCTEXT("ToggleDrawSelectedRange", "Selected Range"),
-			LOCTEXT("ToggleDrawSelectedRange_Tooltip", "Draw poses for the currently selected range"),
-			FSlateIcon(),
-			FUIAction(
-				FExecuteAction::CreateLambda([this](){ bDrawSelection = !bDrawSelection; }),
-				FCanExecuteAction(),
-				FIsActionChecked::CreateLambda([this](){ return bDrawSelection; })),
-			NAME_None,
-			EUserInterfaceActionType::ToggleButton
-		);
+	if(GameplayProvider)
+	{
+		Trace::FAnalysisSessionReadScope SessionReadScope(SharedData.GetAnalysisSession());
+	
+		const FObjectInfo* ComponentObjectInfo = GameplayProvider->FindObjectInfo(GetGameplayTrack().GetObjectId());
+		if(ComponentObjectInfo != nullptr)
+		{
+			// @FIXME: Outer does always equal owning actor, although does in nearly all cases with skeletal mesh components
+			const FObjectInfo* ActorObjectInfo = GameplayProvider->FindObjectInfo(ComponentObjectInfo->OuterId);
+			if(ActorObjectInfo != nullptr)
+			{
+				MenuBuilder.BeginSection(TEXT("DrawingSection"), FText::Format(LOCTEXT("DrawingActor", "Drawing ({0})"), FText::FromString(ActorObjectInfo->Name)));
+				{
+					MenuBuilder.AddMenuEntry(
+						LOCTEXT("ToggleDrawPoseActor", "Draw Pose for Actor"),
+						LOCTEXT("ToggleDrawPoseActor_Tooltip", "Draw the poses in this track and all other tracks for the current actor"),
+						FSlateIcon(),
+						FUIAction(
+							FExecuteAction::CreateLambda([this, ActorObjectInfo, GameplayProvider]()
+							{
+								Trace::FAnalysisSessionReadScope SessionReadScope(SharedData.GetAnalysisSession());
+
+								bool bSetDrawPose = true;
+								SharedData.EnumerateSkeletalMeshPoseTracks([GameplayProvider, ActorObjectInfo, &bSetDrawPose](const TSharedRef<FSkeletalMeshPoseTrack>& InTrack)
+								{
+									const FObjectInfo* OtherComponentObjectInfo = GameplayProvider->FindObjectInfo(InTrack->GetGameplayTrack().GetObjectId());
+									if(OtherComponentObjectInfo->OuterId == ActorObjectInfo->Id)
+									{
+										bSetDrawPose &= InTrack->bDrawPose;
+									}
+								});
+
+								SharedData.EnumerateSkeletalMeshPoseTracks([GameplayProvider, ActorObjectInfo, &bSetDrawPose](const TSharedRef<FSkeletalMeshPoseTrack>& InTrack)
+								{
+									const FObjectInfo* OtherComponentObjectInfo = GameplayProvider->FindObjectInfo(InTrack->GetGameplayTrack().GetObjectId());
+									if(OtherComponentObjectInfo->OuterId == ActorObjectInfo->Id)
+									{
+										InTrack->bDrawPose = !bSetDrawPose;
+									}
+								});
+							}),
+							FCanExecuteAction(),
+							FIsActionChecked::CreateLambda([this, ActorObjectInfo, GameplayProvider]()
+							{
+								Trace::FAnalysisSessionReadScope SessionReadScope(SharedData.GetAnalysisSession());
+
+								bool bDrawPoseSet = true;
+								SharedData.EnumerateSkeletalMeshPoseTracks([GameplayProvider, ActorObjectInfo, &bDrawPoseSet](const TSharedRef<FSkeletalMeshPoseTrack>& InTrack)
+								{
+									const FObjectInfo* OtherComponentObjectInfo = GameplayProvider->FindObjectInfo(InTrack->GetGameplayTrack().GetObjectId());
+									if(OtherComponentObjectInfo->OuterId == ActorObjectInfo->Id)
+									{
+										bDrawPoseSet &= InTrack->bDrawPose;
+									}
+								});
+
+								return bDrawPoseSet;
+							})),
+						NAME_None,
+						EUserInterfaceActionType::ToggleButton
+					);
+
+					MenuBuilder.AddMenuEntry(
+						LOCTEXT("ToggleDrawSkeletonActor", "Draw Skeleton for Actor"),
+						LOCTEXT("ToggleDrawSkeletonActor_Tooltip", "Draw the skeleton for poses in this track and all other tracks for the current actor (when pose drawing is also enabled)"),
+						FSlateIcon(),
+						FUIAction(
+							FExecuteAction::CreateLambda([this, ActorObjectInfo, GameplayProvider]()
+							{
+								Trace::FAnalysisSessionReadScope SessionReadScope(SharedData.GetAnalysisSession());
+
+								bool bSetDrawSkeleton = true;
+								SharedData.EnumerateSkeletalMeshPoseTracks([GameplayProvider, ActorObjectInfo, &bSetDrawSkeleton](const TSharedRef<FSkeletalMeshPoseTrack>& InTrack)
+								{
+									const FObjectInfo* OtherComponentObjectInfo = GameplayProvider->FindObjectInfo(InTrack->GetGameplayTrack().GetObjectId());
+									if(OtherComponentObjectInfo->OuterId == ActorObjectInfo->Id)
+									{
+										bSetDrawSkeleton &= InTrack->bDrawSkeleton;
+									}
+								});
+
+								SharedData.EnumerateSkeletalMeshPoseTracks([GameplayProvider, ActorObjectInfo, &bSetDrawSkeleton](const TSharedRef<FSkeletalMeshPoseTrack>& InTrack)
+								{
+									const FObjectInfo* OtherComponentObjectInfo = GameplayProvider->FindObjectInfo(InTrack->GetGameplayTrack().GetObjectId());
+									if(OtherComponentObjectInfo->OuterId == ActorObjectInfo->Id)
+									{
+										InTrack->bDrawSkeleton = !bSetDrawSkeleton;
+									}
+								});
+							}),
+							FCanExecuteAction(),
+							FIsActionChecked::CreateLambda([this, ActorObjectInfo, GameplayProvider]()
+							{
+								Trace::FAnalysisSessionReadScope SessionReadScope(SharedData.GetAnalysisSession());
+
+								bool bDrawSkeletonSet = true;
+								SharedData.EnumerateSkeletalMeshPoseTracks([GameplayProvider, ActorObjectInfo, &bDrawSkeletonSet](const TSharedRef<FSkeletalMeshPoseTrack>& InTrack)
+								{
+									const FObjectInfo* OtherComponentObjectInfo = GameplayProvider->FindObjectInfo(InTrack->GetGameplayTrack().GetObjectId());
+									if(OtherComponentObjectInfo->OuterId == ActorObjectInfo->Id)
+									{
+										bDrawSkeletonSet &= InTrack->bDrawSkeleton;
+									}
+								});
+
+								return bDrawSkeletonSet;
+							})),
+						NAME_None,
+						EUserInterfaceActionType::ToggleButton
+					);
+				}
+				MenuBuilder.EndSection();
+			}
+		}
 	}
-	MenuBuilder.EndSection();
 }
 
 #if WITH_ENGINE
 
-void FSkeletalMeshPoseTrack::DrawPoses(ULineBatchComponent* InLineBatcher, double SelectionStartTime, double SelectionEndTime)
+void FSkeletalMeshPoseTrack::OnWorldCleanup(UWorld* InWorld, bool bSessionEnded, bool bCleanupResources)
+{
+	FWorldComponentCache& CacheForWorld = GetWorldCache(InWorld);
+
+	if(CacheForWorld.Component)
+	{
+		CacheForWorld.Component->UnregisterComponent();
+		CacheForWorld.Component->MarkPendingKill();
+		CacheForWorld.Component = nullptr;
+	}
+
+	WorldCache.Remove(TWeakObjectPtr<UWorld>(InWorld));
+}
+
+void FSkeletalMeshPoseTrack::DrawPoses(UWorld* InWorld, double InTime)
 {
 	if(SharedData.IsAnalysisSessionValid())
 	{
 		const FAnimationProvider* AnimationProvider = SharedData.GetAnalysisSession().ReadProvider<FAnimationProvider>(FAnimationProvider::ProviderName);
+		const FGameplayProvider* GameplayProvider = SharedData.GetAnalysisSession().ReadProvider<FGameplayProvider>(FGameplayProvider::ProviderName);
 
-		if(AnimationProvider)
+		if(AnimationProvider && GameplayProvider)
 		{
 			Trace::FAnalysisSessionReadScope SessionReadScope(SharedData.GetAnalysisSession());
 
-			AnimationProvider->ReadSkeletalMeshPoseTimeline(GetGameplayTrack().GetObjectId(), [this, &InLineBatcher, &AnimationProvider, &SelectionStartTime, &SelectionEndTime](const FAnimationProvider::SkeletalMeshPoseTimeline& InTimeline)
+			FWorldComponentCache& CacheForWorld = GetWorldCache(InWorld);
+			if(CacheForWorld.Component)
 			{
-				static TArray<FBatchedLine> Lines;
+				CacheForWorld.Component->SetVisibility(false);
+			}
 
-				InTimeline.EnumerateEvents(SelectionStartTime, SelectionEndTime, [&InLineBatcher, &AnimationProvider, &SelectionStartTime, &SelectionEndTime](double InStartTime, double InEndTime, uint32 InDepth, const FSkeletalMeshPoseMessage& InMessage)
+			AnimationProvider->ReadSkeletalMeshPoseTimeline(GetGameplayTrack().GetObjectId(), [this, &CacheForWorld, &AnimationProvider, &GameplayProvider, &InTime](const FAnimationProvider::SkeletalMeshPoseTimeline& InTimeline)
+			{
+				InTimeline.EnumerateEvents(InTime, InTime, [this, &CacheForWorld, &AnimationProvider, &GameplayProvider, &InTime](double InStartTime, double InEndTime, uint32 InDepth, const FSkeletalMeshPoseMessage& InMessage)
 				{
-					if(SelectionStartTime == SelectionEndTime || (InStartTime >= SelectionStartTime && InEndTime <= SelectionEndTime))
+					if((InStartTime <= InTime && InEndTime > InTime))
 					{
 						const FSkeletalMeshInfo* SkeletalMeshInfo = AnimationProvider->FindSkeletalMeshInfo(InMessage.MeshId);
-						if(SkeletalMeshInfo)
+						const FObjectInfo* SkeletalMeshObjectInfo = GameplayProvider->FindObjectInfo(InMessage.MeshId);
+						if(SkeletalMeshInfo && SkeletalMeshObjectInfo)
 						{
-							AnimationProvider->EnumerateSkeletalMeshPose(InMessage, *SkeletalMeshInfo, [&InLineBatcher](const FTransform& InTransform, const FTransform& InParentTransform)
+							UInsightsSkeletalMeshComponent* Component = CacheForWorld.GetComponent();
+							Component->SetVisibility(true);
+
+							if(CacheForWorld.Time != InTime)
 							{
-								Lines.Emplace(InParentTransform.GetLocation(), InTransform.GetLocation(), FLinearColor::Red, 1.0f, 0.2f, SDPG_Foreground);
-							});
+								USkeletalMesh* SkeletalMesh = TSoftObjectPtr<USkeletalMesh>(FSoftObjectPath(SkeletalMeshObjectInfo->PathName)).LoadSynchronous();
+								if(SkeletalMesh)
+								{
+									Component->SetSkeletalMesh(SkeletalMesh);
+								}
+
+								Component->SetPoseFromProvider(*AnimationProvider, InMessage, *SkeletalMeshInfo);
+
+								CacheForWorld.Time = InTime;
+							}
+
+							Component->SetDrawDebugSkeleton(bDrawSkeleton);
+							Component->SetDebugDrawColor(Color);
 						}
 					}
 				});
-
-				InLineBatcher->DrawLines(Lines);
-
-				Lines.Reset();
 			});
 		}
 	}
 }
 
+FSkeletalMeshPoseTrack::FWorldComponentCache& FSkeletalMeshPoseTrack::GetWorldCache(UWorld* InWorld)
+{
+	FWorldComponentCache& Cache = WorldCache.FindOrAdd(TWeakObjectPtr<UWorld>(InWorld));
+	Cache.World = InWorld;
+	return Cache;
+}
+
+UInsightsSkeletalMeshComponent* FSkeletalMeshPoseTrack::FWorldComponentCache::GetComponent()
+{
+	if(Component == nullptr)
+	{
+		Component = NewObject<UInsightsSkeletalMeshComponent>();
+		Component->PrimaryComponentTick.bStartWithTickEnabled = false;
+		Component->PrimaryComponentTick.bCanEverTick = false;
+		Component->RegisterComponentWithWorld(World);
+
+		Time = 0.0;
+	}
+	
+	return Component;
+}
+
+void FSkeletalMeshPoseTrack::AddReferencedObjects(FReferenceCollector& Collector)
+{
+	for(auto& WorldCacheEntry : WorldCache)
+	{
+		Collector.AddReferencedObject(WorldCacheEntry.Value.Component);
+	}
+}
+
 #endif
 
 #undef LOCTEXT_NAMESPACE
\ No newline at end of file
diff --git a/Engine/Plugins/Animation/GameplayInsights/Source/GameplayInsights/Private/SkeletalMeshPoseTrack.h b/Engine/Plugins/Animation/GameplayInsights/Source/GameplayInsights/Private/SkeletalMeshPoseTrack.h
--- a/Engine/Plugins/Animation/GameplayInsights/Source/GameplayInsights/Private/SkeletalMeshPoseTrack.h
+++ b/Engine/Plugins/Animation/GameplayInsights/Source/GameplayInsights/Private/SkeletalMeshPoseTrack.h
@@ -5,18 +5,30 @@
 #include "GameplayTrack.h"
 #include "Insights/ViewModels/TimingEventsTrack.h"
 
-class ULineBatchComponent;
+#if WITH_ENGINE
+#include "UObject/GCObject.h"
+#endif
+
+#if WITH_ENGINE
+class UWorld;
+class UInsightsSkeletalMeshComponent;
+#endif
+
 class FAnimationSharedData;
 class FTimingEventSearchParameters;
 struct FSkeletalMeshPoseMessage;
 
 class FSkeletalMeshPoseTrack : public TGameplayTrackMixin<FTimingEventsTrack>
+#if WITH_ENGINE
+	, public FGCObject
+#endif
 {
 public:
 	static const FName TypeName;
 	static const FName SubTypeName;
 
 	FSkeletalMeshPoseTrack(const FAnimationSharedData& InSharedData, uint64 InObjectID, const TCHAR* InName);
+	~FSkeletalMeshPoseTrack();
 
 	virtual void BuildDrawState(ITimingEventsTrackDrawStateBuilder& Builder, const ITimingTrackUpdateContext& Context) override;
 	virtual void Draw(const ITimingTrackDrawContext& Context) const override;
@@ -25,14 +37,19 @@ public:
 	virtual void BuildContextMenu(FMenuBuilder& MenuBuilder) override;
 
 	// Access drawing flags
-	bool ShouldDrawMarkerTime() const { return bDrawMarkerTime; }
-	bool ShouldDrawSelectedEvent() const { return bDrawSelectedEvent; }
-	bool ShouldDrawHoveredEvent() const { return bDrawHoveredEvent; }
-	bool ShouldDrawSelection() const { return bDrawSelection; }
+	bool ShouldDrawPose() const { return bDrawPose; }
+	bool ShouldDrawSkeleton() const { return bDrawSkeleton; }
 
 #if WITH_ENGINE
-	// Draw poses in the specified time range
-	void DrawPoses(ULineBatchComponent* InLineBatcher, double SelectionStartTime, double SelectionEndTime);
+	// Handle worlds being torn down
+	void OnWorldCleanup(UWorld* InWorld, bool bSessionEnded, bool bCleanupResources);
+
+	// Draw poses at the specified time
+	void DrawPoses(UWorld* InWorld, double InTime);
+
+	// FGCObject interface
+	virtual void AddReferencedObjects(FReferenceCollector& Collector) override;
+	virtual FString GetReferencerName() const override { return TEXT("InsightsSkeletalMeshPoseTrack"); }
 #endif
 
 private:
@@ -43,9 +60,45 @@ private:
 	/** The shared data */
 	const FAnimationSharedData& SharedData;
 
-	/** Whether to draw poses at specified times */
-	bool bDrawMarkerTime;
-	bool bDrawSelectedEvent;
-	bool bDrawHoveredEvent;
-	bool bDrawSelection;
+	/** The color to use to draw this track */
+	FLinearColor Color;
+
+	/** Whether to draw the pose */
+	bool bDrawPose;
+
+	/** Whether to draw the skeleton */
+	bool bDrawSkeleton;
+
+#if WITH_ENGINE
+	/** Cached data per-world */
+	struct FWorldComponentCache
+	{
+		FWorldComponentCache()
+			: World(nullptr)
+			, Component(nullptr)
+			, Time(0.0)
+		{}
+
+		/** Get a cached component for this world */
+		UInsightsSkeletalMeshComponent* GetComponent();
+
+		/** The world we populate */
+		UWorld* World;
+
+		/** Cached component used to visualize in this world */
+		UInsightsSkeletalMeshComponent* Component;
+
+		/** The time we last cached on this component */
+		double Time;
+	};
+
+	// Get the cached data for a world
+	FWorldComponentCache& GetWorldCache(UWorld* InWorld);
+
+	/** Cached map of per-world data */
+	TMap<TWeakObjectPtr<UWorld>, FWorldComponentCache> WorldCache;
+
+	/** Handle used to deal with world switching */
+	FDelegateHandle OnWorldDestroyedHandle;
+#endif
 };
\ No newline at end of file
diff --git a/Engine/Source/Developer/TraceInsights/Private/Insights/Widgets/STimingView.cpp b/Engine/Source/Developer/TraceInsights/Private/Insights/Widgets/STimingView.cpp
--- a/Engine/Source/Developer/TraceInsights/Private/Insights/Widgets/STimingView.cpp
+++ b/Engine/Source/Developer/TraceInsights/Private/Insights/Widgets/STimingView.cpp
@@ -278,6 +278,8 @@ void STimingView::Reset()
 		OnSelectedEventChangedDelegate.Broadcast(SelectedEvent);
 	}
 
+	bPreventThrottling = false;
+
 	Tooltip.Reset();
 
 	LastSelectionType = ESelectionType::None;
@@ -343,6 +345,8 @@ void STimingView::Tick(const FGeometry& AllottedGeometry, const double InCurrent
 
 	ThisGeometry = AllottedGeometry;
 
+	bPreventThrottling = false;
+
 	const float ViewWidth = AllottedGeometry.GetLocalSize().X;
 	const float ViewHeight = AllottedGeometry.GetLocalSize().Y;
 
@@ -1371,7 +1375,12 @@ FReply STimingView::OnMouseButtonDown(const FGeometry& MyGeometry, const FPointe
 		Reply = FReply::Handled().CaptureMouse(SharedThis(this));
 	}
 
-	if (bStartScrubbing)
+	if(bPreventThrottling)
+	{
+		Reply.PreventThrottling();
+	}
+
+	if(bStartScrubbing)
 	{
 		bIsPanning = false;
 		bIsDragging = false;
@@ -2862,6 +2871,13 @@ void STimingView::OnTrackVisibilityChanged()
 
 ////////////////////////////////////////////////////////////////////////////////////////////////////
 
+void STimingView::PreventThrottling()
+{
+	bPreventThrottling = true;
+}
+
+////////////////////////////////////////////////////////////////////////////////////////////////////
+
 TSharedPtr<FBaseTimingTrack> STimingView::FindTrack(uint64 InTrackId)
 {
 	TSharedPtr<FBaseTimingTrack>* TrackPtrPtr = AllTracks.Find(InTrackId);
diff --git a/Engine/Source/Developer/TraceInsights/Private/Insights/Widgets/STimingView.h b/Engine/Source/Developer/TraceInsights/Private/Insights/Widgets/STimingView.h
--- a/Engine/Source/Developer/TraceInsights/Private/Insights/Widgets/STimingView.h
+++ b/Engine/Source/Developer/TraceInsights/Private/Insights/Widgets/STimingView.h
@@ -228,6 +228,7 @@ public:
 	virtual void AddScrollableTrack(TSharedPtr<FBaseTimingTrack> Track) override;
 	virtual void AddForegroundTrack(TSharedPtr<FBaseTimingTrack> Track) override;
 
+	virtual void PreventThrottling() override;
 	virtual void InvalidateScrollableTracksOrder() override;
 	//TODO: virtual void InvalidateScrollableTracksVisibility() override;
 
@@ -483,6 +484,9 @@ protected:
 
 	double TimeMarker;
 
+	/** Throttle flag, allowing tracks to control whether Slate throttle should take place */
+	bool bPreventThrottling;
+
 	/** True of the user is currently dragging the time marker */
 	bool bIsScrubbing;
 
diff --git a/Engine/Source/Developer/TraceInsights/Public/Insights/ITimingViewSession.h b/Engine/Source/Developer/TraceInsights/Public/Insights/ITimingViewSession.h
--- a/Engine/Source/Developer/TraceInsights/Public/Insights/ITimingViewSession.h
+++ b/Engine/Source/Developer/TraceInsights/Public/Insights/ITimingViewSession.h
@@ -53,6 +53,9 @@ public:
 	/** Adds a new foreground track. */
 	virtual void AddForegroundTrack(TSharedPtr<FBaseTimingTrack> Track) = 0;
 
+	/** Prevents mouse movements from throttling application updates */
+	virtual void PreventThrottling() = 0;
+
 	/** Marks the scrollable tracks as not being in the correct order so they will be re-sorted */
 	virtual void InvalidateScrollableTracksOrder() = 0;
 
diff --git a/Engine/Source/Runtime/Engine/Classes/Components/SkeletalMeshComponent.h b/Engine/Source/Runtime/Engine/Classes/Components/SkeletalMeshComponent.h
--- a/Engine/Source/Runtime/Engine/Classes/Components/SkeletalMeshComponent.h
+++ b/Engine/Source/Runtime/Engine/Classes/Components/SkeletalMeshComponent.h
@@ -64,6 +64,8 @@ enum class EAnimCurveType : uint8
 	MaxAnimCurveType UMETA(Hidden)
 };
 
+ENUM_RANGE_BY_COUNT(EAnimCurveType, EAnimCurveType::MaxAnimCurveType);
+
 UENUM()
 enum class EClothMassMode : uint8
 {
diff --git a/Engine/Source/Runtime/Engine/Classes/Components/SkinnedMeshComponent.h b/Engine/Source/Runtime/Engine/Classes/Components/SkinnedMeshComponent.h
--- a/Engine/Source/Runtime/Engine/Classes/Components/SkinnedMeshComponent.h
+++ b/Engine/Source/Runtime/Engine/Classes/Components/SkinnedMeshComponent.h
@@ -334,6 +334,11 @@ public:
 	FColor WireframeColor_DEPRECATED;
 #endif
 
+#if !(UE_BUILD_SHIPPING || UE_BUILD_TEST)
+	/** Debug draw color */
+	TOptional<FLinearColor> DebugDrawColor;
+#endif
+
 protected:
 	/** Information for current ref pose override, if present */
 	FSkelMeshRefPoseOverride* RefPoseOverride;
@@ -350,6 +355,20 @@ public:
 	 */
 	void GetCPUSkinnedVertices(TArray<struct FFinalSkinVertex>& OutVertices, int32 InLODIndex);
 
+#if !(UE_BUILD_SHIPPING || UE_BUILD_TEST)
+	/** Get whether to draw this mesh's debug skeleton */
+	bool ShouldDrawDebugSkeleton() const { return bDrawDebugSkeleton; }
+
+	/** Set whether to draw this mesh's debug skeleton */
+	void SetDrawDebugSkeleton(bool bInDraw) { bDrawDebugSkeleton = bInDraw; }
+
+	/** Get debug draw color */
+	const TOptional<FLinearColor>& GetDebugDrawColor() const { return DebugDrawColor; }
+
+	/** Set debug draw color */
+	void SetDebugDrawColor(const FLinearColor& InColor) { DebugDrawColor = InColor; }
+#endif
+
 	/** Array indicating all active morph targets. This array is updated inside RefreshBoneTransforms based on the Anim Blueprint. */
 	TArray<FActiveMorphTarget> ActiveMorphTargets;
 
@@ -594,6 +613,12 @@ protected:
 	/** External flag indicating that we may not be evaluated every frame */
 	uint8 bExternalEvaluationRateLimited:1;
 
+#if !(UE_BUILD_SHIPPING || UE_BUILD_TEST)
+private:
+	/** Whether to draw this mesh's debug skeleton (regardless of showflags) */
+	uint8 bDrawDebugSkeleton:1;
+#endif
+
 public:
 	/** Set whether we have our tick rate externally controlled non-URO-based interpolation */
 	void EnableExternalTickRateControl(bool bInEnable) { bExternalTickRateControlled = bInEnable; }
diff --git a/Engine/Source/Runtime/Engine/Private/Animation/AnimTrace.cpp b/Engine/Source/Runtime/Engine/Private/Animation/AnimTrace.cpp
--- a/Engine/Source/Runtime/Engine/Private/Animation/AnimTrace.cpp
+++ b/Engine/Source/Runtime/Engine/Private/Animation/AnimTrace.cpp
@@ -32,11 +32,12 @@ UE_TRACE_EVENT_BEGIN(Animation, SkeletalMesh, Important)
 	UE_TRACE_EVENT_FIELD(uint32, BoneCount)
 UE_TRACE_EVENT_END()
 
-UE_TRACE_EVENT_BEGIN(Animation, SkeletalMeshPose)
+UE_TRACE_EVENT_BEGIN(Animation, SkeletalMeshComponent)
 	UE_TRACE_EVENT_FIELD(uint64, Cycle)
 	UE_TRACE_EVENT_FIELD(uint64, ComponentId)
 	UE_TRACE_EVENT_FIELD(uint64, MeshId)
 	UE_TRACE_EVENT_FIELD(uint32, BoneCount)
+	UE_TRACE_EVENT_FIELD(uint32, CurveCount)
 	UE_TRACE_EVENT_FIELD(uint16, LodIndex)
 	UE_TRACE_EVENT_FIELD(uint16, FrameCounter)
 UE_TRACE_EVENT_END()
@@ -81,21 +82,29 @@ UE_TRACE_EVENT_BEGIN(Animation, PoseLink)
 	UE_TRACE_EVENT_FIELD(int32, TargetNameLength)
 UE_TRACE_EVENT_END()
 
+UE_TRACE_EVENT_BEGIN(Animation, Name, Important)
+	UE_TRACE_EVENT_FIELD(uint32, Id)
+UE_TRACE_EVENT_END()
+
 // Object annotations used for tracing
 FUObjectAnnotationSparseBool GSkeletalMeshTraceAnnotations;
 
+// Map used for unique name output
+TMap<FName, uint32> GAnimTraceNames;
+
 void FAnimTrace::Init()
 {
 	if (FParse::Param(FCommandLine::Get(), TEXT("objecttrace")))
 	{
 		UE_TRACE_EVENT_IS_ENABLED(Animation, TickRecord);
 		UE_TRACE_EVENT_IS_ENABLED(Animation, SkeletalMesh);
-		UE_TRACE_EVENT_IS_ENABLED(Animation, SkeletalMeshPose);
+		UE_TRACE_EVENT_IS_ENABLED(Animation, SkeletalMeshComponent);
 // 		UE_TRACE_EVENT_IS_ENABLED(Animation, AnimNodeValueBool);
 // 		UE_TRACE_EVENT_IS_ENABLED(Animation, AnimNodeValueInt);
 // 		UE_TRACE_EVENT_IS_ENABLED(Animation, AnimNodeValueFloat);
 // 		UE_TRACE_EVENT_IS_ENABLED(Animation, AnimNodeValueString);
 // 		UE_TRACE_EVENT_IS_ENABLED(Animation, PoseLink);
+		UE_TRACE_EVENT_IS_ENABLED(Animation, Name);
 		Trace::ToggleEvent(TEXT("Animation"), true);
 	}
 }
@@ -178,34 +187,96 @@ void FAnimTrace::OutputSkeletalMesh(const USkeletalMesh* InMesh)
 	GSkeletalMeshTraceAnnotations.Set(InMesh);
 }
 
-void FAnimTrace::OutputSkeletalMeshPose(const USkeletalMeshComponent* InComponent)
+uint32 FAnimTrace::OutputName(const FName& InName)
+{
+	static uint32 CurrentId = 1;
+	check(IsInGameThread());
+
+	uint32* ExistingIdPtr = GAnimTraceNames.Find(InName);
+	if(ExistingIdPtr == nullptr)
+	{
+		int32 NameStringLength = InName.GetStringLength() + 1;
+
+		auto StringCopyFunc = [NameStringLength, &InName](uint8* Out)
+		{
+			InName.ToString(reinterpret_cast<TCHAR*>(Out), NameStringLength);
+		};
+
+		uint32 NewId = CurrentId++;
+
+		UE_TRACE_LOG(Animation, Name, NameStringLength * sizeof(TCHAR))
+			<< Name.Id(NewId)
+			<< Name.Attachment(StringCopyFunc);
+
+		GAnimTraceNames.Add(InName, NewId);
+		return NewId;
+	}
+	else
+	{
+		return *ExistingIdPtr;
+	}
+}
+
+void FAnimTrace::OutputSkeletalMeshComponent(const USkeletalMeshComponent* InComponent)
 {
-	bool bEventEnabled = UE_TRACE_EVENT_IS_ENABLED(Animation, SkeletalMeshPose);
+	bool bEventEnabled = UE_TRACE_EVENT_IS_ENABLED(Animation, SkeletalMeshComponent);
 	if (!bEventEnabled || InComponent == nullptr)
 	{
 		return;
 	}
 
 	int32 BoneCount = InComponent->GetComponentSpaceTransforms().Num();
-	if(BoneCount > 0)
+	int32 CurveCount = 0;
+	UAnimInstance* AnimInstance = InComponent->GetAnimInstance();
+
+	if(AnimInstance)
+	{
+		for(EAnimCurveType CurveType : TEnumRange<EAnimCurveType>())
+		{
+			CurveCount += AnimInstance->GetAnimationCurveList(CurveType).Num();
+		}
+	}
+	
+	if(BoneCount > 0 || CurveCount > 0)
 	{
 		TRACE_OBJECT(InComponent);
 		TRACE_SKELETAL_MESH(InComponent->SkeletalMesh);
 
-		auto CopyTransforms = [&InComponent, &BoneCount](uint8* Out)
+		auto CopyTransformsAndCurves = [&InComponent, &BoneCount, &CurveCount, &AnimInstance](uint8* Out)
 		{
 			FPlatformMemory::Memcpy(Out, &InComponent->GetComponentToWorld(), sizeof(FTransform));
-			FPlatformMemory::Memcpy(Out + sizeof(FTransform), InComponent->GetComponentSpaceTransforms().GetData(), BoneCount * sizeof(FTransform));
+			Out += sizeof(FTransform);
+
+			if(BoneCount > 0)
+			{
+				const int32 BufferSize = BoneCount * sizeof(FTransform);
+				FPlatformMemory::Memcpy(Out, InComponent->GetComponentSpaceTransforms().GetData(), BufferSize);
+				Out += BufferSize;
+			}
+			if(CurveCount > 0 && AnimInstance)
+			{
+				for(EAnimCurveType CurveType : TEnumRange<EAnimCurveType>())
+				{
+					for(TPair<FName, float> CurvePair : AnimInstance->GetAnimationCurveList(CurveType))
+					{
+						*reinterpret_cast<uint32*>(Out) = OutputName(CurvePair.Key);
+						Out += sizeof(uint32);
+						*reinterpret_cast<float*>(Out) = CurvePair.Value;
+						Out += sizeof(float);
+					}
+				}
+			}
 		};
 
-		UE_TRACE_LOG(Animation, SkeletalMeshPose, (BoneCount + 1) * sizeof(FTransform))
-			<< SkeletalMeshPose.Cycle(FPlatformTime::Cycles64())
-			<< SkeletalMeshPose.ComponentId(FObjectTrace::GetObjectId(InComponent))
-			<< SkeletalMeshPose.MeshId(FObjectTrace::GetObjectId(InComponent->SkeletalMesh))
-			<< SkeletalMeshPose.BoneCount((uint32)BoneCount + 1)
-			<< SkeletalMeshPose.LodIndex((uint16)InComponent->PredictedLODLevel)
-			<< SkeletalMeshPose.FrameCounter((uint16)(GFrameCounter % 0xffff))
-			<< SkeletalMeshPose.Attachment(CopyTransforms);
+		UE_TRACE_LOG(Animation, SkeletalMeshComponent, ((BoneCount + 1) * sizeof(FTransform)) + (CurveCount * (sizeof(float) + sizeof(uint32))))
+			<< SkeletalMeshComponent.Cycle(FPlatformTime::Cycles64())
+			<< SkeletalMeshComponent.ComponentId(FObjectTrace::GetObjectId(InComponent))
+			<< SkeletalMeshComponent.MeshId(FObjectTrace::GetObjectId(InComponent->SkeletalMesh))
+			<< SkeletalMeshComponent.BoneCount((uint32)BoneCount + 1)
+			<< SkeletalMeshComponent.CurveCount((uint32)CurveCount)
+			<< SkeletalMeshComponent.LodIndex((uint16)InComponent->PredictedLODLevel)
+			<< SkeletalMeshComponent.FrameCounter((uint16)(GFrameCounter % 0xffff))
+			<< SkeletalMeshComponent.Attachment(CopyTransformsAndCurves);
 	}
 }
 
diff --git a/Engine/Source/Runtime/Engine/Private/Components/SkeletalMeshComponent.cpp b/Engine/Source/Runtime/Engine/Private/Components/SkeletalMeshComponent.cpp
--- a/Engine/Source/Runtime/Engine/Private/Components/SkeletalMeshComponent.cpp
+++ b/Engine/Source/Runtime/Engine/Private/Components/SkeletalMeshComponent.cpp
@@ -820,6 +820,10 @@ void USkeletalMeshComponent::InitializeComponent()
 void USkeletalMeshComponent::BeginPlay()
 {
 	Super::BeginPlay();
+
+	// Trace the 'first frame' markers
+	TRACE_SKELETAL_MESH_COMPONENT(this);
+
 	if (AnimScriptInstance)
 	{
 		AnimScriptInstance->NativeBeginPlay();
@@ -3733,7 +3737,7 @@ void USkeletalMeshComponent::FinalizeBoneTransform()
 
 	OnBoneTransformsFinalized.Broadcast();
 
-	TRACE_SKELETALMESH_POSE(this);
+	TRACE_SKELETAL_MESH_COMPONENT(this);
 }
 
 void USkeletalMeshComponent::GetCurrentRefToLocalMatrices(TArray<FMatrix>& OutRefToLocals, int32 InLodIdx)
diff --git a/Engine/Source/Runtime/Engine/Private/Components/SkinnedMeshComponent.cpp b/Engine/Source/Runtime/Engine/Private/Components/SkinnedMeshComponent.cpp
--- a/Engine/Source/Runtime/Engine/Private/Components/SkinnedMeshComponent.cpp
+++ b/Engine/Source/Runtime/Engine/Private/Components/SkinnedMeshComponent.cpp
@@ -365,6 +365,10 @@ USkinnedMeshComponent::USkinnedMeshComponent(const FObjectInitializer& ObjectIni
 	bExternalEvaluationRateLimited = false;
 	bExternalTickRateControlled = false;
 
+#if !(UE_BUILD_SHIPPING || UE_BUILD_TEST)
+	bDrawDebugSkeleton = false;
+#endif
+
 	CurrentSkinWeightProfileName = NAME_None;
 }
 
diff --git a/Engine/Source/Runtime/Engine/Private/SkeletalMesh.cpp b/Engine/Source/Runtime/Engine/Private/SkeletalMesh.cpp
--- a/Engine/Source/Runtime/Engine/Private/SkeletalMesh.cpp
+++ b/Engine/Source/Runtime/Engine/Private/SkeletalMesh.cpp
@@ -4093,9 +4093,15 @@ FSkeletalMeshSceneProxy::FSkeletalMeshSceneProxy(const USkinnedMeshComponent* Co
 		,	bForceWireframe(Component->bForceWireframe)
 		,	bCanHighlightSelectedSections(Component->bCanHighlightSelectedSections)
 		,	bRenderStatic(Component->bRenderStatic)
+#if !(UE_BUILD_SHIPPING || UE_BUILD_TEST)
+		,	bDrawDebugSkeleton(Component->ShouldDrawDebugSkeleton())
+#endif
 		,	FeatureLevel(GetScene().GetFeatureLevel())
 		,	bMaterialsNeedMorphUsage_GameThread(false)
 		,	MaterialRelevance(Component->GetMaterialRelevance(GetScene().GetFeatureLevel()))
+#if !(UE_BUILD_SHIPPING || UE_BUILD_TEST)
+		,	DebugDrawColor(Component->GetDebugDrawColor())
+#endif
 #if WITH_EDITORONLY_DATA
 		,	StreamingDistanceMultiplier(FMath::Max(0.0f, Component->StreamingDistanceMultiplier))
 #endif
@@ -4593,7 +4599,7 @@ void FSkeletalMeshSceneProxy::GetMeshElementsConditionallySelectable(const TArra
 				RenderBounds(Collector.GetPDI(ViewIndex), ViewFamily.EngineShowFlags, GetBounds(), IsSelected());
 			}
 
-			if (ViewFamily.EngineShowFlags.Bones)
+			if (ViewFamily.EngineShowFlags.Bones || bDrawDebugSkeleton)
 			{
 				DebugDrawSkeleton(ViewIndex, Collector, ViewFamily.EngineShowFlags);
 			}
@@ -5025,7 +5031,7 @@ void FSkeletalMeshSceneProxy::DebugDrawSkeleton(int32 ViewIndex, FMeshElementCol
 		const int32 ParentIndex = SkeletalMeshForDebug->RefSkeleton.GetParentIndex(Index);
 		FVector Start, End;
 		
-		FLinearColor LineColor = MakeRandomColorForSkeleton(GetPrimitiveComponentId().PrimIDValue);
+		FLinearColor LineColor = DebugDrawColor.Get(MakeRandomColorForSkeleton(GetPrimitiveComponentId().PrimIDValue));
 		const FTransform Transform = ComponentSpaceTransforms[Index] * LocalToWorldTransform;
 
 		if (ParentIndex >= 0)
@@ -5039,7 +5045,7 @@ void FSkeletalMeshSceneProxy::DebugDrawSkeleton(int32 ViewIndex, FMeshElementCol
 			End = Transform.GetLocation();
 		}
 
-		if(EngineShowFlags.Bones)
+		if(EngineShowFlags.Bones || bDrawDebugSkeleton)
 		{
 			if(CVarDebugDrawSimpleBones.GetValueOnRenderThread() != 0)
 			{
diff --git a/Engine/Source/Runtime/Engine/Public/Animation/AnimTrace.h b/Engine/Source/Runtime/Engine/Public/Animation/AnimTrace.h
--- a/Engine/Source/Runtime/Engine/Public/Animation/AnimTrace.h
+++ b/Engine/Source/Runtime/Engine/Public/Animation/AnimTrace.h
@@ -30,8 +30,8 @@ struct FAnimTrace
 	/** Helper function to output a skeletal mesh */
 	ENGINE_API static void OutputSkeletalMesh(const USkeletalMesh* InMesh);
 
-	/** Helper function to output a skeletal mesh pose */
-	ENGINE_API static void OutputSkeletalMeshPose(const USkeletalMeshComponent* InComponent);
+	/** Helper function to output a skeletal mesh pose, curves etc. */
+	ENGINE_API static void OutputSkeletalMeshComponent(const USkeletalMeshComponent* InComponent);
 
 	/** Helper function to output a pose link */
 	ENGINE_API static void OutputPoseLink(const FAnimationUpdateContext& InContext);
@@ -42,6 +42,12 @@ struct FAnimTrace
 	ENGINE_API static void OutputAnimNodeValue(const FAnimationBaseContext& InContext, const TCHAR* InKey, float Value);
 	ENGINE_API static void OutputAnimNodeValue(const FAnimationBaseContext& InContext, const TCHAR* InKey, const FName& Value);
 	ENGINE_API static void OutputAnimNodeValue(const FAnimationBaseContext& InContext, const TCHAR* InKey, const TCHAR* Value);
+
+	/** 
+	 * Helper function to output a name to the trace stream, referenced by ID. 
+	 * @return the ID used to reference the name
+	 */
+	ENGINE_API static uint32 OutputName(const FName& InName);
 };
 
 #define TRACE_ANIM_TICK_RECORDS(Proxy, Component) \
@@ -50,8 +56,8 @@ struct FAnimTrace
 #define TRACE_SKELETAL_MESH(Mesh) \
 	FAnimTrace::OutputSkeletalMesh(Mesh);
 
-#define TRACE_SKELETALMESH_POSE(Component) \
-	FAnimTrace::OutputSkeletalMeshPose(Component);
+#define TRACE_SKELETAL_MESH_COMPONENT(Component) \
+	FAnimTrace::OutputSkeletalMeshComponent(Component);
 
 #define TRACE_POSE_LINK(Context) \
 	FAnimTrace::OutputPoseLink(Context);
diff --git a/Engine/Source/Runtime/Engine/Public/SkeletalMeshTypes.h b/Engine/Source/Runtime/Engine/Public/SkeletalMeshTypes.h
--- a/Engine/Source/Runtime/Engine/Public/SkeletalMeshTypes.h
+++ b/Engine/Source/Runtime/Engine/Public/SkeletalMeshTypes.h
@@ -323,6 +323,9 @@ protected:
 	uint8 bIsCPUSkinned : 1;
 	uint8 bCanHighlightSelectedSections : 1;
 	uint8 bRenderStatic:1;
+#if !(UE_BUILD_SHIPPING || UE_BUILD_TEST)
+	uint8 bDrawDebugSkeleton:1;
+#endif
 
 	TEnumAsByte<ERHIFeatureLevel::Type> FeatureLevel;
 
@@ -379,6 +382,11 @@ protected:
 	/** The primitive's pre-skinned local space bounds. */
 	FBoxSphereBounds PreSkinnedLocalBounds;
 
+#if !(UE_BUILD_SHIPPING || UE_BUILD_TEST)
+	/** The color we draw this component in if drawing debug bones */
+	TOptional<FLinearColor> DebugDrawColor;
+#endif
+
 #if WITH_EDITORONLY_DATA
 	/** The component streaming distance multiplier */
 	float StreamingDistanceMultiplier;
