Removed unecessary include that is causing problems with Linux builds

diff --git a/Engine/Plugins/Animation/GameplayInsights/Source/GameplayInsights/Private/GameplayInsightsModule.cpp b/Engine/Plugins/Animation/GameplayInsights/Source/GameplayInsights/Private/GameplayInsightsModule.cpp
--- a/Engine/Plugins/Animation/GameplayInsights/Source/GameplayInsights/Private/GameplayInsightsModule.cpp
+++ b/Engine/Plugins/Animation/GameplayInsights/Source/GameplayInsights/Private/GameplayInsightsModule.cpp
@@ -7,7 +7,6 @@
 #include "Insights/ITimingViewExtender.h"
 #include "Containers/Ticker.h"
 #include "Modules/ModuleManager.h"
-#include "Framework/Multibox/MultiboxBuilder.h"
 
 #if WITH_ENGINE
 #include "Engine/Engine.h"
