Fix Linux build

#rnx

diff --git a/Engine/Plugins/Animation/GameplayInsights/Source/GameplayInsights/Private/IAnimationProvider.h b/Engine/Plugins/Animation/GameplayInsights/Source/GameplayInsights/Private/IAnimationProvider.h
--- a/Engine/Plugins/Animation/GameplayInsights/Source/GameplayInsights/Private/IAnimationProvider.h
+++ b/Engine/Plugins/Animation/GameplayInsights/Source/GameplayInsights/Private/IAnimationProvider.h
@@ -105,6 +105,15 @@ struct FAnimNodeValueMessage
 	{
 	}
 
+	struct FVectorEntry
+	{
+		FVectorEntry(const FVector& InVector)
+			: Value(InVector)
+		{}
+
+		FVector Value;
+	};
+
 	const TCHAR* Key = nullptr;
 	uint64 AnimInstanceId = 0;
 	union
@@ -121,14 +130,7 @@ struct FAnimNodeValueMessage
 		{
 			float Value;
 		} Float;
-		struct FVectorEntry
-		{
-			FVectorEntry(const FVector& InVector)
-				: Value(InVector)
-			{}
-
-			FVector Value;
-		} Vector;
+		FVectorEntry Vector;
 		struct
 		{
 			const TCHAR* Value;
