Added skeletal mesh curves track & track filtering widget

Displays traced curves as a single track.
Moved common functionality into new FGameplayGraphTrack.
Added support for curve enumeration in animation provider. Curve IDs are stored seperately for each skeletal mesh for easier enumeration.
Added 'stacked' view to gameplay graph tracks so all curves can be seen at once more easily.
Added filtering widget to more easily disable/enable individual gameplay tracks.

diff --git a/Engine/Plugins/Animation/GameplayInsights/Source/GameplayInsights/Private/AnimationProvider.cpp b/Engine/Plugins/Animation/GameplayInsights/Source/GameplayInsights/Private/AnimationProvider.cpp
--- a/Engine/Plugins/Animation/GameplayInsights/Source/GameplayInsights/Private/AnimationProvider.cpp
+++ b/Engine/Plugins/Animation/GameplayInsights/Source/GameplayInsights/Private/AnimationProvider.cpp
@@ -4,6 +4,8 @@
 #include "GameplayProvider.h"
 #include "Insights/ViewModels/TimingEventsTrack.h"
 
+#define LOCTEXT_NAMESPACE "AnimationProvider"
+
 FName FAnimationProvider::ProviderName("AnimationProvider");
 
 FAnimationProvider::FAnimationProvider(Trace::IAnalysisSession& InSession, FGameplayProvider& InGameplayProvider)
@@ -16,16 +18,17 @@ FAnimationProvider::FAnimationProvider(Trace::IAnalysisSession& InSession, FGame
 	GameplayProvider.OnObjectEndPlay().AddRaw(this, &FAnimationProvider::HandleObjectEndPlay);
 }
 
-bool FAnimationProvider::ReadSkeletalMeshPoseTimeline(uint64 InObjectId, TFunctionRef<void(const SkeletalMeshPoseTimeline&)> Callback) const
+bool FAnimationProvider::ReadSkeletalMeshPoseTimeline(uint64 InObjectId, TFunctionRef<void(const SkeletalMeshPoseTimeline&, bool)> Callback) const
 {
 	Session.ReadAccessCheck();
 
 	const uint32* IndexPtr = ObjectIdToSkeletalMeshPoseTimelines.Find(InObjectId);
 	if(IndexPtr != nullptr)
 	{
-		if (*IndexPtr < uint32(SkeletalMeshPoseTimelines.Num()))
+		if (*IndexPtr < uint32(SkeletalMeshPoseTimelineStorage.Num()))
 		{
-			Callback(*SkeletalMeshPoseTimelines[*IndexPtr]);
+			const TSharedRef<FSkeletalMeshTimelineStorage>& SkeletalMeshTimelineStorage = SkeletalMeshPoseTimelineStorage[*IndexPtr];
+			Callback(*SkeletalMeshTimelineStorage->Timeline, SkeletalMeshTimelineStorage->AllCurveIds.Num() > 0);
 			return true;
 		}
 	}
@@ -55,6 +58,37 @@ void FAnimationProvider::GetSkeletalMeshComponentSpacePose(const FSkeletalMeshPo
 	}
 }
 
+void FAnimationProvider::EnumerateSkeletalMeshCurveIds(uint64 InObjectId, TFunctionRef<void(uint32)> Callback) const
+{
+	Session.ReadAccessCheck();
+
+	const uint32* IndexPtr = ObjectIdToSkeletalMeshPoseTimelines.Find(InObjectId);
+	if(IndexPtr != nullptr)
+	{
+		if (*IndexPtr < uint32(SkeletalMeshPoseTimelineStorage.Num()))
+		{
+			const TSharedRef<FSkeletalMeshTimelineStorage>& TimelineStorage = SkeletalMeshPoseTimelineStorage[*IndexPtr];
+			for(uint32 Id : TimelineStorage->AllCurveIds)
+			{
+				Callback(Id);
+			}
+		}
+	}
+}
+
+void FAnimationProvider::EnumerateSkeletalMeshCurves(const FSkeletalMeshPoseMessage& InMessage, TFunctionRef<void(const FSkeletalMeshNamedCurve&)> Callback) const
+{
+	Session.ReadAccessCheck();
+
+	uint64 StartCurveIndex = InMessage.CurveStartIndex;
+	uint64 EndCurveIndex = InMessage.CurveStartIndex + InMessage.NumCurves;
+
+	for(uint64 CurveIndex = StartCurveIndex; CurveIndex < EndCurveIndex; ++CurveIndex)
+	{
+		Callback(SkeletalMeshCurves[CurveIndex]);
+	}
+}
+
 void FAnimationProvider::EnumerateTickRecordTimelines(uint64 InObjectId, TFunctionRef<void(uint64, int32, const TickRecordTimeline&)> Callback) const
 {
 	Session.ReadAccessCheck();
@@ -62,9 +96,9 @@ void FAnimationProvider::EnumerateTickRecordTimelines(uint64 InObjectId, TFuncti
 	const uint32* IndexPtr = ObjectIdToTickRecordTimelineStorage.Find(InObjectId);
 	if(IndexPtr != nullptr)
 	{
-		if (*IndexPtr < uint32(PerObjectTimelineStorage.Num()))
+		if (*IndexPtr < uint32(TickRecordTimelineStorage.Num()))
 		{
-			const TSharedRef<FPerObjectTimelineStorage>& TimelineStorage = PerObjectTimelineStorage[*IndexPtr];
+			const TSharedRef<FTickRecordTimelineStorage>& TimelineStorage = TickRecordTimelineStorage[*IndexPtr];
 			for(auto AssetIdPair : TimelineStorage->AssetIdAndPlayerToTickRecordTimeline)
 			{
 				Callback(AssetIdPair.Key.Key, AssetIdPair.Key.Value, TimelineStorage->Timelines[AssetIdPair.Value].Get());
@@ -80,9 +114,9 @@ bool FAnimationProvider::ReadTickRecordTimeline(uint64 InObjectId, uint64 InAsse
 	const uint32* IndexPtr = ObjectIdToTickRecordTimelineStorage.Find(InObjectId);
 	if(IndexPtr != nullptr)
 	{
-		if (*IndexPtr < uint32(PerObjectTimelineStorage.Num()))
+		if (*IndexPtr < uint32(TickRecordTimelineStorage.Num()))
 		{
-			const TSharedRef<FPerObjectTimelineStorage>& TimelineStorage = PerObjectTimelineStorage[*IndexPtr];
+			const TSharedRef<FTickRecordTimelineStorage>& TimelineStorage = TickRecordTimelineStorage[*IndexPtr];
 			const uint32* TimelineIndexPtr = TimelineStorage->AssetIdAndPlayerToTickRecordTimeline.Find(TTuple<uint64, int32>(InAssetId, InNodeId));
 			if(TimelineIndexPtr != nullptr)
 			{
@@ -215,7 +249,14 @@ const FSkeletalMeshInfo* FAnimationProvider::FindSkeletalMeshInfo(uint64 InObjec
 
 const TCHAR* FAnimationProvider::GetName(uint32 InId) const
 {
-	return NameMap.FindRef(InId);
+	const TCHAR* const* FoundName = NameMap.Find(InId);
+	if(FoundName != nullptr)
+	{
+		return *FoundName;
+	}
+
+	static FText UnknownText(LOCTEXT("UnknownCurve", "Unknown"));
+	return *UnknownText.ToString();
 }
 
 void FAnimationProvider::AppendTickRecord(uint64 InAnimInstanceId, double InTime, uint64 InAssetId, int32 InNodeId, float InBlendWeight, float InPlaybackTime, float InRootMotionWeight, float InPlayRate, float InBlendSpacePositionX, float InBlendSpacePositionY, uint16 InFrameCounter, bool bInLooping, bool bInIsBlendSpace)
@@ -223,16 +264,16 @@ void FAnimationProvider::AppendTickRecord(uint64 InAnimInstanceId, double InTime
 	Session.WriteAccessCheck();
 
 	TSharedPtr<Trace::TPointTimeline<FTickRecordMessage>> Timeline;
-	TSharedPtr<FPerObjectTimelineStorage> TimelineStorage;
+	TSharedPtr<FTickRecordTimelineStorage> TimelineStorage;
 	uint32* TimelineStorageIndexPtr = ObjectIdToTickRecordTimelineStorage.Find(InAnimInstanceId);
 	if(TimelineStorageIndexPtr != nullptr)
 	{
-		TimelineStorage = PerObjectTimelineStorage[*TimelineStorageIndexPtr];
+		TimelineStorage = TickRecordTimelineStorage[*TimelineStorageIndexPtr];
 	}
 	else
 	{
-		ObjectIdToTickRecordTimelineStorage.Add(InAnimInstanceId, PerObjectTimelineStorage.Num());
-		TimelineStorage = PerObjectTimelineStorage.Add_GetRef(MakeShared<FPerObjectTimelineStorage>());
+		ObjectIdToTickRecordTimelineStorage.Add(InAnimInstanceId, TickRecordTimelineStorage.Num());
+		TimelineStorage = TickRecordTimelineStorage.Add_GetRef(MakeShared<FTickRecordTimelineStorage>());
 	}
 
 	check(TimelineStorage.IsValid());
@@ -314,25 +355,26 @@ void FAnimationProvider::AppendSkeletalMeshComponent(uint64 InObjectId, uint64 I
 {
 	Session.WriteAccessCheck();
 
-	TSharedPtr<Trace::TIntervalTimeline<FSkeletalMeshPoseMessage>> Timeline;
+	TSharedPtr<FSkeletalMeshTimelineStorage> TimelineStorage;
 	uint32* IndexPtr = ObjectIdToSkeletalMeshPoseTimelines.Find(InObjectId);
 	if(IndexPtr != nullptr)
 	{
-		Timeline = SkeletalMeshPoseTimelines[*IndexPtr];
+		TimelineStorage = SkeletalMeshPoseTimelineStorage[*IndexPtr];
 	}
 	else
 	{
-		Timeline = MakeShared<Trace::TIntervalTimeline<FSkeletalMeshPoseMessage>>(Session.GetLinearAllocator());
-		ObjectIdToSkeletalMeshPoseTimelines.Add(InObjectId, SkeletalMeshPoseTimelines.Num());
-		SkeletalMeshPoseTimelines.Add(Timeline.ToSharedRef());
+		TimelineStorage = MakeShared<FSkeletalMeshTimelineStorage>();
+		TimelineStorage->Timeline = MakeShared<Trace::TIntervalTimeline<FSkeletalMeshPoseMessage>>(Session.GetLinearAllocator());
+		ObjectIdToSkeletalMeshPoseTimelines.Add(InObjectId, SkeletalMeshPoseTimelineStorage.Num());
+		SkeletalMeshPoseTimelineStorage.Add(TimelineStorage.ToSharedRef());
 	}
 
 	// terminate existing scopes
-	uint64 NumEvents = Timeline->GetEventCount();
+	uint64 NumEvents = TimelineStorage->Timeline->GetEventCount();
 	if(NumEvents > 0)
 	{
 		// Add end event at current time
-		Timeline->EndEvent(NumEvents - 1, InTime);
+		TimelineStorage->Timeline->EndEvent(NumEvents - 1, InTime);
 	}
 
 	FSkeletalMeshPoseMessage Message;
@@ -346,7 +388,7 @@ void FAnimationProvider::AppendSkeletalMeshComponent(uint64 InObjectId, uint64 I
 	Message.LodIndex = InLodIndex;
 	Message.FrameCounter = InFrameCounter;
 
-	Timeline->AppendBeginEvent(InTime, Message);
+	TimelineStorage->Timeline->AppendBeginEvent(InTime, Message);
 
 	for(int32 TransformIndex = 1; TransformIndex < InPose.Num(); ++TransformIndex)
 	{
@@ -356,6 +398,7 @@ void FAnimationProvider::AppendSkeletalMeshComponent(uint64 InObjectId, uint64 I
 	for(const FSkeletalMeshNamedCurve& Curve : InCurves)
 	{
 		SkeletalMeshCurves.PushBack() = Curve;
+		TimelineStorage->AllCurveIds.Add(Curve.Id);
 	}
 
 	Session.UpdateDurationSeconds(InTime);
@@ -374,7 +417,7 @@ void FAnimationProvider::HandleObjectEndPlay(uint64 InObjectId, double InTime, c
 	uint32* SkelMeshPoseIndexPtr = ObjectIdToSkeletalMeshPoseTimelines.Find(InObjectId);
 	if(SkelMeshPoseIndexPtr != nullptr)
 	{
-		TSharedPtr<Trace::TIntervalTimeline<FSkeletalMeshPoseMessage>> Timeline = SkeletalMeshPoseTimelines[*SkelMeshPoseIndexPtr];
+		TSharedPtr<Trace::TIntervalTimeline<FSkeletalMeshPoseMessage>> Timeline = SkeletalMeshPoseTimelineStorage[*SkelMeshPoseIndexPtr]->Timeline;
 		uint64 NumEvents = Timeline->GetEventCount();
 		if(NumEvents > 0)
 		{
@@ -662,3 +705,4 @@ void FAnimationProvider::AppendStateMachineState(uint64 InAnimInstanceId, double
 	Session.UpdateDurationSeconds(InTime);
 }
 
+#undef LOCTEXT_NAMESPACE
\ No newline at end of file
diff --git a/Engine/Plugins/Animation/GameplayInsights/Source/GameplayInsights/Private/AnimationProvider.h b/Engine/Plugins/Animation/GameplayInsights/Source/GameplayInsights/Private/AnimationProvider.h
--- a/Engine/Plugins/Animation/GameplayInsights/Source/GameplayInsights/Private/AnimationProvider.h
+++ b/Engine/Plugins/Animation/GameplayInsights/Source/GameplayInsights/Private/AnimationProvider.h
@@ -22,8 +22,10 @@ public:
 	FAnimationProvider(Trace::IAnalysisSession& InSession, FGameplayProvider& InGameplayProvider);
 
 	/** IAnimationProvider interface */
-	virtual bool ReadSkeletalMeshPoseTimeline(uint64 InObjectId, TFunctionRef<void(const SkeletalMeshPoseTimeline&)> Callback) const override;
+	virtual bool ReadSkeletalMeshPoseTimeline(uint64 InObjectId, TFunctionRef<void(const SkeletalMeshPoseTimeline&, bool)> Callback) const override;
 	virtual void GetSkeletalMeshComponentSpacePose(const FSkeletalMeshPoseMessage& InMessage, const FSkeletalMeshInfo& InMeshInfo, FTransform& OutComponentToWorld, TArray<FTransform>& OutTransforms) const override;
+	virtual void EnumerateSkeletalMeshCurveIds(uint64 InObjectId, TFunctionRef<void(uint32)> Callback) const override;
+	virtual void EnumerateSkeletalMeshCurves(const FSkeletalMeshPoseMessage& InMessage, TFunctionRef<void(const FSkeletalMeshNamedCurve&)> Callback) const override;
 	virtual void EnumerateTickRecordTimelines(uint64 InObjectId, TFunctionRef<void(uint64, int32, const TickRecordTimeline&)> Callback) const override;
 	virtual bool ReadTickRecordTimeline(uint64 InObjectId, uint64 InAssetId, int32 InNodeId, TFunctionRef<void(const TickRecordTimeline&)> Callback) const override;
 	virtual bool ReadAnimGraphTimeline(uint64 InObjectId, TFunctionRef<void(const AnimGraphTimeline&)> Callback) const override;
@@ -108,15 +110,21 @@ private:
 	/** Map of name IDs to name string */
 	TMap<uint32, const TCHAR*> NameMap;
 	
-	struct FPerObjectTimelineStorage
+	struct FTickRecordTimelineStorage
 	{
 		TMap<TTuple<uint64, int32>, uint32> AssetIdAndPlayerToTickRecordTimeline;
 		TArray<TSharedRef<Trace::TPointTimeline<FTickRecordMessage>>> Timelines;
 	};
 
+	struct FSkeletalMeshTimelineStorage
+	{
+		TSharedPtr<Trace::TIntervalTimeline<FSkeletalMeshPoseMessage>> Timeline;
+		TSet<uint32> AllCurveIds;
+	};
+
 	/** Message storage */
-	TArray<TSharedRef<FPerObjectTimelineStorage>> PerObjectTimelineStorage;
-	TArray<TSharedRef<Trace::TIntervalTimeline<FSkeletalMeshPoseMessage>>> SkeletalMeshPoseTimelines;
+	TArray<TSharedRef<FTickRecordTimelineStorage>> TickRecordTimelineStorage;
+	TArray<TSharedRef<FSkeletalMeshTimelineStorage>> SkeletalMeshPoseTimelineStorage;
 	TArray<TSharedRef<Trace::TIntervalTimeline<FSkeletalMeshFrameMessage>>> SkeletalMeshFrameTimelines;
 	TArray<TSharedRef<Trace::TIntervalTimeline<FAnimGraphMessage>>> AnimGraphTimelines;
 	TArray<TSharedRef<Trace::TPointTimeline<FAnimNodeMessage>>> AnimNodeTimelines;
diff --git a/Engine/Plugins/Animation/GameplayInsights/Source/GameplayInsights/Private/AnimationSharedData.cpp b/Engine/Plugins/Animation/GameplayInsights/Source/GameplayInsights/Private/AnimationSharedData.cpp
--- a/Engine/Plugins/Animation/GameplayInsights/Source/GameplayInsights/Private/AnimationSharedData.cpp
+++ b/Engine/Plugins/Animation/GameplayInsights/Source/GameplayInsights/Private/AnimationSharedData.cpp
@@ -5,6 +5,7 @@
 #include "GameplaySharedData.h"
 #include "ObjectEventsTrack.h"
 #include "SkeletalMeshPoseTrack.h"
+#include "SkeletalMeshCurvesTrack.h"
 #include "AnimationTickRecordsTrack.h"
 #include "Insights/ITimingViewSession.h"
 #include "Insights/ViewModels/TimingEvent.h"
@@ -26,7 +27,10 @@ FAnimationSharedData::FAnimationSharedData(FGameplaySharedData& InGameplayShared
 	, AnalysisSession(nullptr)
 	, MarkerTime(0.0)
 	, bTimeMarkerValid(false)
-	, bAnimationTracksEnabled(true)
+	, bSkeletalMeshPoseTracksEnabled(true)
+	, bSkeletalMeshCurveTracksEnabled(true)
+	, bTickRecordTracksEnabled(true)
+	, bAnimNodeTracksEnabled(true)
 {
 }
 
@@ -60,82 +64,103 @@ void FAnimationSharedData::Tick(Insights::ITimingViewSession& InTimingViewSessio
 		// Add tracks for each tracked object's animation data
 		GameplayProvider->EnumerateObjects([this, &InTimingViewSession, &InAnalysisSession, &AnimationProvider, &GameplayProvider](const FObjectInfo& InObjectInfo)
 		{
-			TSharedRef<FObjectEventsTrack> ObjectEventsTrack = GameplaySharedData.GetObjectEventsTrackForId(InTimingViewSession, InAnalysisSession, InObjectInfo);
-
-			AnimationProvider->ReadSkeletalMeshPoseTimeline(InObjectInfo.Id, [this, &InObjectInfo, &ObjectEventsTrack, &InTimingViewSession](const IAnimationProvider::SkeletalMeshPoseTimeline& InTimeline)
+			GameplayProvider->ReadObjectEventsTimeline(InObjectInfo.Id, [this, &InTimingViewSession, &InAnalysisSession, &InObjectInfo, &AnimationProvider, &GameplayProvider](const IGameplayProvider::ObjectEventsTimeline& InTimeline)
 			{
-				ObjectEventsTrack->SetVisibilityFlag(GameplaySharedData.AreGameplayTracksEnabled());
-
-				auto FindSkeletalMeshPoseTrack = [](const FBaseTimingTrack& InTrack)
-				{
-					return InTrack.Is<FSkeletalMeshPoseTrack>();
-				};
-
-				TSharedPtr<FSkeletalMeshPoseTrack> ExistingSkeletalMeshPoseTrack = StaticCastSharedPtr<FSkeletalMeshPoseTrack>(ObjectEventsTrack->GetGameplayTrack().FindChildTrack(InObjectInfo.Id, FindSkeletalMeshPoseTrack));
-				if(!ExistingSkeletalMeshPoseTrack.IsValid())
+				if(InTimeline.GetEventCount() > 0)
 				{
-					TSharedPtr<FSkeletalMeshPoseTrack> SkeletalMeshPoseTrack = MakeShared<FSkeletalMeshPoseTrack>(*this, InObjectInfo.Id, InObjectInfo.Name);
-					SkeletalMeshPoseTrack->SetVisibilityFlag(bAnimationTracksEnabled);
-					SkeletalMeshPoseTracks.Add(SkeletalMeshPoseTrack.ToSharedRef());
-
-					InTimingViewSession.AddScrollableTrack(SkeletalMeshPoseTrack);
-					GameplaySharedData.InvalidateObjectTracksOrder();
-
-					ObjectEventsTrack->GetGameplayTrack().AddChildTrack(SkeletalMeshPoseTrack->GetGameplayTrack());
-				}
-			});
-
-			AnimationProvider->EnumerateTickRecordTimelines(InObjectInfo.Id, [this, &InObjectInfo, &ObjectEventsTrack, &InTimingViewSession, &GameplayProvider](uint64 InAssetId, int32 InNodeId, const IAnimationProvider::TickRecordTimeline& InTimeline)
-			{
-				ObjectEventsTrack->SetVisibilityFlag(GameplaySharedData.AreGameplayTracksEnabled());
+					TSharedRef<FObjectEventsTrack> ObjectEventsTrack = GameplaySharedData.GetObjectEventsTrackForId(InTimingViewSession, InAnalysisSession, InObjectInfo);
 
-				auto FindTickRecordTrackWithAssetId = [&InAssetId](const FBaseTimingTrack& InTrack)
-				{
-					if (InTrack.Is<FAnimationTickRecordsTrack>())
+					AnimationProvider->ReadSkeletalMeshPoseTimeline(InObjectInfo.Id, [this, &InObjectInfo, &ObjectEventsTrack, &InTimingViewSession](const IAnimationProvider::SkeletalMeshPoseTimeline& InTimeline, bool bInHasCurves)
 					{
-						const FAnimationTickRecordsTrack& AnimationTickRecordsTrack = InTrack.As<FAnimationTickRecordsTrack>();
-						return AnimationTickRecordsTrack.GetAssetId() == InAssetId;
-					}
-
-					return false;
-				};
-
-				TSharedPtr<FAnimationTickRecordsTrack> ExistingAnimationTickRecordsTrack = StaticCastSharedPtr<FAnimationTickRecordsTrack>(ObjectEventsTrack->GetGameplayTrack().FindChildTrack(InObjectInfo.Id, FindTickRecordTrackWithAssetId));
-				if(!ExistingAnimationTickRecordsTrack.IsValid())
-				{
-					const FObjectInfo* AssetObjectInfo = GameplayProvider->FindObjectInfo(InAssetId);
-					FString AssetName = AssetObjectInfo ? AssetObjectInfo->Name : LOCTEXT("UnknownAsset", "Unknown").ToString();
-					TSharedPtr<FAnimationTickRecordsTrack> AnimationTickRecordsTrack = MakeShared<FAnimationTickRecordsTrack>(*this, InObjectInfo.Id, InAssetId, InNodeId, *AssetName);
-					AnimationTickRecordsTrack->SetVisibilityFlag(bAnimationTracksEnabled);
-					AnimationTickRecordsTracks.Add(AnimationTickRecordsTrack.ToSharedRef());
-
-					InTimingViewSession.AddScrollableTrack(AnimationTickRecordsTrack);
-					GameplaySharedData.InvalidateObjectTracksOrder();
-
-					ObjectEventsTrack->GetGameplayTrack().AddChildTrack(AnimationTickRecordsTrack->GetGameplayTrack());
-				}
-			});
-
-			AnimationProvider->ReadAnimGraphTimeline(InObjectInfo.Id, [this, &InObjectInfo, &ObjectEventsTrack, &InTimingViewSession](const IAnimationProvider::AnimGraphTimeline& InTimeline)
-			{
-				ObjectEventsTrack->SetVisibilityFlag(GameplaySharedData.AreGameplayTracksEnabled());
-
-				auto FindAnimNodesTrack = [](const FBaseTimingTrack& InTrack)
-				{
-					return InTrack.Is<FAnimNodesTrack>();
-				};
-
-				TSharedPtr<FAnimNodesTrack> ExistingAnimNodesTrack = StaticCastSharedPtr<FAnimNodesTrack>(ObjectEventsTrack->GetGameplayTrack().FindChildTrack(InObjectInfo.Id, FindAnimNodesTrack));
-				if(!ExistingAnimNodesTrack.IsValid())
-				{
-					TSharedPtr<FAnimNodesTrack> AnimNodesTrack = MakeShared<FAnimNodesTrack>(*this, InObjectInfo.Id, InObjectInfo.Name);
-					AnimNodesTrack->SetVisibilityFlag(bAnimationTracksEnabled);
-					AnimNodesTracks.Add(AnimNodesTrack.ToSharedRef());
-
-					InTimingViewSession.AddScrollableTrack(AnimNodesTrack);
-					GameplaySharedData.InvalidateObjectTracksOrder();
-
-					ObjectEventsTrack->GetGameplayTrack().AddChildTrack(AnimNodesTrack->GetGameplayTrack());
+						auto FindSkeletalMeshPoseTrack = [](const FBaseTimingTrack& InTrack)
+						{
+							return InTrack.Is<FSkeletalMeshPoseTrack>();
+						};
+
+						TSharedPtr<FSkeletalMeshPoseTrack> ExistingSkeletalMeshPoseTrack = StaticCastSharedPtr<FSkeletalMeshPoseTrack>(ObjectEventsTrack->GetGameplayTrack().FindChildTrack(InObjectInfo.Id, FindSkeletalMeshPoseTrack));
+						if(!ExistingSkeletalMeshPoseTrack.IsValid())
+						{
+							TSharedPtr<FSkeletalMeshPoseTrack> SkeletalMeshPoseTrack = MakeShared<FSkeletalMeshPoseTrack>(*this, InObjectInfo.Id, InObjectInfo.Name);
+							SkeletalMeshPoseTrack->SetVisibilityFlag(bSkeletalMeshPoseTracksEnabled);
+							SkeletalMeshPoseTracks.Add(SkeletalMeshPoseTrack.ToSharedRef());
+
+							InTimingViewSession.AddScrollableTrack(SkeletalMeshPoseTrack);
+							GameplaySharedData.InvalidateObjectTracksOrder();
+
+							ObjectEventsTrack->GetGameplayTrack().AddChildTrack(SkeletalMeshPoseTrack->GetGameplayTrack());
+						}
+
+						if(bInHasCurves)
+						{
+							auto FindSkeletalMeshCurvesTrack = [](const FBaseTimingTrack& InTrack)
+							{
+								return InTrack.Is<FSkeletalMeshCurvesTrack>();
+							};
+
+							TSharedPtr<FSkeletalMeshCurvesTrack> ExistingSkeletalMeshCurvesTrack = StaticCastSharedPtr<FSkeletalMeshCurvesTrack>(ObjectEventsTrack->GetGameplayTrack().FindChildTrack(InObjectInfo.Id, FindSkeletalMeshCurvesTrack));
+							if(!ExistingSkeletalMeshCurvesTrack.IsValid())
+							{
+								TSharedPtr<FSkeletalMeshCurvesTrack> SkeletalMeshCurvesTrack = MakeShared<FSkeletalMeshCurvesTrack>(*this, InObjectInfo.Id, InObjectInfo.Name);
+								SkeletalMeshCurvesTrack->SetVisibilityFlag(bSkeletalMeshCurveTracksEnabled);
+								SkeletalMeshCurvesTracks.Add(SkeletalMeshCurvesTrack.ToSharedRef());
+
+								InTimingViewSession.AddScrollableTrack(SkeletalMeshCurvesTrack);
+								GameplaySharedData.InvalidateObjectTracksOrder();
+
+								ObjectEventsTrack->GetGameplayTrack().AddChildTrack(SkeletalMeshCurvesTrack->GetGameplayTrack());
+							}
+						}
+					});
+
+					AnimationProvider->EnumerateTickRecordTimelines(InObjectInfo.Id, [this, &InObjectInfo, &ObjectEventsTrack, &InTimingViewSession, &GameplayProvider](uint64 InAssetId, int32 InNodeId, const IAnimationProvider::TickRecordTimeline& InTimeline)
+					{
+						auto FindTickRecordTrackWithAssetId = [&InAssetId](const FBaseTimingTrack& InTrack)
+						{
+							if (InTrack.Is<FAnimationTickRecordsTrack>())
+							{
+								const FAnimationTickRecordsTrack& AnimationTickRecordsTrack = InTrack.As<FAnimationTickRecordsTrack>();
+								return AnimationTickRecordsTrack.GetAssetId() == InAssetId;
+							}
+
+							return false;
+						};
+
+						TSharedPtr<FAnimationTickRecordsTrack> ExistingAnimationTickRecordsTrack = StaticCastSharedPtr<FAnimationTickRecordsTrack>(ObjectEventsTrack->GetGameplayTrack().FindChildTrack(InObjectInfo.Id, FindTickRecordTrackWithAssetId));
+						if(!ExistingAnimationTickRecordsTrack.IsValid())
+						{
+							const FObjectInfo* AssetObjectInfo = GameplayProvider->FindObjectInfo(InAssetId);
+							FString AssetName = AssetObjectInfo ? AssetObjectInfo->Name : LOCTEXT("UnknownAsset", "Unknown").ToString();
+							TSharedPtr<FAnimationTickRecordsTrack> AnimationTickRecordsTrack = MakeShared<FAnimationTickRecordsTrack>(*this, InObjectInfo.Id, InAssetId, InNodeId, *AssetName);
+							AnimationTickRecordsTrack->SetVisibilityFlag(bTickRecordTracksEnabled);
+							AnimationTickRecordsTracks.Add(AnimationTickRecordsTrack.ToSharedRef());
+
+							InTimingViewSession.AddScrollableTrack(AnimationTickRecordsTrack);
+							GameplaySharedData.InvalidateObjectTracksOrder();
+
+							ObjectEventsTrack->GetGameplayTrack().AddChildTrack(AnimationTickRecordsTrack->GetGameplayTrack());
+						}
+					});
+
+					AnimationProvider->ReadAnimGraphTimeline(InObjectInfo.Id, [this, &InObjectInfo, &ObjectEventsTrack, &InTimingViewSession](const IAnimationProvider::AnimGraphTimeline& InTimeline)
+					{
+						auto FindAnimNodesTrack = [](const FBaseTimingTrack& InTrack)
+						{
+							return InTrack.Is<FAnimNodesTrack>();
+						};
+
+						TSharedPtr<FAnimNodesTrack> ExistingAnimNodesTrack = StaticCastSharedPtr<FAnimNodesTrack>(ObjectEventsTrack->GetGameplayTrack().FindChildTrack(InObjectInfo.Id, FindAnimNodesTrack));
+						if(!ExistingAnimNodesTrack.IsValid())
+						{
+							TSharedPtr<FAnimNodesTrack> AnimNodesTrack = MakeShared<FAnimNodesTrack>(*this, InObjectInfo.Id, InObjectInfo.Name);
+							AnimNodesTrack->SetVisibilityFlag(bAnimNodeTracksEnabled);
+							AnimNodesTracks.Add(AnimNodesTrack.ToSharedRef());
+
+							InTimingViewSession.AddScrollableTrack(AnimNodesTrack);
+							GameplaySharedData.InvalidateObjectTracksOrder();
+
+							ObjectEventsTrack->GetGameplayTrack().AddChildTrack(AnimNodesTrack->GetGameplayTrack());
+						}
+					});
 				}
 			});
 		});
@@ -159,7 +184,7 @@ void FAnimationSharedData::ExtendFilterMenu(FMenuBuilder& InMenuBuilder)
 {
 	InMenuBuilder.AddMenuEntry(
 		LOCTEXT("ToggleAnimationTracks", "Animation Tracks"),
-		LOCTEXT("ToggleAnimationTracks_Tooltip", "Show/hide the animation tracks"),
+		LOCTEXT("ToggleAnimationTracks_Tooltip", "Show/hide all animation tracks"),
 		FSlateIcon(),
 		FUIAction(
 			FExecuteAction::CreateRaw(this, &FAnimationSharedData::ToggleAnimationTracks),
@@ -168,26 +193,129 @@ void FAnimationSharedData::ExtendFilterMenu(FMenuBuilder& InMenuBuilder)
 		NAME_None,
 		EUserInterfaceActionType::ToggleButton
 	);
+
+	InMenuBuilder.AddMenuEntry(
+		LOCTEXT("ToggleSkelMeshPoseTracks", "Pose Tracks"),
+		LOCTEXT("ToggleSkelMeshPoseTracks_Tooltip", "Show/hide the skeletal mesh pose tracks"),
+		FSlateIcon(),
+		FUIAction(
+			FExecuteAction::CreateRaw(this, &FAnimationSharedData::ToggleSkeletalMeshPoseTracks),
+			FCanExecuteAction(),
+			FIsActionChecked::CreateLambda([this](){ return bSkeletalMeshPoseTracksEnabled; })),
+		NAME_None,
+		EUserInterfaceActionType::ToggleButton
+	);
+
+	InMenuBuilder.AddMenuEntry(
+		LOCTEXT("ToggleSkelMeshCurveTracks", "Curve Tracks"),
+		LOCTEXT("ToggleSkelMeshCurveTracks_Tooltip", "Show/hide the skeletal mesh curve tracks"),
+		FSlateIcon(),
+		FUIAction(
+			FExecuteAction::CreateRaw(this, &FAnimationSharedData::ToggleSkeletalMeshCurveTracks),
+			FCanExecuteAction(),
+			FIsActionChecked::CreateLambda([this](){ return bSkeletalMeshCurveTracksEnabled; })),
+		NAME_None,
+		EUserInterfaceActionType::ToggleButton
+	);
+
+	InMenuBuilder.AddMenuEntry(
+		LOCTEXT("ToggleAnimTickRecordTracks", "Blend Weights Tracks"),
+		LOCTEXT("ToggleAnimTickRecordTracks_Tooltip", "Show/hide the blend weights (tick records) tracks"),
+		FSlateIcon(),
+		FUIAction(
+			FExecuteAction::CreateRaw(this, &FAnimationSharedData::ToggleTickRecordTracks),
+			FCanExecuteAction(),
+			FIsActionChecked::CreateLambda([this](){ return bTickRecordTracksEnabled; })),
+		NAME_None,
+		EUserInterfaceActionType::ToggleButton
+	);
+
+	InMenuBuilder.AddMenuEntry(
+		LOCTEXT("ToggleAnimNodeTracks", "Graph Tracks"),
+		LOCTEXT("ToggleAnimNodeTracks_Tooltip", "Show/hide the animation graph tracks"),
+		FSlateIcon(),
+		FUIAction(
+			FExecuteAction::CreateRaw(this, &FAnimationSharedData::ToggleAnimNodeTracks),
+			FCanExecuteAction(),
+			FIsActionChecked::CreateLambda([this](){ return bAnimNodeTracksEnabled; })),
+		NAME_None,
+		EUserInterfaceActionType::ToggleButton
+	);
 }
 
 void FAnimationSharedData::ToggleAnimationTracks()
 {
-	bAnimationTracksEnabled = !bAnimationTracksEnabled;
+	bool bAnimationTracksEnabled = !AreAnimationTracksEnabled();
+
+	bSkeletalMeshPoseTracksEnabled = bAnimationTracksEnabled;
+	bSkeletalMeshCurveTracksEnabled = bAnimationTracksEnabled;
+	bTickRecordTracksEnabled = bAnimationTracksEnabled;
+	bAnimNodeTracksEnabled = bAnimationTracksEnabled;
 
 	for(TSharedRef<FSkeletalMeshPoseTrack> PoseTrack : SkeletalMeshPoseTracks)
 	{
 		PoseTrack->SetVisibilityFlag(bAnimationTracksEnabled);
 	}
 
+	for(TSharedRef<FSkeletalMeshCurvesTrack> CurvesTrack : SkeletalMeshCurvesTracks)
+	{
+		CurvesTrack->SetVisibilityFlag(bAnimationTracksEnabled);
+	}
+
 	for(TSharedRef<FAnimationTickRecordsTrack> TickRecordTrack : AnimationTickRecordsTracks)
 	{
 		TickRecordTrack->SetVisibilityFlag(bAnimationTracksEnabled);
 	}
+
+	for(TSharedRef<FAnimNodesTrack> AnimNodesTrack : AnimNodesTracks)
+	{
+		AnimNodesTrack->SetVisibilityFlag(bAnimationTracksEnabled);
+	}
 }
 
 bool FAnimationSharedData::AreAnimationTracksEnabled() const
 {
-	return bAnimationTracksEnabled;
+	return bSkeletalMeshPoseTracksEnabled && bSkeletalMeshCurveTracksEnabled && bTickRecordTracksEnabled && bAnimNodeTracksEnabled;
+}
+
+void FAnimationSharedData::ToggleSkeletalMeshPoseTracks()
+{
+	bSkeletalMeshPoseTracksEnabled = !bSkeletalMeshPoseTracksEnabled;
+
+	for(TSharedRef<FSkeletalMeshPoseTrack> PoseTrack : SkeletalMeshPoseTracks)
+	{
+		PoseTrack->SetVisibilityFlag(bSkeletalMeshPoseTracksEnabled);
+	}
+}
+
+void FAnimationSharedData::ToggleSkeletalMeshCurveTracks()
+{
+	bSkeletalMeshCurveTracksEnabled = !bSkeletalMeshCurveTracksEnabled;
+
+	for(TSharedRef<FSkeletalMeshCurvesTrack> CurvesTrack : SkeletalMeshCurvesTracks)
+	{
+		CurvesTrack->SetVisibilityFlag(bSkeletalMeshCurveTracksEnabled);
+	}
+}
+
+void FAnimationSharedData::ToggleTickRecordTracks()
+{
+	bTickRecordTracksEnabled = !bTickRecordTracksEnabled;
+
+	for(TSharedRef<FAnimationTickRecordsTrack> TickRecordsTrack : AnimationTickRecordsTracks)
+	{
+		TickRecordsTrack->SetVisibilityFlag(bTickRecordTracksEnabled);
+	}
+}
+
+void FAnimationSharedData::ToggleAnimNodeTracks()
+{
+	bAnimNodeTracksEnabled = !bAnimNodeTracksEnabled;
+
+	for(TSharedRef<FAnimNodesTrack> AnimNodesTrack : AnimNodesTracks)
+	{
+		AnimNodesTrack->SetVisibilityFlag(bAnimNodeTracksEnabled);
+	}
 }
 
 void FAnimationSharedData::OnTimeMarkerChanged(Insights::ETimeChangedFlags InFlags, double InTimeMarker)
diff --git a/Engine/Plugins/Animation/GameplayInsights/Source/GameplayInsights/Private/AnimationSharedData.h b/Engine/Plugins/Animation/GameplayInsights/Source/GameplayInsights/Private/AnimationSharedData.h
--- a/Engine/Plugins/Animation/GameplayInsights/Source/GameplayInsights/Private/AnimationSharedData.h
+++ b/Engine/Plugins/Animation/GameplayInsights/Source/GameplayInsights/Private/AnimationSharedData.h
@@ -9,6 +9,7 @@ namespace Trace { class IAnalysisSession; }
 namespace Insights { class ITimingViewSession; }
 namespace Insights { enum class ETimeChangedFlags : int32; }
 class FSkeletalMeshPoseTrack;
+class FSkeletalMeshCurvesTrack;
 class FAnimationTickRecordsTrack;
 class FAnimNodesTrack;
 class FMenuBuilder;
@@ -66,8 +67,10 @@ private:
 	// UI handlers
 	void ToggleAnimationTracks();
 	void OnTimeMarkerChanged(Insights::ETimeChangedFlags InFlags, double InTimeMarker);
-	void ToggleAnimationParentTracks();
-	bool AreAnimationParentTracksEnabled() const;
+	void ToggleSkeletalMeshPoseTracks();
+	void ToggleSkeletalMeshCurveTracks();
+	void ToggleTickRecordTracks();
+	void ToggleAnimNodeTracks();
 
 private:
 	// The gameplay shared data we are linked to
@@ -78,6 +81,7 @@ private:
 
 	// All the tracks we manage
 	TArray<TSharedRef<FSkeletalMeshPoseTrack>> SkeletalMeshPoseTracks;
+	TArray<TSharedRef<FSkeletalMeshCurvesTrack>> SkeletalMeshCurvesTracks;
 	TArray<TSharedRef<FAnimationTickRecordsTrack>> AnimationTickRecordsTracks;
 	TArray<TSharedRef<FAnimNodesTrack>> AnimNodesTracks;
 
@@ -90,6 +94,9 @@ private:
 	/** Validity flags for pose times/ranges */
 	bool bTimeMarkerValid;
 
-	// Whether all of our animation tracks are enabled
-	bool bAnimationTracksEnabled;
+	// Flags controlling whether check type of our animation tracks are enabled
+	bool bSkeletalMeshPoseTracksEnabled;
+	bool bSkeletalMeshCurveTracksEnabled;
+	bool bTickRecordTracksEnabled;
+	bool bAnimNodeTracksEnabled;
 };
\ No newline at end of file
diff --git a/Engine/Plugins/Animation/GameplayInsights/Source/GameplayInsights/Private/AnimationTickRecordsTrack.cpp b/Engine/Plugins/Animation/GameplayInsights/Source/GameplayInsights/Private/AnimationTickRecordsTrack.cpp
--- a/Engine/Plugins/Animation/GameplayInsights/Source/GameplayInsights/Private/AnimationTickRecordsTrack.cpp
+++ b/Engine/Plugins/Animation/GameplayInsights/Source/GameplayInsights/Private/AnimationTickRecordsTrack.cpp
@@ -44,15 +44,6 @@ FString FTickRecordSeries::FormatValue(double Value) const
 	return FGraphSeries::FormatValue(Value);
 }
 
-void FTickRecordSeries::UpdateAutoZoom(const FTimingTrackViewport& InViewport, const FAnimationTickRecordsTrack& InTrack)
-{
-	const float TimelineDY = FMath::Max(1.0f, InViewport.GetLayout().TimelineDY);
-	const float TopY = FMath::Max(1.0f, TimelineDY);
-	const float BottomY = FMath::Max(TopY, InTrack.GetHeight() - TimelineDY);
-
-	FGraphSeries::UpdateAutoZoom(TopY, BottomY, CurrentMin, CurrentMax);
-}
-
 static FLinearColor MakeSeriesColor(uint32 InSeed, bool bInLine = false)
 {
 	FRandomStream Stream(InSeed);
@@ -67,25 +58,15 @@ static FLinearColor MakeSeriesColor(FTickRecordSeries::ESeriesType InSeed, bool
 }
 
 FAnimationTickRecordsTrack::FAnimationTickRecordsTrack(const FAnimationSharedData& InSharedData, uint64 InObjectID, uint64 InAssetId, int32 InNodeId, const TCHAR* InName)
-	: TGameplayTrackMixin<FGraphTrack>(InObjectID, MakeTrackName(InSharedData.GetGameplaySharedData(), InAssetId, InName))
+	: FGameplayGraphTrack(InObjectID, MakeTrackName(InSharedData.GetGameplaySharedData(), InAssetId, InName))
 	, SharedData(InSharedData)
 	, AssetId(InAssetId)
 	, NodeId(InNodeId)
-	, RequestedTrackSizeScale(1.0f)
-	, BorderY(0.0f)
 {
 	uint32 Hash = GetTypeHash(GetName());
 	MainSeriesLineColor = MakeSeriesColor(Hash, true);
 	MainSeriesFillColor = MakeSeriesColor(Hash);
 
-	AddAllSeries();
-
-	bDrawPoints = false;
-	bDrawBoxes = false;
-	bDrawBaseline = false;
-	bUseEventDuration = false;
-	VisibleOptions &= ~(EGraphOptions::ShowBars | EGraphOptions::UseEventDuration);
-
 #if WITH_EDITOR
 	const FGameplayProvider* GameplayProvider = SharedData.GetAnalysisSession().ReadProvider<FGameplayProvider>(FGameplayProvider::ProviderName);
 
@@ -194,6 +175,8 @@ void FAnimationTickRecordsTrack::AddAllSeries()
 	const FGameplayProvider* GameplayProvider = SharedData.GetAnalysisSession().ReadProvider<FGameplayProvider>(FGameplayProvider::ProviderName);
 	if(GameplayProvider)
 	{
+		Trace::FAnalysisSessionReadScope SessionReadScope(SharedData.GetAnalysisSession());
+
 		const FClassInfo& ClassInfo = GameplayProvider->GetClassInfoFromObject(AssetId);
 		if(FCString::Stristr(ClassInfo.Name, TEXT("BlendSpace")) != nullptr)
 		{
@@ -206,165 +189,101 @@ void FAnimationTickRecordsTrack::AddAllSeries()
 }
 
 template<typename ProjectionType>
-static void UpdateSeries(const FAnimationSharedData& InSharedData, FAnimationTickRecordsTrack& InTrack, FTickRecordSeries& InSeries, const FTimingTrackViewport& InViewport, ProjectionType Projection)
-{		
-	const FGameplayProvider* GameplayProvider = InSharedData.GetAnalysisSession().ReadProvider<FGameplayProvider>(FGameplayProvider::ProviderName);
-	const FAnimationProvider* AnimationProvider = InSharedData.GetAnalysisSession().ReadProvider<FAnimationProvider>(FAnimationProvider::ProviderName);
+bool FAnimationTickRecordsTrack::UpdateSeriesBoundsHelper(FTickRecordSeries& InSeries, const FTimingTrackViewport& InViewport, ProjectionType Projection)
+{
+	bool bFoundEvents = false;
+
+	const FAnimationProvider* AnimationProvider = SharedData.GetAnalysisSession().ReadProvider<FAnimationProvider>(FAnimationProvider::ProviderName);
 
-	if(GameplayProvider && AnimationProvider)
+	if(AnimationProvider)
 	{
-		Trace::FAnalysisSessionReadScope SessionReadScope(InSharedData.GetAnalysisSession());
+		Trace::FAnalysisSessionReadScope SessionReadScope(SharedData.GetAnalysisSession());
 
-		// Calc visible range before we build events (as builder uses scale factors calculated in auto zoom!)
 		InSeries.CurrentMin = 0.0;
 		InSeries.CurrentMax = 0.0;
 
-		AnimationProvider->ReadTickRecordTimeline(InTrack.GetGameplayTrack().GetObjectId(), InTrack.GetAssetId(), InTrack.GetNodeId(), [&InViewport, &InSeries, &Projection](const FAnimationProvider::TickRecordTimeline& InTimeline)
+		AnimationProvider->ReadTickRecordTimeline(GetGameplayTrack().GetObjectId(), GetAssetId(), GetNodeId(), [&bFoundEvents, &InViewport, &InSeries, &Projection](const FAnimationProvider::TickRecordTimeline& InTimeline)
 		{
-			InTimeline.EnumerateEvents(InViewport.GetStartTime(), InViewport.GetEndTime(), [&InSeries, &Projection](double InStartTime, double InEndTime, uint32 InDepth, const FTickRecordMessage& InMessage)
+			InTimeline.EnumerateEvents(InViewport.GetStartTime(), InViewport.GetEndTime(), [&bFoundEvents, &InSeries, &Projection](double InStartTime, double InEndTime, uint32 InDepth, const FTickRecordMessage& InMessage)
 			{
-				const double Value = Invoke(Projection, InMessage);
+				const float Value = Invoke(Projection, InMessage);
 				InSeries.CurrentMin = FMath::Min(InSeries.CurrentMin, Value);
 				InSeries.CurrentMax = FMath::Max(InSeries.CurrentMax, Value);
+				bFoundEvents = true;
 			});
 		});
+	}
 
-		InSeries.UpdateAutoZoom(InViewport, InTrack);
+	return bFoundEvents;
+}
+
+template<typename ProjectionType>
+void FAnimationTickRecordsTrack::UpdateSeriesHelper(FTickRecordSeries& InSeries, const FTimingTrackViewport& InViewport, ProjectionType Projection)
+{		
+	const FAnimationProvider* AnimationProvider = SharedData.GetAnalysisSession().ReadProvider<FAnimationProvider>(FAnimationProvider::ProviderName);
+
+	if(AnimationProvider)
+	{
+		Trace::FAnalysisSessionReadScope SessionReadScope(SharedData.GetAnalysisSession());
 
-		FGraphTrackBuilder Builder(InTrack, InSeries, InViewport);
+		FGraphTrackBuilder Builder(*this, InSeries, InViewport);
 
-		AnimationProvider->ReadTickRecordTimeline(InTrack.GetGameplayTrack().GetObjectId(), InTrack.GetAssetId(), InTrack.GetNodeId(), [&InTrack, &GameplayProvider, &AnimationProvider, &Builder, &InViewport, &Projection](const FAnimationProvider::TickRecordTimeline& InTimeline)
+		AnimationProvider->ReadTickRecordTimeline(GetGameplayTrack().GetObjectId(), GetAssetId(), GetNodeId(), [this, &AnimationProvider, &Builder, &InViewport, &Projection](const FAnimationProvider::TickRecordTimeline& InTimeline)
 		{
-			InTimeline.EnumerateEvents(InViewport.GetStartTime(), InViewport.GetEndTime(), [&InTrack, &Builder, &GameplayProvider, &Projection](double InStartTime, double InEndTime, uint32 InDepth, const FTickRecordMessage& InMessage)
+			InTimeline.EnumerateEvents(InViewport.GetStartTime(), InViewport.GetEndTime(), [this, &Builder, &Projection](double InStartTime, double InEndTime, uint32 InDepth, const FTickRecordMessage& InMessage)
 			{
 				Builder.AddEvent(InStartTime, InEndTime - InStartTime, Invoke(Projection, InMessage), InMessage.bContinuous);
-
-				InTrack.HeightInLanes = 1;
 			});
 		});
 	}
 }
 
-void FAnimationTickRecordsTrack::UpdateTrackHeight(const ITimingTrackUpdateContext& Context)
+bool FAnimationTickRecordsTrack::UpdateSeriesBounds(FGameplayGraphSeries& InSeries, const FTimingTrackViewport& InViewport)
 {
-	const FTimingTrackViewport& Viewport = Context.GetViewport();
-
-	const float CurrentTrackHeight = GetHeight();
-	const float DesiredTrackHeight = Viewport.GetLayout().ComputeTrackHeight(HeightInLanes) * RequestedTrackSizeScale;
-
-	if (CurrentTrackHeight < DesiredTrackHeight)
+	FTickRecordSeries& TickRecordSeries = *static_cast<FTickRecordSeries*>(&InSeries);
+	switch (TickRecordSeries.Type)
 	{
-		float NewTrackHeight;
-		if (Viewport.IsDirty(ETimingTrackViewportDirtyFlags::VLayoutChanged))
-		{
-			NewTrackHeight = DesiredTrackHeight;
-		}
-		else
-		{
-			NewTrackHeight = FMath::CeilToFloat(CurrentTrackHeight * 0.9f + DesiredTrackHeight * 0.1f);
-		}
-
-		SetHeight(NewTrackHeight);
-
-		for (TSharedPtr<FGraphSeries>& Series : AllSeries)
-		{
-			StaticCastSharedPtr<FTickRecordSeries>(Series)->UpdateAutoZoom(Context.GetViewport(), *this);
-		}
+	case FTickRecordSeries::ESeriesType::BlendWeight:
+		return UpdateSeriesBoundsHelper(TickRecordSeries, InViewport, &FTickRecordMessage::BlendWeight);
+	case FTickRecordSeries::ESeriesType::PlaybackTime:
+		return UpdateSeriesBoundsHelper(TickRecordSeries, InViewport, &FTickRecordMessage::PlaybackTime);
+	case FTickRecordSeries::ESeriesType::RootMotionWeight:
+		return UpdateSeriesBoundsHelper(TickRecordSeries, InViewport, &FTickRecordMessage::RootMotionWeight);
+	case FTickRecordSeries::ESeriesType::PlayRate:
+		return UpdateSeriesBoundsHelper(TickRecordSeries, InViewport, &FTickRecordMessage::PlayRate);
+	case FTickRecordSeries::ESeriesType::BlendSpacePositionX:
+		return UpdateSeriesBoundsHelper(TickRecordSeries, InViewport, &FTickRecordMessage::BlendSpacePositionX);
+	case FTickRecordSeries::ESeriesType::BlendSpacePositionY:
+		return UpdateSeriesBoundsHelper(TickRecordSeries, InViewport, &FTickRecordMessage::BlendSpacePositionY);
 	}
-	else if (CurrentTrackHeight > DesiredTrackHeight)
-	{
-		float NewTrackHeight;
-		if (Viewport.IsDirty(ETimingTrackViewportDirtyFlags::VLayoutChanged))
-		{
-			NewTrackHeight = DesiredTrackHeight;
-		}
-		else
-		{
-			NewTrackHeight = FMath::FloorToFloat(CurrentTrackHeight * 0.9f + DesiredTrackHeight * 0.1f);
-		}
 
-		SetHeight(NewTrackHeight);
-
-		for (TSharedPtr<FGraphSeries>& Series : AllSeries)
-		{
-			StaticCastSharedPtr<FTickRecordSeries>(Series)->UpdateAutoZoom(Context.GetViewport(), *this);
-		}
-	}
+	return false;
 }
 
-void FAnimationTickRecordsTrack::PreUpdate(const ITimingTrackUpdateContext& Context)
+void FAnimationTickRecordsTrack::UpdateSeries(FGameplayGraphSeries& InSeries, const FTimingTrackViewport& InViewport)
 {
-	FGraphTrack::PreUpdate(Context);
-
-	// update border size
-	BorderY = Context.GetViewport().GetLayout().TimelineDY;
-
-	const bool bIsEntireGraphTrackDirty = IsDirty() || Context.GetViewport().IsHorizontalViewportDirty();
-	bool bNeedsUpdate = bIsEntireGraphTrackDirty;
-
-	if (!bNeedsUpdate)
+	FTickRecordSeries& TickRecordSeries = *static_cast<FTickRecordSeries*>(&InSeries);
+	switch (TickRecordSeries.Type)
 	{
-		for (TSharedPtr<FGraphSeries>& Series : AllSeries)
-		{
-			if (Series->IsVisible() && Series->IsDirty())
-			{
-				// At least one series is dirty.
-				bNeedsUpdate = true;
-				break;
-			}
-		}
+	case FTickRecordSeries::ESeriesType::BlendWeight:
+		UpdateSeriesHelper(TickRecordSeries, InViewport, &FTickRecordMessage::BlendWeight);
+		break;
+	case FTickRecordSeries::ESeriesType::PlaybackTime:
+		UpdateSeriesHelper(TickRecordSeries, InViewport, &FTickRecordMessage::PlaybackTime);
+		break;
+	case FTickRecordSeries::ESeriesType::RootMotionWeight:
+		UpdateSeriesHelper(TickRecordSeries, InViewport, &FTickRecordMessage::RootMotionWeight);
+		break;
+	case FTickRecordSeries::ESeriesType::PlayRate:
+		UpdateSeriesHelper(TickRecordSeries, InViewport, &FTickRecordMessage::PlayRate);
+		break;
+	case FTickRecordSeries::ESeriesType::BlendSpacePositionX:
+		UpdateSeriesHelper(TickRecordSeries, InViewport, &FTickRecordMessage::BlendSpacePositionX);
+		break;
+	case FTickRecordSeries::ESeriesType::BlendSpacePositionY:
+		UpdateSeriesHelper(TickRecordSeries, InViewport, &FTickRecordMessage::BlendSpacePositionY);
+		break;
 	}
-
-	if (bNeedsUpdate)
-	{
-		ClearDirtyFlag();
-
-		HeightInLanes = 0;
-
-		const FTimingTrackViewport& Viewport = Context.GetViewport();
-
-		for (TSharedPtr<FGraphSeries>& Series : AllSeries)
-		{
-			if (Series->IsVisible() && (bIsEntireGraphTrackDirty || Series->IsDirty()))
-			{
-				// Clear the flag before updating, because the update itself may further need to set the series as dirty.
-				Series->ClearDirtyFlag();
-
-				TSharedPtr<FTickRecordSeries> TickRecordSeries = StaticCastSharedPtr<FTickRecordSeries>(Series);
-				switch (TickRecordSeries->Type)
-				{
-				case FTickRecordSeries::ESeriesType::BlendWeight:
-					UpdateSeries(SharedData, *this, *TickRecordSeries, Viewport, &FTickRecordMessage::BlendWeight);
-					break;
-				case FTickRecordSeries::ESeriesType::PlaybackTime:
-					UpdateSeries(SharedData, *this, *TickRecordSeries, Viewport, &FTickRecordMessage::PlaybackTime);
-					break;
-				case FTickRecordSeries::ESeriesType::RootMotionWeight:
-					UpdateSeries(SharedData, *this, *TickRecordSeries, Viewport, &FTickRecordMessage::RootMotionWeight);
-					break;
-				case FTickRecordSeries::ESeriesType::PlayRate:
-					UpdateSeries(SharedData, *this, *TickRecordSeries, Viewport, &FTickRecordMessage::PlayRate);
-					break;
-				case FTickRecordSeries::ESeriesType::BlendSpacePositionX:
-					UpdateSeries(SharedData, *this, *TickRecordSeries, Viewport, &FTickRecordMessage::BlendSpacePositionX);
-					break;
-				case FTickRecordSeries::ESeriesType::BlendSpacePositionY:
-					UpdateSeries(SharedData, *this, *TickRecordSeries, Viewport, &FTickRecordMessage::BlendSpacePositionY);
-					break;
-				}
-			}
-		}
-
-		UpdateStats();
-	}
-
-	UpdateTrackHeight(Context);
-}
-
-void FAnimationTickRecordsTrack::Draw(const ITimingTrackDrawContext& Context) const
-{
-	FGraphTrack::Draw(Context);
-	GetGameplayTrack().DrawHeaderForTimingTrack(Context, *this, true);
 }
 
 void FAnimationTickRecordsTrack::InitTooltip(FTooltipDrawState& Tooltip, const ITimingEvent& HoveredTimingEvent) const
@@ -474,53 +393,7 @@ void FAnimationTickRecordsTrack::BuildContextMenu(FMenuBuilder& MenuBuilder)
 	MenuBuilder.EndSection();
 #endif
 
-	FGraphTrack::BuildContextMenu(MenuBuilder);
-
-	MenuBuilder.BeginSection("TrackSize", LOCTEXT("TrackSizeMenuHeader", "Track Size"));
-	{
-		MenuBuilder.AddMenuEntry
-		(
-			LOCTEXT("NormalTrack", "Normal"),
-			LOCTEXT("NormalTrack_Tooltip", "Draw this track at the standard size."),
-			FSlateIcon(),
-			FUIAction(
-				FExecuteAction::CreateLambda([this](){ RequestedTrackSizeScale = 1.0f; }),
-				FCanExecuteAction(),
-				FIsActionChecked::CreateLambda([this](){ return RequestedTrackSizeScale == 1.0f; })
-			),
-			NAME_None,
-			EUserInterfaceActionType::RadioButton
-		);
-
-		MenuBuilder.AddMenuEntry
-		(
-			LOCTEXT("LargeTrack", "Large"),
-			LOCTEXT("LargeTrack_Tooltip", "Make this track larger than normal."),
-			FSlateIcon(),
-			FUIAction(
-				FExecuteAction::CreateLambda([this](){ RequestedTrackSizeScale = 2.0f; }),
-				FCanExecuteAction(),
-				FIsActionChecked::CreateLambda([this](){ return RequestedTrackSizeScale == 2.0f; })
-			),
-			NAME_None,
-			EUserInterfaceActionType::RadioButton
-		);
-
-		MenuBuilder.AddMenuEntry
-		(
-			LOCTEXT("ExtraLargeTrack", "Extra Large"),
-			LOCTEXT("ExtraLargeTrack_Tooltip", "Make this track much larger than normal."),
-			FSlateIcon(),
-			FUIAction(
-				FExecuteAction::CreateLambda([this](){ RequestedTrackSizeScale = 4.0f; }),
-				FCanExecuteAction(),
-				FIsActionChecked::CreateLambda([this](){ return RequestedTrackSizeScale == 4.0f; })
-			),
-			NAME_None,
-			EUserInterfaceActionType::RadioButton
-		);
-	}
-	MenuBuilder.EndSection();
+	FGameplayGraphTrack::BuildContextMenu(MenuBuilder);
 }
 
 FText FAnimationTickRecordsTrack::MakeTrackName(const FGameplaySharedData& InSharedData, uint64 InAssetID, const TCHAR* InName) const
diff --git a/Engine/Plugins/Animation/GameplayInsights/Source/GameplayInsights/Private/AnimationTickRecordsTrack.h b/Engine/Plugins/Animation/GameplayInsights/Source/GameplayInsights/Private/AnimationTickRecordsTrack.h
--- a/Engine/Plugins/Animation/GameplayInsights/Source/GameplayInsights/Private/AnimationTickRecordsTrack.h
+++ b/Engine/Plugins/Animation/GameplayInsights/Source/GameplayInsights/Private/AnimationTickRecordsTrack.h
@@ -2,9 +2,7 @@
 
 #pragma once
 
-#include "GameplayTrack.h"
-#include "Insights/ViewModels/GraphTrack.h"
-#include "Insights/ViewModels/GraphSeries.h"
+#include "GameplayGraphTrack.h"
 
 class FAnimationSharedData;
 struct FTickRecordMessage;
@@ -14,7 +12,7 @@ class UAnimBlueprintGeneratedClass;
 class FAnimationTickRecordsTrack;
 class FTimingTrackViewport;
 
-class FTickRecordSeries : public FGraphSeries
+class FTickRecordSeries : public FGameplayGraphSeries
 {
 public:
 	enum class ESeriesType : uint32
@@ -31,29 +29,23 @@ public:
 
 	virtual FString FormatValue(double Value) const override;
 
-	// Custom overload for auto zoom
-	void UpdateAutoZoom(const FTimingTrackViewport& InViewport, const FAnimationTickRecordsTrack& InTrack);
-
 public:
 	ESeriesType Type;
-	double CurrentMin;
-	double CurrentMax;
 };
 
-class FAnimationTickRecordsTrack : public TGameplayTrackMixin<FGraphTrack>
+class FAnimationTickRecordsTrack : public FGameplayGraphTrack
 {
-	INSIGHTS_DECLARE_RTTI(FAnimationTickRecordsTrack, TGameplayTrackMixin<FGraphTrack>)
+	INSIGHTS_DECLARE_RTTI(FAnimationTickRecordsTrack, FGameplayGraphTrack)
 
 public:
 	FAnimationTickRecordsTrack(const FAnimationSharedData& InSharedData, uint64 InObjectId, uint64 InAssetId, int32 InNodeId, const TCHAR* InName);
 
-	virtual void PreUpdate(const ITimingTrackUpdateContext& Context) override;
-	virtual void Draw(const ITimingTrackDrawContext& Context) const override;
 	virtual void InitTooltip(FTooltipDrawState& Tooltip, const ITimingEvent& HoveredTimingEvent) const override;
 	virtual const TSharedPtr<const ITimingEvent> SearchEvent(const FTimingEventSearchParameters& InSearchParameters) const override;
 	virtual void BuildContextMenu(FMenuBuilder& MenuBuilder) override;
-
-	virtual float GetBorderY() const { return BorderY; }
+	virtual bool UpdateSeriesBounds(FGameplayGraphSeries& InSeries, const FTimingTrackViewport& InViewport) override;
+	virtual void UpdateSeries(FGameplayGraphSeries& InSeries, const FTimingTrackViewport& InViewport) override;
+	virtual void AddAllSeries() override;
 
 	/** Get the asset ID that this track uses */
 	uint64 GetAssetId() const { return AssetId; }
@@ -61,26 +53,21 @@ public:
 	/** Get the node ID that this track uses */
 	int32 GetNodeId() const { return NodeId; }
 
-	/** The current height in lanes */
-	int32 HeightInLanes;
-
 private:
-	// Add all the series that this track can represent
-	void AddAllSeries();
-
-	// Helper for PreUpdate to update the track height
-	void UpdateTrackHeight(const ITimingTrackUpdateContext& Context);
-
 	// Helper function used to find a tick record
 	void FindTickRecordMessage(const FTimingEventSearchParameters& InParameters, TFunctionRef<void(double, double, uint32, const FTickRecordMessage&)> InFoundPredicate) const;
 
-	// Override certain UI items to be disabled
-	virtual bool ContextMenu_UseEventDuration_CanExecute() { return false; }
-	virtual bool ContextMenu_ShowBars_CanExecute() { return false; }
-
 	// Helper used to build track name
 	FText MakeTrackName(const FGameplaySharedData& InSharedData, uint64 InAssetID, const TCHAR* InName) const;
 
+	// Helper function for series bounds update
+	template<typename ProjectionType>
+	bool UpdateSeriesBoundsHelper(FTickRecordSeries& InSeries, const FTimingTrackViewport& InViewport, ProjectionType Projection);
+
+	// Helper function for series update
+	template<typename ProjectionType>
+	void UpdateSeriesHelper(FTickRecordSeries& InSeries, const FTimingTrackViewport& InViewport, ProjectionType Projection);
+
 private:
 	/** The shared data */
 	const FAnimationSharedData& SharedData;
@@ -99,10 +86,4 @@ private:
 
 	/** The node ID that this track's data comes from */
 	int32 NodeId;
-
-	/** The track size we want to be displayed at */
-	float RequestedTrackSizeScale;
-
-	/** The size of the border between tracks */
-	float BorderY;
 };
\ No newline at end of file
diff --git a/Engine/Plugins/Animation/GameplayInsights/Source/GameplayInsights/Private/GameplayGraphTrack.cpp b/Engine/Plugins/Animation/GameplayInsights/Source/GameplayInsights/Private/GameplayGraphTrack.cpp
new file mode 100644
--- /dev/null
+++ b/Engine/Plugins/Animation/GameplayInsights/Source/GameplayInsights/Private/GameplayGraphTrack.cpp
@@ -0,0 +1,319 @@
+// Copyright 1998-2019 Epic Games, Inc. All Rights Reserved.
+
+#include "GameplayGraphTrack.h"
+#include "GameplayProvider.h"
+#include "AnimationProvider.h"
+#include "Insights/ViewModels/TimingTrackViewport.h"
+#include "Insights/ViewModels/TimingEvent.h"
+#include "AnimationSharedData.h"
+#include "Framework/MultiBox/MultiBoxBuilder.h"
+#include "Insights/ViewModels/TimingEventSearch.h"
+#include "Insights/ViewModels/TooltipDrawState.h"
+#include "Modules/ModuleManager.h"
+#include "Insights/ViewModels/GraphTrackBuilder.h"
+#include "Insights/Common/PaintUtils.h"
+#include "Insights/ViewModels/ITimingViewDrawHelper.h"
+#include "Fonts/FontMeasure.h"
+
+#define LOCTEXT_NAMESPACE "GameplayGraphTrack"
+
+INSIGHTS_IMPLEMENT_RTTI(FGameplayGraphTrack)
+
+PRAGMA_DISABLE_OPTIMIZATION
+
+static float GetSeriesHeight(const FTimingViewLayout& InLayout)
+{
+	// 1.0f is for the top horizontal line of each track
+	return (InLayout.EventDY + InLayout.EventH);
+}
+
+void FGameplayGraphSeries::ComputePosition(const FTimingTrackViewport& InViewport, const FGameplayGraphTrack& InTrack, int32 InActiveSeriesIndex, float& OutTopY, float& OutBottomY) const
+{
+	const float ActiveSeriesIndex = InTrack.GetLayout() == EGameplayGraphLayout::Overlay ? 0.0f : (float)InActiveSeriesIndex;
+	const float TimelineDY = FMath::Max(1.0f, InViewport.GetLayout().TimelineDY);
+	const float TopY = FMath::Max(1.0f, TimelineDY);
+	const float SeriesHeight = FMath::Max(0.0f, GetSeriesHeight(InViewport.GetLayout())) * InTrack.GetRequestedTrackSizeScale();
+	const float BottomY = TopY + SeriesHeight;
+	const float OffsetY = SeriesHeight * ActiveSeriesIndex;
+
+	OutTopY = OffsetY + TopY;
+	OutBottomY = OffsetY + BottomY;
+}
+
+void FGameplayGraphSeries::UpdateAutoZoom(const FTimingTrackViewport& InViewport, const FGameplayGraphTrack& InTrack, int32 InActiveSeriesIndex)
+{
+	float TopY, BottomY;
+	ComputePosition(InViewport, InTrack, InActiveSeriesIndex, TopY, BottomY);
+
+	FGraphSeries::UpdateAutoZoom(TopY, BottomY, CurrentMin, CurrentMax, false);
+}
+
+FGameplayGraphTrack::FGameplayGraphTrack(uint64 InObjectID, const FText& InName)
+	: TGameplayTrackMixin<FGraphTrack>(InObjectID, InName)
+	, RequestedTrackSizeScale(1.0f)
+	, BorderY(0.0f)
+	, NumActiveSeries(0)
+	, Layout(EGameplayGraphLayout::Stack)
+{
+	bDrawPoints = false;
+	bDrawBoxes = false;
+	bDrawBaseline = false;
+	bUseEventDuration = false;
+	bDrawLabels = false;
+	VisibleOptions &= ~(EGraphOptions::ShowBars | EGraphOptions::UseEventDuration);
+}
+
+void FGameplayGraphTrack::UpdateTrackHeight(const ITimingTrackUpdateContext& Context)
+{
+	const FTimingTrackViewport& Viewport = Context.GetViewport();
+
+	int32 NumLanes = 0;
+	if(Layout == EGameplayGraphLayout::Overlay) 
+	{
+		NumLanes = (NumActiveSeries > 0 ? 1 : 0);
+	}
+	else
+	{
+		NumLanes = NumActiveSeries;
+	}
+	
+	const float CurrentTrackHeight = GetHeight();
+	const float TimelineDY2 = 2.0f * Viewport.GetLayout().TimelineDY;
+	const float DesiredTrackHeight = FMath::Max(0.0f, ((Viewport.GetLayout().ComputeTrackHeight(NumLanes) - TimelineDY2) * RequestedTrackSizeScale) + TimelineDY2);
+
+	if (CurrentTrackHeight < DesiredTrackHeight)
+	{
+		float NewTrackHeight;
+		if (Viewport.IsDirty(ETimingTrackViewportDirtyFlags::VLayoutChanged))
+		{
+			NewTrackHeight = DesiredTrackHeight;
+		}
+		else
+		{
+			NewTrackHeight = FMath::CeilToFloat(CurrentTrackHeight * 0.9f + DesiredTrackHeight * 0.1f);
+		}
+
+		SetHeight(NewTrackHeight);
+	}
+	else if (CurrentTrackHeight > DesiredTrackHeight)
+	{
+		float NewTrackHeight;
+		if (Viewport.IsDirty(ETimingTrackViewportDirtyFlags::VLayoutChanged))
+		{
+			NewTrackHeight = DesiredTrackHeight;
+		}
+		else
+		{
+			NewTrackHeight = FMath::FloorToFloat(CurrentTrackHeight * 0.9f + DesiredTrackHeight * 0.1f);
+		}
+
+		SetHeight(NewTrackHeight);
+	}
+}
+
+void FGameplayGraphTrack::UpdateSeriesInternal(FGameplayGraphSeries& InSeries, const FTimingTrackViewport& InViewport, int32 InActiveSeriesIndex)
+{
+	if(UpdateSeriesBounds(InSeries, InViewport))
+	{
+		InSeries.UpdateAutoZoom(InViewport, *this, InActiveSeriesIndex);
+	}
+
+	UpdateSeries(InSeries, InViewport);
+}
+
+void FGameplayGraphTrack::PreUpdate(const ITimingTrackUpdateContext& Context)
+{
+	if(AllSeries.Num() == 0)
+	{
+		AddAllSeries();
+	}
+
+	FGraphTrack::PreUpdate(Context);
+
+	// update border size
+	BorderY = Context.GetViewport().GetLayout().TimelineDY;
+
+	const bool bIsEntireGraphTrackDirty = IsDirty() || Context.GetViewport().IsHorizontalViewportDirty() || Context.GetViewport().IsDirty(ETimingTrackViewportDirtyFlags::VLayoutChanged);
+	bool bNeedsUpdate = bIsEntireGraphTrackDirty;
+
+	if (!bNeedsUpdate)
+	{
+		for (TSharedPtr<FGraphSeries>& Series : AllSeries)
+		{
+			if (Series->IsVisible() && Series->IsDirty())
+			{
+				// At least one series is dirty.
+				bNeedsUpdate = true;
+				break;
+			}
+		}
+	}
+
+	if (bNeedsUpdate)
+	{
+		ClearDirtyFlag();
+
+		NumActiveSeries = 0;
+
+		const FTimingTrackViewport& Viewport = Context.GetViewport();
+
+		for (TSharedPtr<FGraphSeries>& Series : AllSeries)
+		{
+			FGameplayGraphSeries& GameplayGraphSeries = *StaticCastSharedPtr<FGameplayGraphSeries>(Series);
+
+			if (Series->IsVisible() && (bIsEntireGraphTrackDirty || Series->IsDirty()))
+			{
+				// Clear the flag before updating, because the update itself may further need to set the series as dirty.
+				Series->ClearDirtyFlag();
+
+				UpdateSeriesInternal(GameplayGraphSeries, Viewport, NumActiveSeries);
+			}
+
+			if(GameplayGraphSeries.IsDrawn())
+			{
+				NumActiveSeries++;
+			}
+		}
+
+		UpdateStats();
+	}
+
+	UpdateTrackHeight(Context);
+}
+
+void FGameplayGraphTrack::Draw(const ITimingTrackDrawContext& Context) const
+{
+	FGraphTrack::Draw(Context);
+	GetGameplayTrack().DrawHeaderForTimingTrack(Context, *this, false);
+
+	const FTimingTrackViewport& Viewport = Context.GetViewport();
+	if(bDrawLabels && !Viewport.GetLayout().bIsCompactMode)
+	{
+		FDrawContext& DrawContext = Context.GetDrawContext();
+		const ITimingViewDrawHelper& DrawHelper = Context.GetHelper();
+
+		int32 ActiveSeriesIndex = 0;
+		for (const TSharedPtr<FGraphSeries>& Series : AllSeries)
+		{
+			FGameplayGraphSeries& GameplayGraphSeries = *StaticCastSharedPtr<FGameplayGraphSeries>(Series);
+
+			if (GameplayGraphSeries.IsDrawn())
+			{
+				const FString NameString = Series->GetName().ToString();
+				const TSharedRef<FSlateFontMeasure> FontMeasureService = FSlateApplication::Get().GetRenderer()->GetFontMeasureService();
+				const float NameHeight = FontMeasureService->Measure(NameString, DrawHelper.GetEventFont()).Y;
+
+				const float LocalPosY = FMath::RoundToFloat(GetPosY());
+				float TopY, BottomY;
+				GameplayGraphSeries.ComputePosition(Viewport, *this, ActiveSeriesIndex, TopY, BottomY);
+				const float BaselineY = ((TopY + BottomY) * 0.5f) - (NameHeight * 0.5f);
+				const float X = ((float)GetGameplayTrack().GetIndent() * GameplayTrackConstants::IndentSize) + 2.0f;
+				DrawContext.DrawText(DrawHelper.GetHeaderBackgroundLayerId(), X + 1.0f, BaselineY + LocalPosY + 1.0f, NameString, DrawHelper.GetEventFont(), FLinearColor::Black);
+				DrawContext.DrawText(DrawHelper.GetHeaderBackgroundLayerId() + 1, X, BaselineY + LocalPosY, NameString, DrawHelper.GetEventFont(), Series->GetColor());
+
+				ActiveSeriesIndex++;
+			}
+		}
+	}
+}
+
+void FGameplayGraphTrack::BuildContextMenu(FMenuBuilder& MenuBuilder)
+{
+	FGraphTrack::BuildContextMenu(MenuBuilder);
+
+	MenuBuilder.BeginSection("Layout", LOCTEXT("TrackLayoutMenuHeader", "Track Layout"));
+	{
+		MenuBuilder.AddMenuEntry
+		(
+			LOCTEXT("OverlayLayout", "Overlay"),
+			LOCTEXT("OverlayLayout_Tooltip", "Draw series overlaid one on top of the other."),
+			FSlateIcon(),
+			FUIAction(
+				FExecuteAction::CreateLambda([this](){ Layout = EGameplayGraphLayout::Overlay; bDrawLabels = false; SetDirtyFlag(); }),
+				FCanExecuteAction(),
+				FIsActionChecked::CreateLambda([this](){ return Layout == EGameplayGraphLayout::Overlay; })
+			),
+			NAME_None,
+			EUserInterfaceActionType::RadioButton
+		);
+
+		MenuBuilder.AddMenuEntry
+		(
+			LOCTEXT("StackLayout", "Stack"),
+			LOCTEXT("StackLayout_Tooltip", "Draw series in a vertical stack."),
+			FSlateIcon(),
+			FUIAction(
+				FExecuteAction::CreateLambda([this](){ Layout = EGameplayGraphLayout::Stack; SetDirtyFlag(); }),
+				FCanExecuteAction(),
+				FIsActionChecked::CreateLambda([this](){ return Layout == EGameplayGraphLayout::Stack; })
+			),
+			NAME_None,
+			EUserInterfaceActionType::RadioButton
+		);
+
+		MenuBuilder.AddMenuEntry
+		(
+			LOCTEXT("DrawLabels", "Labels"),
+			LOCTEXT("DrawLabels_Tooltip", "Draw series labels (stack view only)."),
+			FSlateIcon(),
+			FUIAction(
+				FExecuteAction::CreateLambda([this](){ bDrawLabels = !bDrawLabels; }),
+				FCanExecuteAction::CreateLambda([this](){ return Layout == EGameplayGraphLayout::Stack; }),
+				FIsActionChecked::CreateLambda([this](){ return bDrawLabels; })
+			),
+			NAME_None,
+			EUserInterfaceActionType::ToggleButton
+		);
+	}
+	MenuBuilder.EndSection();
+
+	MenuBuilder.BeginSection("TrackSize", LOCTEXT("TrackSizeMenuHeader", "Track Size"));
+	{
+		MenuBuilder.AddMenuEntry
+		(
+			LOCTEXT("NormalTrack", "Normal"),
+			LOCTEXT("NormalTrack_Tooltip", "Draw this track at the standard size."),
+			FSlateIcon(),
+			FUIAction(
+				FExecuteAction::CreateLambda([this](){ RequestedTrackSizeScale = 1.0f; SetDirtyFlag(); }),
+				FCanExecuteAction(),
+				FIsActionChecked::CreateLambda([this](){ return RequestedTrackSizeScale == 1.0f; })
+			),
+			NAME_None,
+			EUserInterfaceActionType::RadioButton
+		);
+
+		MenuBuilder.AddMenuEntry
+		(
+			LOCTEXT("LargeTrack", "Large"),
+			LOCTEXT("LargeTrack_Tooltip", "Make this track larger than normal."),
+			FSlateIcon(),
+			FUIAction(
+				FExecuteAction::CreateLambda([this](){ RequestedTrackSizeScale = 2.0f; SetDirtyFlag(); }),
+				FCanExecuteAction(),
+				FIsActionChecked::CreateLambda([this](){ return RequestedTrackSizeScale == 2.0f; })
+			),
+			NAME_None,
+			EUserInterfaceActionType::RadioButton
+		);
+
+		MenuBuilder.AddMenuEntry
+		(
+			LOCTEXT("ExtraLargeTrack", "Extra Large"),
+			LOCTEXT("ExtraLargeTrack_Tooltip", "Make this track much larger than normal."),
+			FSlateIcon(),
+			FUIAction(
+				FExecuteAction::CreateLambda([this](){ RequestedTrackSizeScale = 4.0f; SetDirtyFlag(); }),
+				FCanExecuteAction(),
+				FIsActionChecked::CreateLambda([this](){ return RequestedTrackSizeScale == 4.0f; })
+			),
+			NAME_None,
+			EUserInterfaceActionType::RadioButton
+		);
+	}
+	MenuBuilder.EndSection();
+}
+
+PRAGMA_ENABLE_OPTIMIZATION
+
+#undef LOCTEXT_NAMESPACE
\ No newline at end of file
diff --git a/Engine/Plugins/Animation/GameplayInsights/Source/GameplayInsights/Private/GameplayGraphTrack.h b/Engine/Plugins/Animation/GameplayInsights/Source/GameplayInsights/Private/GameplayGraphTrack.h
new file mode 100644
--- /dev/null
+++ b/Engine/Plugins/Animation/GameplayInsights/Source/GameplayInsights/Private/GameplayGraphTrack.h
@@ -0,0 +1,94 @@
+// Copyright 1998-2019 Epic Games, Inc. All Rights Reserved.
+
+#pragma once
+
+#include "GameplayTrack.h"
+#include "Insights/ViewModels/GraphTrack.h"
+#include "Insights/ViewModels/GraphSeries.h"
+
+class FTimingEventSearchParameters;
+class FGameplayGraphTrack;
+
+/** The various layouts that we display series with */
+enum class EGameplayGraphLayout : int32
+{
+	/** Draw series overlaid one on top of the other */
+	Overlay,
+
+	/** Draw series in a vertical stack */
+	Stack
+};
+
+// Common functionality for gameplay-related graph series
+class FGameplayGraphSeries : public FGraphSeries
+{
+public:
+	// Custom overload for auto zoom
+	void UpdateAutoZoom(const FTimingTrackViewport& InViewport, const FGameplayGraphTrack& InTrack, int32 InActiveSeriesIndex);
+
+	// Check whether this series has any events with the current view
+	bool IsDrawn() const { return IsVisible() && (Events.Num() > 0 || (LinePoints.Num() > 0 && LinePoints[0].Num() > 0) || Points.Num() > 0 || Boxes.Num() > 0); }
+
+	// Helper for auto-zoom and drawing
+	void ComputePosition(const FTimingTrackViewport& InViewport, const FGameplayGraphTrack& InTrack, int32 InActiveSeriesIndex, float& OutTopY, float& OutBottomY) const;
+
+public:
+	float CurrentMin;
+	float CurrentMax;
+};
+
+// Track holding common functionality for gameplay-related graph tracks
+class FGameplayGraphTrack : public TGameplayTrackMixin<FGraphTrack>
+{
+	INSIGHTS_DECLARE_RTTI(FGameplayGraphTrack, TGameplayTrackMixin<FGraphTrack>)
+
+public:
+	FGameplayGraphTrack(uint64 InObjectID, const FText& InName);
+
+	virtual void Draw(const ITimingTrackDrawContext& Context) const override;
+	virtual void PreUpdate(const ITimingTrackUpdateContext& Context) override;
+	virtual void BuildContextMenu(FMenuBuilder& MenuBuilder) override;
+
+	// Get the requested track size scale
+	float GetRequestedTrackSizeScale() const { return RequestedTrackSizeScale; }
+
+	// Get the number of active series (i.e. visible)
+	int32 GetNumActiveSeries() const { return NumActiveSeries; }
+
+	// Get the layout we are displaying out graphs as
+	EGameplayGraphLayout GetLayout() const { return Layout; }
+
+protected:
+	// Add all the series to this track
+	virtual void AddAllSeries() {}
+
+	// Update a single series' data (from PreUpdate())
+	virtual void UpdateSeries(FGameplayGraphSeries& InSeries, const FTimingTrackViewport& InViewport) {}
+
+	// Update a single series' bounds
+	// @return true from this if there are graph events in the current viewport
+	virtual bool UpdateSeriesBounds(FGameplayGraphSeries& InSeries, const FTimingTrackViewport& InViewport) { return false; }
+
+private:
+	// Helper for PreUpdate to update the track height
+	void UpdateTrackHeight(const ITimingTrackUpdateContext& Context);
+
+	// Helper to update and zoom series
+	void UpdateSeriesInternal(FGameplayGraphSeries& InSeries, const FTimingTrackViewport& InViewport, int32 InActiveSeriesIndex);
+
+protected:
+	/** The track size we want to be displayed at */
+	float RequestedTrackSizeScale;
+
+	/** The size of the border between tracks */
+	float BorderY;
+
+	/** The current number of active series */
+	int32 NumActiveSeries;
+
+	/** The current series layout */
+	EGameplayGraphLayout Layout;
+
+	/** Whether to draw labels */
+	bool bDrawLabels;
+};
\ No newline at end of file
diff --git a/Engine/Plugins/Animation/GameplayInsights/Source/GameplayInsights/Private/GameplaySharedData.cpp b/Engine/Plugins/Animation/GameplayInsights/Source/GameplayInsights/Private/GameplaySharedData.cpp
--- a/Engine/Plugins/Animation/GameplayInsights/Source/GameplayInsights/Private/GameplaySharedData.cpp
+++ b/Engine/Plugins/Animation/GameplayInsights/Source/GameplayInsights/Private/GameplaySharedData.cpp
@@ -7,6 +7,7 @@
 #include "GameplayProvider.h"
 #include "Insights/ITimingViewSession.h"
 #include "Framework/MultiBox/MultiBoxBuilder.h"
+#include "SGameplayTrackTree.h"
 
 #define LOCTEXT_NAMESPACE "GameplaySharedData"
 
@@ -75,28 +76,33 @@ TSharedRef<FObjectEventsTrack> FGameplaySharedData::GetObjectEventsTrackForId(In
 	return LeafObjectEventsTrack.ToSharedRef();
 }
 
-static void UpdateTrackOrderRecursive(TSharedRef<FObjectEventsTrack> InTrack, int32& InOutOrder)
+static void UpdateTrackOrderRecursive(TSharedRef<FBaseTimingTrack> InTrack, int32& InOutOrder)
 {
-	// recurse down object-track children, then non-object track leaf tracks to set 
-	// overall ordering based on depth-first traversal of the hierarchy
+	if(InTrack->Is<FObjectEventsTrack>())
+	{
+		TSharedRef<FObjectEventsTrack> ObjectEventsTrack = StaticCastSharedRef<FObjectEventsTrack>(InTrack);
 
-	InTrack->SetOrder(InOutOrder++);
+		// recurse down object-track children, then non-object track leaf tracks to set 
+		// overall ordering based on depth-first traversal of the hierarchy
 
-	for(FGameplayTrack* ChildTrack : InTrack->GetGameplayTrack().GetChildTracks())
-	{
-		if(ChildTrack->GetTimingTrack()->Is<FObjectEventsTrack>())
+		ObjectEventsTrack->SetOrder(InOutOrder++);
+
+		for(FGameplayTrack* ChildTrack : ObjectEventsTrack->GetGameplayTrack().GetChildTracks())
 		{
-			ChildTrack->SetIndent(InTrack->GetGameplayTrack().GetIndent() + 1);
-			UpdateTrackOrderRecursive(StaticCastSharedRef<FObjectEventsTrack>(ChildTrack->GetTimingTrack()), InOutOrder);
+			if(ChildTrack->GetTimingTrack()->Is<FObjectEventsTrack>())
+			{
+				ChildTrack->SetIndent(ObjectEventsTrack->GetGameplayTrack().GetIndent() + 1);
+				UpdateTrackOrderRecursive(StaticCastSharedRef<FObjectEventsTrack>(ChildTrack->GetTimingTrack()), InOutOrder);
+			}
 		}
-	}
 
-	for(FGameplayTrack* ChildTrack : InTrack->GetGameplayTrack().GetChildTracks())
-	{
-		if(!ChildTrack->GetTimingTrack()->Is<FObjectEventsTrack>())
+		for(FGameplayTrack* ChildTrack : ObjectEventsTrack->GetGameplayTrack().GetChildTracks())
 		{
-			ChildTrack->SetIndent(InTrack->GetGameplayTrack().GetIndent() + 1);
-			ChildTrack->GetTimingTrack()->SetOrder(InOutOrder++);
+			if(!ChildTrack->GetTimingTrack()->Is<FObjectEventsTrack>())
+			{
+				ChildTrack->SetIndent(ObjectEventsTrack->GetGameplayTrack().GetIndent() + 1);
+				ChildTrack->GetTimingTrack()->SetOrder(InOutOrder++);
+			}
 		}
 	}
 }
@@ -114,12 +120,12 @@ void FGameplaySharedData::Tick(Insights::ITimingViewSession& InTimingViewSession
 		// Add a track for each tracked object
 		GameplayProvider->EnumerateObjects([this, &InTimingViewSession, &InAnalysisSession, &GameplayProvider](const FObjectInfo& InObjectInfo)
 		{
-			TSharedPtr<FObjectEventsTrack> ObjectEventsTrack = GetObjectEventsTrackForId(InTimingViewSession, InAnalysisSession, InObjectInfo);
-			ObjectEventsTrack->SetVisibilityFlag(false);
-
-			GameplayProvider->ReadObjectEventsTimeline(InObjectInfo.Id, [this, &ObjectEventsTrack](const IGameplayProvider::ObjectEventsTimeline& InTimeline)
+			GameplayProvider->ReadObjectEventsTimeline(InObjectInfo.Id, [this, &InTimingViewSession, &InAnalysisSession, &InObjectInfo](const IGameplayProvider::ObjectEventsTimeline& InTimeline)
 			{
-				ObjectEventsTrack->SetVisibilityFlag(bObjectTracksEnabled);
+				if(InTimeline.GetEventCount() > 0)
+				{
+					GetObjectEventsTrackForId(InTimingViewSession, InAnalysisSession, InObjectInfo);
+				}
 			});
 		});
 
@@ -134,9 +140,25 @@ void FGameplaySharedData::Tick(Insights::ITimingViewSession& InTimingViewSession
 
 void FGameplaySharedData::ExtendFilterMenu(FMenuBuilder& InMenuBuilder)
 {
-	InMenuBuilder.AddMenuEntry(
+	InMenuBuilder.AddSubMenu(
 		LOCTEXT("ToggleGameplayTracks", "Gameplay Tracks"),
-		LOCTEXT("ToggleGameplayTracks_Tooltip", "Show/hide the gameplay tracks"),
+		LOCTEXT("ToggleGameplayTracks_Tooltip", "Show/hide individual gameplay tracks"),
+		FNewMenuDelegate::CreateLambda([this](FMenuBuilder& InSubMenuBuilder)
+		{ 
+			InSubMenuBuilder.AddWidget(
+				SNew(SBox)
+				.MaxDesiredHeight(300.0f)
+				.MinDesiredWidth(300.0f)
+				[
+					SNew(SGameplayTrackTree, *this)
+				],
+				FText(), true);
+		})
+	);
+
+	InMenuBuilder.AddMenuEntry(
+		LOCTEXT("ToggleEventTracks", "Event Tracks"),
+		LOCTEXT("ToggleEventTracks_Tooltip", "Show/hide the gameplay event tracks"),
 		FSlateIcon(),
 		FUIAction(
 			FExecuteAction::CreateRaw(this, &FGameplaySharedData::ToggleGameplayTracks),
@@ -152,30 +174,29 @@ void FGameplaySharedData::SortTracks()
 	int32 Order = 10000;
 
 	// Find current roots
-	static TArray<TSharedRef<FObjectEventsTrack>> Roots;
-	check(Roots.Num() == 0);
+	RootTracks.Reset();
 
 	for(auto ObjectTrackPair : ObjectTracks)
 	{
 		if(ObjectTrackPair.Value->GetGameplayTrack().GetParentTrack() == nullptr)
 		{
-			Roots.Add(ObjectTrackPair.Value.ToSharedRef());
+			RootTracks.Add(ObjectTrackPair.Value.ToSharedRef());
 		}
 	}
 
 	// Sort roots alphabetically
-	Algo::Sort(Roots, [](TSharedRef<FObjectEventsTrack> InTrack0, TSharedRef<FObjectEventsTrack> InTrack1)
+	Algo::Sort(RootTracks, [](TSharedPtr<FBaseTimingTrack> InTrack0, TSharedPtr<FBaseTimingTrack> InTrack1)
 	{
 		return InTrack0->GetName() < InTrack1->GetName();
 	});
 
 	// update ordering/indent
-	for(TSharedRef<FObjectEventsTrack> RootTrack : Roots)
+	for(TSharedPtr<FBaseTimingTrack> RootTrack : RootTracks)
 	{
-		UpdateTrackOrderRecursive(RootTrack, Order);
+		UpdateTrackOrderRecursive(RootTrack.ToSharedRef(), Order);
 	}
 
-	Roots.Reset();
+	OnTracksChangedDelegate.Broadcast();
 }
 
 void FGameplaySharedData::ToggleGameplayTracks()
diff --git a/Engine/Plugins/Animation/GameplayInsights/Source/GameplayInsights/Private/GameplaySharedData.h b/Engine/Plugins/Animation/GameplayInsights/Source/GameplayInsights/Private/GameplaySharedData.h
--- a/Engine/Plugins/Animation/GameplayInsights/Source/GameplayInsights/Private/GameplaySharedData.h
+++ b/Engine/Plugins/Animation/GameplayInsights/Source/GameplayInsights/Private/GameplaySharedData.h
@@ -36,8 +36,14 @@ public:
 	// Enumerate object tracks
 	void EnumerateObjectTracks(TFunctionRef<void(const TSharedRef<FObjectEventsTrack>&)> InCallback) const;
 
+	// Get the root tracks
+	const TArray<TSharedRef<FBaseTimingTrack>>& GetRootTracks() const { return RootTracks; }
+
+	// Delegate fired when tracks change
+	FSimpleMulticastDelegate& OnTracksChanged() { return OnTracksChangedDelegate; }
+
 private:
-	// Re-sort tracks if trrack ordering has changed
+	// Re-sort tracks if track ordering has changed
 	void SortTracks();
 
 	// UI handlers
@@ -47,9 +53,15 @@ private:
 	// Track for each tracked object, mapped from Object ID -> track
 	TMap<uint64, TSharedPtr<FObjectEventsTrack>> ObjectTracks;
 
+	// The root tracks
+	TArray<TSharedRef<FBaseTimingTrack>> RootTracks;
+
 	// Cached analysis session, set in Tick()
 	const Trace::IAnalysisSession* AnalysisSession;
 
+	// Delegate fired when tracks change
+	FSimpleMulticastDelegate OnTracksChangedDelegate;
+
 	// Dirty flag for adding object tracks, used to trigger re-sorting
 	bool bObjectTracksDirty;
 
diff --git a/Engine/Plugins/Animation/GameplayInsights/Source/GameplayInsights/Private/GameplayTimingViewExtender.cpp b/Engine/Plugins/Animation/GameplayInsights/Source/GameplayInsights/Private/GameplayTimingViewExtender.cpp
--- a/Engine/Plugins/Animation/GameplayInsights/Source/GameplayInsights/Private/GameplayTimingViewExtender.cpp
+++ b/Engine/Plugins/Animation/GameplayInsights/Source/GameplayInsights/Private/GameplayTimingViewExtender.cpp
@@ -11,6 +11,8 @@
 #include "Editor/EditorEngine.h"
 #endif
 
+#define LOCTEXT_NAMESPACE "GameplayTimingViewExtender"
+
 void FGameplayTimingViewExtender::OnBeginSession(Insights::ITimingViewSession& InSession)
 {
 	FPerSessionData* PerSessionData = PerSessionDataMap.Find(&InSession);
@@ -57,8 +59,12 @@ void FGameplayTimingViewExtender::ExtendFilterMenu(Insights::ITimingViewSession&
 	FPerSessionData* PerSessionData = PerSessionDataMap.Find(&InSession);
 	if(PerSessionData != nullptr)
 	{
-		PerSessionData->GameplaySharedData->ExtendFilterMenu(InMenuBuilder);
-		PerSessionData->AnimationSharedData->ExtendFilterMenu(InMenuBuilder);
+		InMenuBuilder.BeginSection("GameplayTracks", LOCTEXT("GameplayTracksSection", "Gameplay Tracks"));
+		{
+			PerSessionData->GameplaySharedData->ExtendFilterMenu(InMenuBuilder);
+			PerSessionData->AnimationSharedData->ExtendFilterMenu(InMenuBuilder);
+		}
+		InMenuBuilder.EndSection();
 	}
 }
 
@@ -111,4 +117,6 @@ void FGameplayTimingViewExtender::GetCustomDebugObjects(const IAnimationBlueprin
 	}
 }
 
-#endif
\ No newline at end of file
+#endif
+
+#undef LOCTEXT_NAMESPACE
\ No newline at end of file
diff --git a/Engine/Plugins/Animation/GameplayInsights/Source/GameplayInsights/Private/GameplayTrack.cpp b/Engine/Plugins/Animation/GameplayInsights/Source/GameplayInsights/Private/GameplayTrack.cpp
--- a/Engine/Plugins/Animation/GameplayInsights/Source/GameplayInsights/Private/GameplayTrack.cpp
+++ b/Engine/Plugins/Animation/GameplayInsights/Source/GameplayInsights/Private/GameplayTrack.cpp
@@ -9,11 +9,6 @@
 #include "Insights/Common/PaintUtils.h"
 #include "Insights/ViewModels/ITimingViewDrawHelper.h"
 
-namespace GameplayTrackConstants
-{
-	constexpr float IndentSize = 12.0f;
-}
-
 void FGameplayTrack::AddChildTrack(FGameplayTrack& InChildTrack)
 {
 	check(InChildTrack.Parent == nullptr);
diff --git a/Engine/Plugins/Animation/GameplayInsights/Source/GameplayInsights/Private/GameplayTrack.h b/Engine/Plugins/Animation/GameplayInsights/Source/GameplayInsights/Private/GameplayTrack.h
--- a/Engine/Plugins/Animation/GameplayInsights/Source/GameplayInsights/Private/GameplayTrack.h
+++ b/Engine/Plugins/Animation/GameplayInsights/Source/GameplayInsights/Private/GameplayTrack.h
@@ -4,6 +4,11 @@
 
 #include "Insights/ViewModels/TimingEventsTrack.h"
 
+namespace GameplayTrackConstants
+{
+	constexpr float IndentSize = 12.0f;
+}
+
 // Provides parent/child hierarchy structure and the owning object Id
 // Designed as a compositional member of outer timing tracks
 class FGameplayTrack
diff --git a/Engine/Plugins/Animation/GameplayInsights/Source/GameplayInsights/Private/IAnimationProvider.h b/Engine/Plugins/Animation/GameplayInsights/Source/GameplayInsights/Private/IAnimationProvider.h
--- a/Engine/Plugins/Animation/GameplayInsights/Source/GameplayInsights/Private/IAnimationProvider.h
+++ b/Engine/Plugins/Animation/GameplayInsights/Source/GameplayInsights/Private/IAnimationProvider.h
@@ -190,8 +190,10 @@ public:
 	typedef Trace::ITimeline<FAnimStateMachineMessage> StateMachinesTimeline;
 	typedef Trace::ITimeline<FBlendSpacePlayerMessage> BlendSpacePlayersTimeline;
 
-	virtual bool ReadSkeletalMeshPoseTimeline(uint64 InObjectId, TFunctionRef<void(const SkeletalMeshPoseTimeline&)> Callback) const = 0;
+	virtual bool ReadSkeletalMeshPoseTimeline(uint64 InObjectId, TFunctionRef<void(const SkeletalMeshPoseTimeline&, bool)> Callback) const = 0;
 	virtual void GetSkeletalMeshComponentSpacePose(const FSkeletalMeshPoseMessage& InMessage, const FSkeletalMeshInfo& InMeshInfo, FTransform& OutComponentToWorld, TArray<FTransform>& OutTransforms) const = 0;
+	virtual void EnumerateSkeletalMeshCurveIds(uint64 InObjectId, TFunctionRef<void(uint32)> Callback) const = 0;
+	virtual void EnumerateSkeletalMeshCurves(const FSkeletalMeshPoseMessage& InMessage, TFunctionRef<void(const FSkeletalMeshNamedCurve&)> Callback) const = 0;
 	virtual void EnumerateTickRecordTimelines(uint64 InObjectId, TFunctionRef<void(uint64, int32, const TickRecordTimeline&)> Callback) const = 0;
 	virtual bool ReadTickRecordTimeline(uint64 InObjectId, uint64 InAssetId, int32 InNodeId, TFunctionRef<void(const TickRecordTimeline&)> Callback) const = 0;
 	virtual bool ReadAnimGraphTimeline(uint64 InObjectId, TFunctionRef<void(const AnimGraphTimeline&)> Callback) const = 0;
diff --git a/Engine/Plugins/Animation/GameplayInsights/Source/GameplayInsights/Private/SGameplayTrackTree.cpp b/Engine/Plugins/Animation/GameplayInsights/Source/GameplayInsights/Private/SGameplayTrackTree.cpp
new file mode 100644
--- /dev/null
+++ b/Engine/Plugins/Animation/GameplayInsights/Source/GameplayInsights/Private/SGameplayTrackTree.cpp
@@ -0,0 +1,307 @@
+// Copyright 1998-2019 Epic Games, Inc. All Rights Reserved.
+
+#include "SGameplayTrackTree.h"
+
+#include "EditorStyleSet.h"
+#include "Framework/Application/SlateApplication.h"
+#include "Widgets/Input/SCheckBox.h"
+#include "Widgets/Input/SSearchBox.h"
+#include "Widgets/Text/STextBlock.h"
+#include "Widgets/Layout/SScrollBorder.h"
+#include "Widgets/SBoxPanel.h"
+#include "Widgets/Views/SHeaderRow.h"
+#include "ObjectEventsTrack.h"
+#include "GameplaySharedData.h"
+
+#define LOCTEXT_NAMESPACE "SGameplayTrackTree"
+
+enum class EGameplayTrackFilterState
+{
+	Hidden,
+	Visible,
+	Highlighted,
+};
+
+// Simple wrapper around a gameplay track for filtering
+struct FGameplayTrackTreeEntry
+{
+	FGameplayTrackTreeEntry(TSharedRef<FBaseTimingTrack> InTimingTrack)
+		: WeakTimingTrack(InTimingTrack)
+		, FilterState(EGameplayTrackFilterState::Hidden)
+	{}
+
+	FText GetName() const
+	{
+		FText Name;
+
+		TSharedPtr<FBaseTimingTrack> TimingTrack = WeakTimingTrack.Pin();
+		if(TimingTrack.IsValid())
+		{
+			Name = FText::FromString(TimingTrack->GetName());
+		}
+
+		return Name;
+	}
+
+	bool IsVisible() const
+	{
+		TSharedPtr<FBaseTimingTrack> TimingTrack = WeakTimingTrack.Pin();
+		if(TimingTrack.IsValid())
+		{
+			return TimingTrack->IsVisible();
+		}
+
+		return false;
+	}
+
+	void SetVisibilityFlag(bool bInIsVisible)
+	{
+		TSharedPtr<FBaseTimingTrack> TimingTrack = WeakTimingTrack.Pin();
+		if(TimingTrack.IsValid())
+		{
+			TimingTrack->SetVisibilityFlag(bInIsVisible);
+		}
+
+		for(const TSharedRef<FGameplayTrackTreeEntry>& Child : Children)
+		{
+			if(Child->FilterState != EGameplayTrackFilterState::Hidden)
+			{
+				Child->SetVisibilityFlag(bInIsVisible);
+			}
+		}
+	}
+
+	TWeakPtr<FBaseTimingTrack> WeakTimingTrack;
+	TArray<TSharedRef<FGameplayTrackTreeEntry>> Children;
+	EGameplayTrackFilterState FilterState;
+};
+
+// A list entry widget for a gameplay track
+class SGameplayTrackTreeEntry : public SCompoundWidget
+{
+public:
+	SLATE_BEGIN_ARGS(SGameplayTrackTreeEntry) {}
+
+	SLATE_ARGUMENT(TSharedPtr<FGameplayTrackTreeEntry>, TreeEntry)
+
+	SLATE_ATTRIBUTE(FText, SearchText)
+
+	SLATE_END_ARGS()
+
+	void Construct(const FArguments& InArgs)
+	{
+		WeakTreeEntry = InArgs._TreeEntry;
+		SearchText = InArgs._SearchText;
+
+		ChildSlot
+		[
+			SNew(SCheckBox)
+			.IsChecked_Lambda([this]()
+			{ 
+				TSharedPtr<FGameplayTrackTreeEntry> TreeEntry = WeakTreeEntry.Pin();
+				if(TreeEntry.IsValid())
+				{
+					return TreeEntry->IsVisible() ? ECheckBoxState::Checked : ECheckBoxState::Unchecked; 
+				}
+				return ECheckBoxState::Unchecked; 
+			})
+			.OnCheckStateChanged_Lambda([this](ECheckBoxState InCheckBoxState)
+			{
+				TSharedPtr<FGameplayTrackTreeEntry> TreeEntry = WeakTreeEntry.Pin();
+				if(TreeEntry.IsValid())
+				{
+					TreeEntry->SetVisibilityFlag(InCheckBoxState == ECheckBoxState::Checked);
+				}
+			})
+			.Content()
+			[
+				SNew(STextBlock)
+				.Text(InArgs._TreeEntry->GetName())
+				.HighlightText(SearchText)
+			]
+		];
+	}
+
+	// The tree entry we represent
+	TWeakPtr<FGameplayTrackTreeEntry> WeakTreeEntry;
+
+	// The search text to highlight
+	TAttribute<FText> SearchText;
+};
+
+void SGameplayTrackTree::Construct(const FArguments& InArgs, FGameplaySharedData& InSharedData)
+{
+	SharedData = &InSharedData;
+
+	InSharedData.OnTracksChanged().AddSP(this, &SGameplayTrackTree::HandleTracksChanged);
+
+	TreeView = SNew(STreeView<TSharedRef<FGameplayTrackTreeEntry>>)
+		.ItemHeight(20.0f)
+		.SelectionMode(ESelectionMode::None)
+		.TreeItemsSource(&FilteredTracks)
+		.OnGenerateRow(this, &SGameplayTrackTree::OnGenerateRow)
+		.OnGetChildren(this, &SGameplayTrackTree::OnGetChildren);
+
+	ChildSlot
+	[
+		SNew(SVerticalBox)
+		+SVerticalBox::Slot()
+		.AutoHeight()
+		[
+			SNew(SHorizontalBox)
+			+SHorizontalBox::Slot()
+			.AutoWidth()
+			.Padding(2.0f, 0.0f)
+			[
+				// Checkbox for bulk operations
+				SNew(SCheckBox)
+				.IsChecked_Lambda([this]()
+				{
+					int32 NumVisible = 0;
+					for (TSharedRef<FGameplayTrackTreeEntry> Track : FilteredTracks)
+					{
+						if (Track->IsVisible())
+						{
+							NumVisible++;
+						}
+					}
+
+					if (NumVisible == 0)
+					{
+						return ECheckBoxState::Unchecked;
+					}
+					else if (NumVisible == FilteredTracks.Num())
+					{
+						return ECheckBoxState::Checked;
+					}
+
+					return ECheckBoxState::Undetermined;
+				})
+				.OnCheckStateChanged_Lambda([this](ECheckBoxState InCheckBoxState)
+				{
+					const bool bVisible = InCheckBoxState != ECheckBoxState::Unchecked;
+					for (TSharedRef<FGameplayTrackTreeEntry> Track : FilteredTracks)
+					{
+						Track->SetVisibilityFlag(bVisible);
+					}
+				})
+			]
+			+ SHorizontalBox::Slot()
+			.FillWidth(1.0f)
+			[
+				// Search box allows for filtering
+				SAssignNew(SearchBox, SSearchBox)
+				.OnTextChanged_Lambda([this](const FText& InText){ SearchText = InText; RefreshFilter(); })
+			]
+		]
+		+SVerticalBox::Slot()
+		.FillHeight(1.0f)
+		[
+			SNew(SScrollBorder, TreeView.ToSharedRef())
+			[
+				TreeView.ToSharedRef()
+			]
+		]
+	];
+
+	// Set focus to the search box on creation
+	FSlateApplication::Get().SetKeyboardFocus(SearchBox);
+	FSlateApplication::Get().SetUserFocus(0, SearchBox);
+
+	RefreshFilter();
+}
+
+void SGameplayTrackTree::HandleTracksChanged()
+{
+	RefreshFilter();
+}
+
+TSharedRef<ITableRow> SGameplayTrackTree::OnGenerateRow(TSharedRef<FGameplayTrackTreeEntry> InItem, const TSharedRef<STableViewBase>& InOwnerTable)
+{
+	return
+		SNew(STableRow<TSharedRef<FGameplayTrackTreeEntry>>, InOwnerTable)
+		[
+			SNew(SGameplayTrackTreeEntry)
+			.TreeEntry(InItem)
+			.SearchText_Lambda([this](){ return SearchText; })
+		];
+}
+
+void SGameplayTrackTree::OnGetChildren(TSharedRef<FGameplayTrackTreeEntry> InItem, TArray<TSharedRef<FGameplayTrackTreeEntry>>& OutChildren)
+{
+	for(const TSharedRef<FGameplayTrackTreeEntry>& ChildItem : InItem->Children)
+	{
+		if(ChildItem->FilterState != EGameplayTrackFilterState::Hidden)
+		{
+			OutChildren.Add(ChildItem);
+		}
+	}
+}
+
+EGameplayTrackFilterState SGameplayTrackTree::RefreshFilter_Helper(const TSharedRef<FGameplayTrackTreeEntry>& InTreeEntry)
+{
+	TSharedPtr<FObjectEventsTrack> ObjectEventsTrack;
+	TSharedPtr<FBaseTimingTrack> TimingTrack = InTreeEntry->WeakTimingTrack.Pin();
+	if(TimingTrack.IsValid())
+	{
+		if(TimingTrack->Is<FObjectEventsTrack>())
+		{
+			ObjectEventsTrack = StaticCastSharedPtr<FObjectEventsTrack>(TimingTrack);
+		}
+	}
+
+	if(ObjectEventsTrack.IsValid() && ObjectEventsTrack->GetGameplayTrack().GetChildTracks().Num() > 0)
+	{
+		InTreeEntry->FilterState = EGameplayTrackFilterState::Hidden;
+
+		for(FGameplayTrack* ChildTrack : ObjectEventsTrack->GetGameplayTrack().GetChildTracks())
+		{
+			TSharedRef<FGameplayTrackTreeEntry> NewTreeEntry = MakeShared<FGameplayTrackTreeEntry>(ChildTrack->GetTimingTrack());
+			InTreeEntry->Children.Add(NewTreeEntry);
+			EGameplayTrackFilterState ChildFilterState = RefreshFilter_Helper(NewTreeEntry);
+			InTreeEntry->FilterState = FMath::Max(ChildFilterState, InTreeEntry->FilterState);
+		}
+	}
+	
+	if(InTreeEntry->FilterState == EGameplayTrackFilterState::Hidden)
+	{
+		if(SearchText.IsEmpty())
+		{
+			InTreeEntry->FilterState = EGameplayTrackFilterState::Visible;
+		}
+		else if (InTreeEntry->GetName().ToString().Contains(SearchText.ToString()))
+		{
+			InTreeEntry->FilterState = EGameplayTrackFilterState::Highlighted;
+		}
+		else
+		{
+			InTreeEntry->FilterState = EGameplayTrackFilterState::Hidden;
+		}
+	}
+
+	if(InTreeEntry->FilterState == EGameplayTrackFilterState::Highlighted)
+	{
+		TreeView->SetItemExpansion(InTreeEntry, true);
+	}
+
+	return InTreeEntry->FilterState;
+}
+
+void SGameplayTrackTree::RefreshFilter()
+{
+	FilteredTracks.Reset();
+
+	for(const TSharedRef<FBaseTimingTrack>& RootTrack : SharedData->GetRootTracks())
+	{
+		TSharedRef<FGameplayTrackTreeEntry> NewTreeEntry = MakeShared<FGameplayTrackTreeEntry>(RootTrack);
+		EGameplayTrackFilterState FilterState = RefreshFilter_Helper(NewTreeEntry);
+		if(FilterState != EGameplayTrackFilterState::Hidden)
+		{
+			FilteredTracks.Add(NewTreeEntry);	
+		}
+	}
+
+	TreeView->RequestTreeRefresh();
+}
+
+#undef LOCTEXT_NAMESPACE
diff --git a/Engine/Plugins/Animation/GameplayInsights/Source/GameplayInsights/Private/SGameplayTrackTree.h b/Engine/Plugins/Animation/GameplayInsights/Source/GameplayInsights/Private/SGameplayTrackTree.h
new file mode 100644
--- /dev/null
+++ b/Engine/Plugins/Animation/GameplayInsights/Source/GameplayInsights/Private/SGameplayTrackTree.h
@@ -0,0 +1,56 @@
+// Copyright 1998-2019 Epic Games, Inc. All Rights Reserved.
+
+#pragma once
+
+#include "CoreMinimal.h"
+#include "Widgets/SCompoundWidget.h"
+#include "Widgets/Views/STreeView.h"
+
+class SSearchBox;
+class FBaseTimingTrack;
+class FGameplaySharedData;
+struct FGameplayTrackTreeEntry;
+enum class EGameplayTrackFilterState;
+
+// A widget representing all the gameplay tracks
+class SGameplayTrackTree : public SCompoundWidget
+{
+public:
+	SLATE_BEGIN_ARGS(SGameplayTrackTree) {}
+
+	SLATE_END_ARGS()
+
+	void Construct(const FArguments& InArgs, FGameplaySharedData& InSharedData);
+
+private:
+	// Generate a row for the tree
+	TSharedRef<ITableRow> OnGenerateRow(TSharedRef<FGameplayTrackTreeEntry> Item, const TSharedRef<STableViewBase>& OwnerTable);
+
+	// Get the children of an item
+	void OnGetChildren(TSharedRef<FGameplayTrackTreeEntry> InItem, TArray<TSharedRef<FGameplayTrackTreeEntry>>& OutChildren);
+
+	// Refresh the tree as the filter has changed
+	void RefreshFilter();
+
+	// Recursive helper to filter tree
+	EGameplayTrackFilterState RefreshFilter_Helper(const TSharedRef<FGameplayTrackTreeEntry>& InTreeEntry);
+
+	// Handle the tracks changing under us
+	void HandleTracksChanged();
+
+private:
+	// Shared data source
+	FGameplaySharedData* SharedData;
+
+	// The search widget
+	TSharedPtr<SSearchBox> SearchBox;
+
+	// The list view widget
+	TSharedPtr<STreeView<TSharedRef<FGameplayTrackTreeEntry>>> TreeView;
+
+	// Text we are searching for
+	FText SearchText;
+
+	// Filtered list of tracks
+	TArray<TSharedRef<FGameplayTrackTreeEntry>> FilteredTracks;
+};
diff --git a/Engine/Plugins/Animation/GameplayInsights/Source/GameplayInsights/Private/SkeletalMeshCurvesTrack.cpp b/Engine/Plugins/Animation/GameplayInsights/Source/GameplayInsights/Private/SkeletalMeshCurvesTrack.cpp
new file mode 100644
--- /dev/null
+++ b/Engine/Plugins/Animation/GameplayInsights/Source/GameplayInsights/Private/SkeletalMeshCurvesTrack.cpp
@@ -0,0 +1,209 @@
+// Copyright 1998-2019 Epic Games, Inc. All Rights Reserved.
+
+#include "SkeletalMeshCurvesTrack.h"
+#include "GameplayProvider.h"
+#include "AnimationProvider.h"
+#include "Insights/ViewModels/TimingTrackViewport.h"
+#include "Insights/ViewModels/TimingEvent.h"
+#include "AnimationSharedData.h"
+#include "Framework/MultiBox/MultiBoxBuilder.h"
+#include "Insights/ViewModels/TimingEventSearch.h"
+#include "Insights/ViewModels/TooltipDrawState.h"
+#include "Modules/ModuleManager.h"
+#include "Insights/ViewModels/GraphTrackBuilder.h"
+
+#define LOCTEXT_NAMESPACE "SkeletalMeshCurvesTrack"
+
+INSIGHTS_IMPLEMENT_RTTI(FSkeletalMeshCurvesTrack)
+
+FSkeletalMeshCurvesTrack::FSkeletalMeshCurvesTrack(const FAnimationSharedData& InSharedData, uint64 InObjectID, const TCHAR* InName)
+	: FGameplayGraphTrack(InObjectID, FText::Format(LOCTEXT("TrackNameFormat", "Curves - {0}"), FText::FromString(FString(InName))))
+	, SharedData(InSharedData)
+{
+}
+
+void FSkeletalMeshCurvesTrack::AddAllSeries()
+{
+	const FAnimationProvider* AnimationProvider = SharedData.GetAnalysisSession().ReadProvider<FAnimationProvider>(FAnimationProvider::ProviderName);
+
+	bool bFirstSeries = true;
+	if(AnimationProvider)
+	{
+		Trace::FAnalysisSessionReadScope SessionReadScope(SharedData.GetAnalysisSession());
+
+		AnimationProvider->EnumerateSkeletalMeshCurveIds(GetGameplayTrack().GetObjectId(), [this, &bFirstSeries, &AnimationProvider](uint32 InCurveId)
+		{
+			auto MakeCurveSeriesColor = [](uint32 InSeed, bool bInLine)
+			{
+				FRandomStream Stream(InSeed);
+				const uint8 Hue = (uint8)(Stream.FRand() * 255.0f);
+				const uint8 SatVal = bInLine ? 196 : 128;
+				return FLinearColor::MakeFromHSV8(Hue, SatVal, SatVal);
+			};
+
+			TSharedRef<FSkeletalMeshCurveSeries> Series = MakeShared<FSkeletalMeshCurveSeries>();
+
+			const TCHAR* CurveName = AnimationProvider->GetName(InCurveId);
+
+			Series->SetName(CurveName);
+			Series->SetDescription(FText::Format(LOCTEXT("CurveNameFormat", "Values for curve '{0}'"), FText::FromString(CurveName)));
+
+			const FLinearColor LineColor = MakeCurveSeriesColor(InCurveId, true);
+			const FLinearColor FillColor = MakeCurveSeriesColor(InCurveId, false);
+			Series->SetColor(LineColor, LineColor, FillColor);
+
+			Series->CurveId = InCurveId;
+			Series->SetVisibility(bFirstSeries);
+			Series->SetBaselineY(25.0f);
+			Series->SetScaleY(20.0f);
+			Series->EnableAutoZoom();
+			AllSeries.Add(Series);	
+
+			bFirstSeries = false;
+		});
+	}
+}
+
+bool FSkeletalMeshCurvesTrack::UpdateSeriesBounds(FGameplayGraphSeries& InSeries, const FTimingTrackViewport& InViewport)
+{
+	bool bFoundEvents = false;
+
+	const FAnimationProvider* AnimationProvider = SharedData.GetAnalysisSession().ReadProvider<FAnimationProvider>(FAnimationProvider::ProviderName);
+
+	if(AnimationProvider)
+	{
+		Trace::FAnalysisSessionReadScope SessionReadScope(SharedData.GetAnalysisSession());
+
+		FSkeletalMeshCurveSeries& CurveSeries = *static_cast<FSkeletalMeshCurveSeries*>(&InSeries);
+
+		CurveSeries.CurrentMin = 0.0f;
+		CurveSeries.CurrentMax = 0.0f;
+
+		AnimationProvider->ReadSkeletalMeshPoseTimeline(GetGameplayTrack().GetObjectId(), [&bFoundEvents, &InViewport, &CurveSeries, &AnimationProvider](const FAnimationProvider::SkeletalMeshPoseTimeline& InTimeline, bool bInHasCurves)
+		{
+			InTimeline.EnumerateEvents(InViewport.GetStartTime(), InViewport.GetEndTime(), [&bFoundEvents, &CurveSeries, &AnimationProvider](double InStartTime, double InEndTime, uint32 InDepth, const FSkeletalMeshPoseMessage& InMessage)
+			{
+				AnimationProvider->EnumerateSkeletalMeshCurves(InMessage, [&bFoundEvents, &CurveSeries](const FSkeletalMeshNamedCurve& InCurve)
+				{
+					if(InCurve.Id == CurveSeries.CurveId)
+					{
+						CurveSeries.CurrentMin = FMath::Min(CurveSeries.CurrentMin, InCurve.Value);
+						CurveSeries.CurrentMax = FMath::Max(CurveSeries.CurrentMax, InCurve.Value);
+						bFoundEvents = true;
+					}
+				});
+			});
+		});
+	}
+
+	return bFoundEvents;
+}
+
+void FSkeletalMeshCurvesTrack::UpdateSeries(FGameplayGraphSeries& InSeries, const FTimingTrackViewport& InViewport)
+{
+	const FAnimationProvider* AnimationProvider = SharedData.GetAnalysisSession().ReadProvider<FAnimationProvider>(FAnimationProvider::ProviderName);
+
+	if(AnimationProvider)
+	{
+		Trace::FAnalysisSessionReadScope SessionReadScope(SharedData.GetAnalysisSession());
+
+		FSkeletalMeshCurveSeries& CurveSeries = *static_cast<FSkeletalMeshCurveSeries*>(&InSeries);
+
+		FGraphTrackBuilder Builder(*this, CurveSeries, InViewport);
+
+		AnimationProvider->ReadSkeletalMeshPoseTimeline(GetGameplayTrack().GetObjectId(), [this, &Builder, &InViewport, &CurveSeries, &AnimationProvider](const FAnimationProvider::SkeletalMeshPoseTimeline& InTimeline, bool bInHasCurves)
+		{
+			int32 FrameCounter = 0;
+			int32 LastFrameWithCurve = 0;
+
+			InTimeline.EnumerateEvents(InViewport.GetStartTime(), InViewport.GetEndTime(), [this, &FrameCounter, &LastFrameWithCurve, &Builder, &CurveSeries, &AnimationProvider](double InStartTime, double InEndTime, uint32 InDepth, const FSkeletalMeshPoseMessage& InMessage)
+			{
+				FrameCounter++;
+
+				AnimationProvider->EnumerateSkeletalMeshCurves(InMessage, [this, &FrameCounter, &LastFrameWithCurve, &Builder, &CurveSeries, &InStartTime, &InEndTime](const FSkeletalMeshNamedCurve& InCurve)
+				{
+					if(InCurve.Id == CurveSeries.CurveId)
+					{
+						Builder.AddEvent(InStartTime, InEndTime - InStartTime, InCurve.Value, LastFrameWithCurve == FrameCounter - 1);
+
+						LastFrameWithCurve = FrameCounter;
+					}
+				});
+			});
+		});
+	}
+}
+
+void FSkeletalMeshCurvesTrack::InitTooltip(FTooltipDrawState& Tooltip, const ITimingEvent& HoveredTimingEvent) const
+{
+	const FGraphTrackEvent& GraphTrackEvent = *static_cast<const FGraphTrackEvent*>(&HoveredTimingEvent);
+
+	FTimingEventSearchParameters SearchParameters(HoveredTimingEvent.GetStartTime(), HoveredTimingEvent.GetEndTime(), ETimingEventSearchFlags::StopAtFirstMatch);
+
+	FindSkeletalMeshPoseMessage(SearchParameters, [this, &Tooltip, &GraphTrackEvent](double InFoundStartTime, double InFoundEndTime, uint32 InFoundDepth, const FSkeletalMeshPoseMessage& InMessage)
+	{
+		Tooltip.ResetContent();
+
+		Tooltip.AddTitle(GetName());
+
+		{
+			Trace::FAnalysisSessionReadScope SessionReadScope(SharedData.GetAnalysisSession());
+
+			const FAnimationProvider* AnimationProvider = SharedData.GetAnalysisSession().ReadProvider<FAnimationProvider>(FAnimationProvider::ProviderName);
+			if(AnimationProvider)
+			{
+				const TCHAR* CurveName = AnimationProvider->GetName(StaticCastSharedRef<const FSkeletalMeshCurveSeries>(GraphTrackEvent.GetSeries())->CurveId);
+
+				Tooltip.AddNameValueTextLine(LOCTEXT("CurveName", "Curve").ToString(), CurveName);
+			}
+		}
+
+		Tooltip.AddNameValueTextLine(LOCTEXT("EventTime", "Time").ToString(), FText::AsNumber(InFoundStartTime).ToString());
+		Tooltip.AddNameValueTextLine(LOCTEXT("EventValue", "Value").ToString(), FText::AsNumber(GraphTrackEvent.GetValue()).ToString());
+
+		Tooltip.UpdateLayout();
+	});
+}
+
+const TSharedPtr<const ITimingEvent> FSkeletalMeshCurvesTrack::SearchEvent(const FTimingEventSearchParameters& InSearchParameters) const
+{
+	TSharedPtr<const ITimingEvent> FoundEvent;
+
+	FindSkeletalMeshPoseMessage(InSearchParameters, [this, &FoundEvent](double InFoundStartTime, double InFoundEndTime, uint32 InFoundDepth, const FSkeletalMeshPoseMessage& InFoundMessage)
+	{
+		FoundEvent = MakeShared<const FTimingEvent>(SharedThis(this), InFoundStartTime, InFoundEndTime, InFoundDepth);
+	});
+
+	return FoundEvent;
+}
+
+void FSkeletalMeshCurvesTrack::FindSkeletalMeshPoseMessage(const FTimingEventSearchParameters& InParameters, TFunctionRef<void(double, double, uint32, const FSkeletalMeshPoseMessage&)> InFoundPredicate) const
+{
+	TTimingEventSearch<FSkeletalMeshPoseMessage>::Search(
+		InParameters,
+
+		[this](TTimingEventSearch<FSkeletalMeshPoseMessage>::FContext& InContext)
+		{
+			const FAnimationProvider* AnimationProvider = SharedData.GetAnalysisSession().ReadProvider<FAnimationProvider>(FAnimationProvider::ProviderName);
+
+			if(AnimationProvider)
+			{
+				Trace::FAnalysisSessionReadScope SessionReadScope(SharedData.GetAnalysisSession());
+
+				AnimationProvider->ReadSkeletalMeshPoseTimeline(GetGameplayTrack().GetObjectId(), [&InContext](const FAnimationProvider::SkeletalMeshPoseTimeline& InTimeline, bool bInHasCurves)
+				{
+					InTimeline.EnumerateEvents(InContext.GetParameters().StartTime, InContext.GetParameters().EndTime, [&InContext](double InEventStartTime, double InEventEndTime, uint32 InDepth, const FSkeletalMeshPoseMessage& InMessage)
+					{
+						InContext.Check(InEventStartTime, InEventEndTime, 0, InMessage);
+					});
+				});
+			}
+		},
+
+		[&InFoundPredicate](double InFoundStartTime, double InFoundEndTime, uint32 InFoundDepth, const FSkeletalMeshPoseMessage& InEvent)
+		{
+			InFoundPredicate(InFoundStartTime, InFoundEndTime, InFoundDepth, InEvent);
+		});
+}
+
+#undef LOCTEXT_NAMESPACE
\ No newline at end of file
diff --git a/Engine/Plugins/Animation/GameplayInsights/Source/GameplayInsights/Private/SkeletalMeshCurvesTrack.h b/Engine/Plugins/Animation/GameplayInsights/Source/GameplayInsights/Private/SkeletalMeshCurvesTrack.h
new file mode 100644
--- /dev/null
+++ b/Engine/Plugins/Animation/GameplayInsights/Source/GameplayInsights/Private/SkeletalMeshCurvesTrack.h
@@ -0,0 +1,38 @@
+// Copyright 1998-2019 Epic Games, Inc. All Rights Reserved.
+
+#pragma once
+
+#include "GameplayGraphTrack.h"
+
+class FAnimationSharedData;
+class FTimingEventSearchParameters;
+struct FSkeletalMeshPoseMessage;
+class FSkeletalMeshCurvesTrack;
+
+class FSkeletalMeshCurveSeries : public FGameplayGraphSeries
+{
+public:
+	uint32 CurveId;
+};
+
+class FSkeletalMeshCurvesTrack : public FGameplayGraphTrack
+{
+	INSIGHTS_DECLARE_RTTI(FSkeletalMeshCurvesTrack, FGameplayGraphTrack)
+
+public:
+	FSkeletalMeshCurvesTrack(const FAnimationSharedData& InSharedData, uint64 InObjectID, const TCHAR* InName);
+
+	virtual void InitTooltip(FTooltipDrawState& Tooltip, const ITimingEvent& HoveredTimingEvent) const override;
+	virtual const TSharedPtr<const ITimingEvent> SearchEvent(const FTimingEventSearchParameters& InSearchParameters) const override;
+	virtual void AddAllSeries() override;
+	virtual bool UpdateSeriesBounds(FGameplayGraphSeries& InSeries, const FTimingTrackViewport& InViewport) override;
+	virtual void UpdateSeries(FGameplayGraphSeries& InSeries, const FTimingTrackViewport& InViewport) override;
+
+private:
+	// Helper function used to find a skeletal mesh pose
+	void FindSkeletalMeshPoseMessage(const FTimingEventSearchParameters& InParameters, TFunctionRef<void(double, double, uint32, const FSkeletalMeshPoseMessage&)> InFoundPredicate) const;
+
+private:
+	/** The shared data */
+	const FAnimationSharedData& SharedData;
+};
\ No newline at end of file
diff --git a/Engine/Plugins/Animation/GameplayInsights/Source/GameplayInsights/Private/SkeletalMeshPoseTrack.cpp b/Engine/Plugins/Animation/GameplayInsights/Source/GameplayInsights/Private/SkeletalMeshPoseTrack.cpp
--- a/Engine/Plugins/Animation/GameplayInsights/Source/GameplayInsights/Private/SkeletalMeshPoseTrack.cpp
+++ b/Engine/Plugins/Animation/GameplayInsights/Source/GameplayInsights/Private/SkeletalMeshPoseTrack.cpp
@@ -69,7 +69,7 @@ void FSkeletalMeshPoseTrack::BuildDrawState(ITimingEventsTrackDrawStateBuilder&
 	{
 		Trace::FAnalysisSessionReadScope SessionReadScope(SharedData.GetAnalysisSession());
 
-		AnimationProvider->ReadSkeletalMeshPoseTimeline(GetGameplayTrack().GetObjectId(), [&Context, &Builder](const FAnimationProvider::SkeletalMeshPoseTimeline& InTimeline)
+		AnimationProvider->ReadSkeletalMeshPoseTimeline(GetGameplayTrack().GetObjectId(), [&Context, &Builder](const FAnimationProvider::SkeletalMeshPoseTimeline& InTimeline, bool bInHasCurves)
 		{
 			InTimeline.EnumerateEvents(Context.GetViewport().GetStartTime(), Context.GetViewport().GetEndTime(), [&Builder](double InStartTime, double InEndTime, uint32 InDepth, const FSkeletalMeshPoseMessage& InMessage)
 			{
@@ -148,7 +148,7 @@ void FSkeletalMeshPoseTrack::FindSkeletalMeshPoseMessage(const FTimingEventSearc
 			{
 				Trace::FAnalysisSessionReadScope SessionReadScope(SharedData.GetAnalysisSession());
 
-				AnimationProvider->ReadSkeletalMeshPoseTimeline(GetGameplayTrack().GetObjectId(), [&InContext](const FAnimationProvider::SkeletalMeshPoseTimeline& InTimeline)
+				AnimationProvider->ReadSkeletalMeshPoseTimeline(GetGameplayTrack().GetObjectId(), [&InContext](const FAnimationProvider::SkeletalMeshPoseTimeline& InTimeline, bool bInHasCurves)
 				{
 					InTimeline.EnumerateEvents(InContext.GetParameters().StartTime, InContext.GetParameters().EndTime, [&InContext](double InEventStartTime, double InEventEndTime, uint32 InDepth, const FSkeletalMeshPoseMessage& InMessage)
 					{
@@ -362,7 +362,7 @@ void FSkeletalMeshPoseTrack::DrawPoses(UWorld* InWorld, double InTime)
 				CacheForWorld.Component->SetVisibility(false);
 			}
 
-			AnimationProvider->ReadSkeletalMeshPoseTimeline(GetGameplayTrack().GetObjectId(), [this, &CacheForWorld, &AnimationProvider, &GameplayProvider, &InTime](const FAnimationProvider::SkeletalMeshPoseTimeline& InTimeline)
+			AnimationProvider->ReadSkeletalMeshPoseTimeline(GetGameplayTrack().GetObjectId(), [this, &CacheForWorld, &AnimationProvider, &GameplayProvider, &InTime](const FAnimationProvider::SkeletalMeshPoseTimeline& InTimeline, bool bInHasCurves)
 			{
 				InTimeline.EnumerateEvents(InTime, InTime, [this, &CacheForWorld, &AnimationProvider, &GameplayProvider, &InTime](double InStartTime, double InEndTime, uint32 InDepth, const FSkeletalMeshPoseMessage& InMessage)
 				{
