Added missing include.

#rnx

diff --git a/Engine/Plugins/Animation/GameplayInsights/Source/GameplayInsights/Private/GameplayTimingViewExtender.cpp b/Engine/Plugins/Animation/GameplayInsights/Source/GameplayInsights/Private/GameplayTimingViewExtender.cpp
--- a/Engine/Plugins/Animation/GameplayInsights/Source/GameplayInsights/Private/GameplayTimingViewExtender.cpp
+++ b/Engine/Plugins/Animation/GameplayInsights/Source/GameplayInsights/Private/GameplayTimingViewExtender.cpp
@@ -5,6 +5,7 @@
 #include "GameplaySharedData.h"
 #include "AnimationSharedData.h"
 #include "UObject/WeakObjectPtr.h"
+#include "Framework/MultiBox/MultiBoxBuilder.h"
 
 #if WITH_ENGINE
 #include "Engine/World.h"
@@ -119,4 +120,4 @@ void FGameplayTimingViewExtender::GetCustomDebugObjects(const IAnimationBlueprin
 
 #endif
 
-#undef LOCTEXT_NAMESPACE
\ No newline at end of file
+#undef LOCTEXT_NAMESPACE
