Fixed some copyright lines now that the year is no longer required.
#rnx

diff --git a/Engine/Plugins/Animation/GameplayInsights/Source/GameplayInsights/Private/GameplayGraphTrack.cpp b/Engine/Plugins/Animation/GameplayInsights/Source/GameplayInsights/Private/GameplayGraphTrack.cpp
--- a/Engine/Plugins/Animation/GameplayInsights/Source/GameplayInsights/Private/GameplayGraphTrack.cpp
+++ b/Engine/Plugins/Animation/GameplayInsights/Source/GameplayInsights/Private/GameplayGraphTrack.cpp
@@ -1,4 +1,4 @@
-// Copyright 1998-2019 Epic Games, Inc. All Rights Reserved.
+// Copyright Epic Games, Inc. All Rights Reserved.
 
 #include "GameplayGraphTrack.h"
 #include "GameplayProvider.h"
@@ -316,4 +316,4 @@ void FGameplayGraphTrack::BuildContextMenu(FMenuBuilder& MenuBuilder)
 
 PRAGMA_ENABLE_OPTIMIZATION
 
-#undef LOCTEXT_NAMESPACE
\ No newline at end of file
+#undef LOCTEXT_NAMESPACE
diff --git a/Engine/Plugins/Animation/GameplayInsights/Source/GameplayInsights/Private/GameplayGraphTrack.h b/Engine/Plugins/Animation/GameplayInsights/Source/GameplayInsights/Private/GameplayGraphTrack.h
--- a/Engine/Plugins/Animation/GameplayInsights/Source/GameplayInsights/Private/GameplayGraphTrack.h
+++ b/Engine/Plugins/Animation/GameplayInsights/Source/GameplayInsights/Private/GameplayGraphTrack.h
@@ -1,4 +1,4 @@
-// Copyright 1998-2019 Epic Games, Inc. All Rights Reserved.
+// Copyright Epic Games, Inc. All Rights Reserved.
 
 #pragma once
 
@@ -91,4 +91,4 @@ protected:
 
 	/** Whether to draw labels */
 	bool bDrawLabels;
-};
\ No newline at end of file
+};
diff --git a/Engine/Plugins/Animation/GameplayInsights/Source/GameplayInsights/Private/SGameplayTrackTree.cpp b/Engine/Plugins/Animation/GameplayInsights/Source/GameplayInsights/Private/SGameplayTrackTree.cpp
--- a/Engine/Plugins/Animation/GameplayInsights/Source/GameplayInsights/Private/SGameplayTrackTree.cpp
+++ b/Engine/Plugins/Animation/GameplayInsights/Source/GameplayInsights/Private/SGameplayTrackTree.cpp
@@ -1,4 +1,4 @@
-// Copyright 1998-2019 Epic Games, Inc. All Rights Reserved.
+// Copyright Epic Games, Inc. All Rights Reserved.
 
 #include "SGameplayTrackTree.h"
 
diff --git a/Engine/Plugins/Animation/GameplayInsights/Source/GameplayInsights/Private/SGameplayTrackTree.h b/Engine/Plugins/Animation/GameplayInsights/Source/GameplayInsights/Private/SGameplayTrackTree.h
--- a/Engine/Plugins/Animation/GameplayInsights/Source/GameplayInsights/Private/SGameplayTrackTree.h
+++ b/Engine/Plugins/Animation/GameplayInsights/Source/GameplayInsights/Private/SGameplayTrackTree.h
@@ -1,4 +1,4 @@
-// Copyright 1998-2019 Epic Games, Inc. All Rights Reserved.
+// Copyright Epic Games, Inc. All Rights Reserved.
 
 #pragma once
 
diff --git a/Engine/Plugins/Animation/GameplayInsights/Source/GameplayInsights/Private/SkeletalMeshCurvesTrack.cpp b/Engine/Plugins/Animation/GameplayInsights/Source/GameplayInsights/Private/SkeletalMeshCurvesTrack.cpp
--- a/Engine/Plugins/Animation/GameplayInsights/Source/GameplayInsights/Private/SkeletalMeshCurvesTrack.cpp
+++ b/Engine/Plugins/Animation/GameplayInsights/Source/GameplayInsights/Private/SkeletalMeshCurvesTrack.cpp
@@ -1,4 +1,4 @@
-// Copyright 1998-2019 Epic Games, Inc. All Rights Reserved.
+// Copyright Epic Games, Inc. All Rights Reserved.
 
 #include "SkeletalMeshCurvesTrack.h"
 #include "GameplayProvider.h"
@@ -206,4 +206,4 @@ void FSkeletalMeshCurvesTrack::FindSkeletalMeshPoseMessage(const FTimingEventSea
 		});
 }
 
-#undef LOCTEXT_NAMESPACE
\ No newline at end of file
+#undef LOCTEXT_NAMESPACE
diff --git a/Engine/Plugins/Animation/GameplayInsights/Source/GameplayInsights/Private/SkeletalMeshCurvesTrack.h b/Engine/Plugins/Animation/GameplayInsights/Source/GameplayInsights/Private/SkeletalMeshCurvesTrack.h
--- a/Engine/Plugins/Animation/GameplayInsights/Source/GameplayInsights/Private/SkeletalMeshCurvesTrack.h
+++ b/Engine/Plugins/Animation/GameplayInsights/Source/GameplayInsights/Private/SkeletalMeshCurvesTrack.h
@@ -1,4 +1,4 @@
-// Copyright 1998-2019 Epic Games, Inc. All Rights Reserved.
+// Copyright Epic Games, Inc. All Rights Reserved.
 
 #pragma once
 
@@ -35,4 +35,4 @@ private:
 private:
 	/** The shared data */
 	const FAnimationSharedData& SharedData;
-};
\ No newline at end of file
+};
