Added schematic view of animation graphs

Initial basic implementation, more features to follow.
Accessed via the context menu of a graph track.
Shows anim graph nodes in a linearized tree similar to ShowDebug ANIMATION.
Select a node to 'pin' all the properties that it traced (multi select supported) to the bottom of the window.
Scrub the timeline to see those properties change over time.
Asset references use clickable hyperlinks to go to the specified asset editor.

diff --git a/Engine/Plugins/Animation/GameplayInsights/Source/GameplayInsights/GameplayInsights.Build.cs b/Engine/Plugins/Animation/GameplayInsights/Source/GameplayInsights/GameplayInsights.Build.cs
--- a/Engine/Plugins/Animation/GameplayInsights/Source/GameplayInsights/GameplayInsights.Build.cs
+++ b/Engine/Plugins/Animation/GameplayInsights/Source/GameplayInsights/GameplayInsights.Build.cs
@@ -17,6 +17,7 @@ namespace UnrealBuildTool.Rules
 				"TraceServices",
 				"TraceInsights",
 				"AssetRegistry",
+				"ApplicationCore",
 			});
 
 			if (Target.bCompileAgainstEngine)
diff --git a/Engine/Plugins/Animation/GameplayInsights/Source/GameplayInsights/Private/AnimNodesTrack.cpp b/Engine/Plugins/Animation/GameplayInsights/Source/GameplayInsights/Private/AnimNodesTrack.cpp
--- a/Engine/Plugins/Animation/GameplayInsights/Source/GameplayInsights/Private/AnimNodesTrack.cpp
+++ b/Engine/Plugins/Animation/GameplayInsights/Source/GameplayInsights/Private/AnimNodesTrack.cpp
@@ -53,7 +53,7 @@ FAnimNodesTrack::FAnimNodesTrack(const FAnimationSharedData& InSharedData, uint6
 			const FClassInfo* AnimInstanceClassInfo = GameplayProvider->FindClassInfo(AnimInstanceInfo->ClassId);
 			if(AnimInstanceClassInfo)
 			{
-				InstanceClass = FindObject<UAnimBlueprintGeneratedClass>(ANY_PACKAGE, AnimInstanceClassInfo->PathName);
+				InstanceClass = FSoftObjectPath(AnimInstanceClassInfo->PathName);
 			}
 		}
 	}
@@ -171,9 +171,9 @@ void FAnimNodesTrack::FindAnimGraphMessage(const FTimingEventSearchParameters& I
 
 void FAnimNodesTrack::BuildContextMenu(FMenuBuilder& MenuBuilder)
 {
-#if WITH_EDITOR
 	MenuBuilder.BeginSection(TEXT("DebugSection"), LOCTEXT("Debug", "Debug"));
 	{
+#if WITH_EDITOR
 		MenuBuilder.AddMenuEntry(
 			LOCTEXT("ToggleDebug", "Debug this graph"),
 			LOCTEXT("ToggleDebug_Tooltip", "Debug this graph in the animation blueprint editor, opens editor for asset if it exists"),
@@ -181,7 +181,7 @@ void FAnimNodesTrack::BuildContextMenu(FMenuBuilder& MenuBuilder)
 			FUIAction(
 				FExecuteAction::CreateLambda([this]()
 				{  
-					if(InstanceClass.Get())
+					if(InstanceClass.LoadSynchronous())
 					{
 						if(UAnimBlueprint* AnimBlueprint = Cast<UAnimBlueprint>(InstanceClass->ClassGeneratedBy))
 						{
@@ -216,7 +216,7 @@ void FAnimNodesTrack::BuildContextMenu(FMenuBuilder& MenuBuilder)
 				{
 					if(InstanceClass.Get())
 					{
-						if(UAnimBlueprint* AnimBlueprint = Cast<UAnimBlueprint>(InstanceClass->ClassGeneratedBy))
+						if(UAnimBlueprint* AnimBlueprint = Cast<UAnimBlueprint>(InstanceClass.Get()->ClassGeneratedBy))
 						{
 							const FGameplayProvider* GameplayProvider = SharedData.GetAnalysisSession().ReadProvider<FGameplayProvider>(FGameplayProvider::ProviderName);
 							if(GameplayProvider)
@@ -244,16 +244,27 @@ void FAnimNodesTrack::BuildContextMenu(FMenuBuilder& MenuBuilder)
 			NAME_None,
 			EUserInterfaceActionType::ToggleButton
 		);
+#endif
+		MenuBuilder.AddMenuEntry(
+			LOCTEXT("OpenAnimGraph", "View this graph"),
+			LOCTEXT("OpenAnimGraph_Tooltip", "Open this graph in the schematic anim graph viewer"),
+			FSlateIcon(),
+			FUIAction(
+				FExecuteAction::CreateLambda([this]()
+				{ 
+					SharedData.OpenAnimGraphTab(GetGameplayTrack().GetObjectId());
+				})
+			)
+		);
 	}
 	MenuBuilder.EndSection();
-#endif
 }
 
 #if WITH_EDITOR
 
 UAnimInstance* FAnimNodesTrack::LazyCreateAnimInstance(USkeletalMeshComponent* InComponent)
 {
-	if(InstanceClass.Get() != nullptr)
+	if(InstanceClass.LoadSynchronous() != nullptr)
 	{
 		if(AnimInstance == nullptr)
 		{
@@ -267,7 +278,7 @@ UAnimInstance* FAnimNodesTrack::LazyCreateAnimInstance(USkeletalMeshComponent* I
 
 void FAnimNodesTrack::UpdateDebugData(double InTime)
 {
-	if(InstanceClass.Get())
+	if(InstanceClass.LoadSynchronous())
 	{
 		const FAnimationProvider* AnimationProvider = SharedData.GetAnalysisSession().ReadProvider<FAnimationProvider>(FAnimationProvider::ProviderName);
 		const FGameplayProvider* GameplayProvider = SharedData.GetAnalysisSession().ReadProvider<FGameplayProvider>(FGameplayProvider::ProviderName);
@@ -281,7 +292,7 @@ void FAnimNodesTrack::UpdateDebugData(double InTime)
 				const Trace::IFrameProvider& FramesProvider = Trace::ReadFrameProvider(SharedData.GetAnalysisSession());
 
 
-				FAnimBlueprintDebugData& DebugData = InstanceClass->GetAnimBlueprintDebugData();
+				FAnimBlueprintDebugData& DebugData = InstanceClass.Get()->GetAnimBlueprintDebugData();
 				DebugData.ResetNodeVisitSites();
 
 				// round to nearest frame boundary
@@ -295,13 +306,16 @@ void FAnimNodesTrack::UpdateDebugData(double InTime)
 						{
 							// Basic verification - check node count is the same
 							// @TODO: could add some form of node hash/CRC to the class to improve this
-							if(InMessage.NodeCount == InstanceClass->AnimNodeProperties.Num())
+							if(InMessage.NodeCount == InstanceClass.Get()->AnimNodeProperties.Num())
 							{
 								AnimationProvider->ReadAnimNodesTimeline(GetGameplayTrack().GetObjectId(), [InGraphStartTime, InGraphEndTime, &DebugData](const FAnimationProvider::AnimNodesTimeline& InNodesTimeline)
 								{
 									InNodesTimeline.EnumerateEvents(InGraphStartTime, InGraphEndTime, [&DebugData](double InStartTime, double InEndTime, uint32 InDepth, const FAnimNodeMessage& InMessage)
 									{
-										DebugData.RecordNodeVisit(InMessage.NodeId, InMessage.PreviousNodeId, InMessage.Weight);
+										if(InMessage.Phase == EAnimGraphPhase::Update)
+										{
+											DebugData.RecordNodeVisit(InMessage.NodeId, InMessage.PreviousNodeId, InMessage.Weight);
+										}
 									});
 								});
 
@@ -313,44 +327,11 @@ void FAnimNodesTrack::UpdateDebugData(double InTime)
 									});
 								});
 
-								AnimationProvider->ReadAnimNodeValuesTimeline(GetGameplayTrack().GetObjectId(), [InGraphStartTime, InGraphEndTime, GameplayProvider, &DebugData](const FAnimationProvider::AnimNodeValuesTimeline& InNodeValuesTimeline)
+								AnimationProvider->ReadAnimNodeValuesTimeline(GetGameplayTrack().GetObjectId(), [InGraphStartTime, InGraphEndTime, AnimationProvider, &DebugData](const FAnimationProvider::AnimNodeValuesTimeline& InNodeValuesTimeline)
 								{
-									InNodeValuesTimeline.EnumerateEvents(InGraphStartTime, InGraphEndTime, [GameplayProvider, &DebugData](double InStartTime, double InEndTime, uint32 InDepth, const FAnimNodeValueMessage& InMessage)
+									InNodeValuesTimeline.EnumerateEvents(InGraphStartTime, InGraphEndTime, [AnimationProvider, &DebugData](double InStartTime, double InEndTime, uint32 InDepth, const FAnimNodeValueMessage& InMessage)
 									{
-										FText Text;
-										const FText KeyValueFormat(LOCTEXT("KeyValueFormat", "{0} = {1}"));
-
-										switch(InMessage.Type)
-										{
-										case EAnimNodeValueType::Bool:
-											Text = FText::Format(KeyValueFormat, FText::FromString(InMessage.Key), InMessage.Bool.bValue ? LOCTEXT("True", "true") : LOCTEXT("False", "false"));
-											break;
-										case EAnimNodeValueType::Int32:
-											Text = FText::Format(KeyValueFormat, FText::FromString(InMessage.Key), FText::AsNumber(InMessage.Int32.Value));
-											break;
-										case EAnimNodeValueType::Float:
-											Text = FText::Format(KeyValueFormat, FText::FromString(InMessage.Key), FText::AsNumber(InMessage.Float.Value));
-											break;
-										case EAnimNodeValueType::Vector:
-											Text = FText::Format(KeyValueFormat, FText::FromString(InMessage.Key), FText::Format(LOCTEXT("VectorFormat", "({0}, {1}, {2})"), FText::AsNumber(InMessage.Vector.Value.X), FText::AsNumber(InMessage.Vector.Value.Y), FText::AsNumber(InMessage.Vector.Value.Z)));
-											break;
-										case EAnimNodeValueType::String:
-											Text = FText::Format(KeyValueFormat, FText::FromString(InMessage.Key), FText::FromString(InMessage.String.Value));
-											break;
-										case EAnimNodeValueType::Object:
-										{
-											const FObjectInfo& ObjectInfo = GameplayProvider->GetObjectInfo(InMessage.Object.Value);
-											Text = FText::Format(KeyValueFormat, FText::FromString(InMessage.Key), FText::FromString(ObjectInfo.PathName));
-											break;
-										}
-										case EAnimNodeValueType::Class:
-										{
-											const FClassInfo& ClassInfo = GameplayProvider->GetClassInfo(InMessage.Class.Value);
-											Text = FText::Format(KeyValueFormat, FText::FromString(InMessage.Key), FText::FromString(ClassInfo.PathName));
-											break;
-										}
-										}
-
+										FText Text = AnimationProvider->FormatNodeKeyValue(InMessage);
 										DebugData.RecordNodeValue(InMessage.NodeId, Text.ToString());
 									});
 								});
@@ -388,9 +369,9 @@ void FAnimNodesTrack::UpdateDebugData(double InTime)
 
 void FAnimNodesTrack::GetCustomDebugObjects(const IAnimationBlueprintEditor& InAnimationBlueprintEditor, TArray<FCustomDebugObject>& OutDebugList)
 {
-	if(InstanceClass.Get())
+	if(InstanceClass.LoadSynchronous())
 	{
-		if(UAnimBlueprint* AnimBlueprint = Cast<UAnimBlueprint>(InstanceClass->ClassGeneratedBy))
+		if(UAnimBlueprint* AnimBlueprint = Cast<UAnimBlueprint>(InstanceClass.Get()->ClassGeneratedBy))
 		{
 			const FGameplayProvider* GameplayProvider = SharedData.GetAnalysisSession().ReadProvider<FGameplayProvider>(FGameplayProvider::ProviderName);
 			if(GameplayProvider)
@@ -419,4 +400,13 @@ void FAnimNodesTrack::GetCustomDebugObjects(const IAnimationBlueprintEditor& InA
 
 #endif
 
+#if WITH_ENGINE
+
+void FAnimNodesTrack::AddReferencedObjects(FReferenceCollector& Collector)
+{
+	Collector.AddReferencedObject(AnimInstance);
+}
+
+#endif
+
 #undef LOCTEXT_NAMESPACE
\ No newline at end of file
diff --git a/Engine/Plugins/Animation/GameplayInsights/Source/GameplayInsights/Private/AnimNodesTrack.h b/Engine/Plugins/Animation/GameplayInsights/Source/GameplayInsights/Private/AnimNodesTrack.h
--- a/Engine/Plugins/Animation/GameplayInsights/Source/GameplayInsights/Private/AnimNodesTrack.h
+++ b/Engine/Plugins/Animation/GameplayInsights/Source/GameplayInsights/Private/AnimNodesTrack.h
@@ -19,6 +19,9 @@ class USkeletalMeshComponent;
 class UAnimInstance;
 
 class FAnimNodesTrack : public TGameplayTrackMixin<FTimingEventsTrack>
+#if WITH_ENGINE
+	, public FGCObject
+#endif
 {
 public:
 	static const FName TypeName;
@@ -40,6 +43,12 @@ public:
 	void GetCustomDebugObjects(const IAnimationBlueprintEditor& InAnimationBlueprintEditor, TArray<FCustomDebugObject>& OutDebugList);
 #endif
 
+#if WITH_ENGINE
+	// FGCObject interface
+	virtual void AddReferencedObjects(FReferenceCollector& Collector) override;
+	virtual FString GetReferencerName() const override { return TEXT("InsightsAnimNodesTrack"); }
+#endif
+
 private:
 	// Helper function used to find an anim graph message
 	void FindAnimGraphMessage(const FTimingEventSearchParameters& InParameters, TFunctionRef<void(double, double, uint32, const FAnimGraphMessage&)> InFoundPredicate) const;
@@ -54,7 +63,7 @@ private:
 
 #if WITH_EDITOR
 	/** Instance class, if any, used for instantiating debug info */
-	TWeakObjectPtr<UAnimBlueprintGeneratedClass> InstanceClass;
+	TSoftObjectPtr<UAnimBlueprintGeneratedClass> InstanceClass;
 
 	/** Data used for anim BP debugging */
 	UAnimInstance* AnimInstance;
diff --git a/Engine/Plugins/Animation/GameplayInsights/Source/GameplayInsights/Private/AnimationProvider.cpp b/Engine/Plugins/Animation/GameplayInsights/Source/GameplayInsights/Private/AnimationProvider.cpp
--- a/Engine/Plugins/Animation/GameplayInsights/Source/GameplayInsights/Private/AnimationProvider.cpp
+++ b/Engine/Plugins/Animation/GameplayInsights/Source/GameplayInsights/Private/AnimationProvider.cpp
@@ -259,6 +259,49 @@ const TCHAR* FAnimationProvider::GetName(uint32 InId) const
 	return *UnknownText.ToString();
 }
 
+FText FAnimationProvider::FormatNodeKeyValue(const FAnimNodeValueMessage& InMessage) const
+{
+	return FText::Format(LOCTEXT("KeyValueFormat", "{0} = {1}"), FText::FromString(InMessage.Key), FormatNodeValue(InMessage));
+}
+
+FText FAnimationProvider::FormatNodeValue(const FAnimNodeValueMessage& InMessage) const
+{
+	FText Text;
+
+	switch(InMessage.Type)
+	{
+	case EAnimNodeValueType::Bool:
+		Text = InMessage.Bool.bValue ? LOCTEXT("True", "true") : LOCTEXT("False", "false");
+		break;
+	case EAnimNodeValueType::Int32:
+		Text =  FText::AsNumber(InMessage.Int32.Value);
+		break;
+	case EAnimNodeValueType::Float:
+		Text = FText::AsNumber(InMessage.Float.Value);
+		break;
+	case EAnimNodeValueType::Vector:
+		Text = FText::Format(LOCTEXT("VectorFormat", "({0}, {1}, {2})"), FText::AsNumber(InMessage.Vector.Value.X), FText::AsNumber(InMessage.Vector.Value.Y), FText::AsNumber(InMessage.Vector.Value.Z));
+		break;
+	case EAnimNodeValueType::String:
+		Text = FText::FromString(InMessage.String.Value);
+		break;
+	case EAnimNodeValueType::Object:
+	{
+		const FObjectInfo& ObjectInfo = GameplayProvider.GetObjectInfo(InMessage.Object.Value);
+		Text = FText::FromString(ObjectInfo.PathName);
+		break;
+	}
+	case EAnimNodeValueType::Class:
+	{
+		const FClassInfo& ClassInfo = GameplayProvider.GetClassInfo(InMessage.Class.Value);
+		Text = FText::FromString(ClassInfo.PathName);
+		break;
+	}
+	}
+
+	return Text;
+}
+
 void FAnimationProvider::AppendTickRecord(uint64 InAnimInstanceId, double InTime, uint64 InAssetId, int32 InNodeId, float InBlendWeight, float InPlaybackTime, float InRootMotionWeight, float InPlayRate, float InBlendSpacePositionX, float InBlendSpacePositionY, uint16 InFrameCounter, bool bInLooping, bool bInIsBlendSpace)
 {
 	Session.WriteAccessCheck();
diff --git a/Engine/Plugins/Animation/GameplayInsights/Source/GameplayInsights/Private/AnimationProvider.h b/Engine/Plugins/Animation/GameplayInsights/Source/GameplayInsights/Private/AnimationProvider.h
--- a/Engine/Plugins/Animation/GameplayInsights/Source/GameplayInsights/Private/AnimationProvider.h
+++ b/Engine/Plugins/Animation/GameplayInsights/Source/GameplayInsights/Private/AnimationProvider.h
@@ -36,6 +36,8 @@ public:
 	virtual bool ReadStateMachinesTimeline(uint64 InObjectId, TFunctionRef<void(const StateMachinesTimeline&)> Callback) const override;
 	virtual const FSkeletalMeshInfo* FindSkeletalMeshInfo(uint64 InObjectId) const override;
 	virtual const TCHAR* GetName(uint32 InId) const override;
+	virtual FText FormatNodeKeyValue(const FAnimNodeValueMessage& InMessage) const override;
+	virtual FText FormatNodeValue(const FAnimNodeValueMessage& InMessage) const override;
 
 	/** Add a tick record */
 	void AppendTickRecord(uint64 InAnimInstanceId, double InTime, uint64 InAssetId, int32 InNodeId, float InBlendWeight, float InPlaybackTime, float InRootMotionWeight, float InPlayRate, float InBlendSpacePositionX, float InBlendSpacePositionY, uint16 InFrameCounter, bool bInLooping, bool bInIsBlendSpace);
diff --git a/Engine/Plugins/Animation/GameplayInsights/Source/GameplayInsights/Private/AnimationSharedData.cpp b/Engine/Plugins/Animation/GameplayInsights/Source/GameplayInsights/Private/AnimationSharedData.cpp
--- a/Engine/Plugins/Animation/GameplayInsights/Source/GameplayInsights/Private/AnimationSharedData.cpp
+++ b/Engine/Plugins/Animation/GameplayInsights/Source/GameplayInsights/Private/AnimationSharedData.cpp
@@ -19,6 +19,8 @@
 #include "EditorViewportClient.h"
 #include "Editor/EditorEngine.h"
 #endif
+#include "SAnimGraphSchematicView.h"
+#include "GameplayInsightsModule.h"
 
 #define LOCTEXT_NAMESPACE "AnimationSharedData"
 
@@ -34,8 +36,21 @@ FAnimationSharedData::FAnimationSharedData(FGameplaySharedData& InGameplayShared
 {
 }
 
+FAnimationSharedData::~FAnimationSharedData()
+{
+	for(TWeakPtr<SDockTab> DocumentTab : OpenDocumentTabs)
+	{
+		if(DocumentTab.IsValid())
+		{
+			DocumentTab.Pin()->RequestCloseTab();
+		}
+	}
+}
+
 void FAnimationSharedData::OnBeginSession(Insights::ITimingViewSession& InTimingViewSession)
 {
+	TimingViewSession = &InTimingViewSession;
+
 	SkeletalMeshPoseTracks.Reset();
 	AnimationTickRecordsTracks.Reset();
 
@@ -48,6 +63,8 @@ void FAnimationSharedData::OnEndSession(Insights::ITimingViewSession& InTimingVi
 	AnimationTickRecordsTracks.Reset();
 
 	InTimingViewSession.OnTimeMarkerChanged().Remove(TimeMarkerChangedHandle);
+
+	TimingViewSession = nullptr;
 }
 
 void FAnimationSharedData::Tick(Insights::ITimingViewSession& InTimingViewSession, const Trace::IAnalysisSession& InAnalysisSession)
@@ -439,4 +456,48 @@ void FAnimationSharedData::GetCustomDebugObjects(const IAnimationBlueprintEditor
 
 #endif
 
+/** Search predicate for an open anim graph document */
+struct FSearchForAnimInstanceTab : public FTabManager::FSearchPreference
+{
+	FSearchForAnimInstanceTab(uint64 InAnimInstanceId)
+		: AnimInstanceId(InAnimInstanceId)
+	{
+	}
+
+	virtual TSharedPtr<SDockTab> Search(const FTabManager& Manager, FName PlaceholderId, const TSharedRef<SDockTab>& UnmanagedTab) const override
+	{
+		TSharedRef<SAnimGraphSchematicView> Content = StaticCastSharedRef<SAnimGraphSchematicView>(UnmanagedTab->GetContent());
+		if(Content->GetAnimInstanceId() == AnimInstanceId)
+		{
+			return UnmanagedTab;
+		}
+
+		return TSharedPtr<SDockTab>();
+	}
+
+	uint64 AnimInstanceId;
+};
+
+void FAnimationSharedData::OpenAnimGraphTab(uint64 InAnimInstanceId) const
+{
+	if(TimingViewSession && AnalysisSession)
+	{
+		FGameplayInsightsModule& GameplayInsightsModule = FModuleManager::GetModuleChecked<FGameplayInsightsModule>("GameplayInsights");
+		TSharedRef<SDockTab> Tab = GameplayInsightsModule.SpawnTimingProfilerDocumentTab(FSearchForAnimInstanceTab(InAnimInstanceId));
+
+		Tab->SetContent(SNew(SAnimGraphSchematicView, InAnimInstanceId, *TimingViewSession, *AnalysisSession));
+
+		const FGameplayProvider* GameplayProvider = AnalysisSession->ReadProvider<FGameplayProvider>(FGameplayProvider::ProviderName);
+		if(GameplayProvider)
+		{
+			Trace::FAnalysisSessionReadScope SessionReadScope(*AnalysisSession);
+
+			const FObjectInfo& ObjectInfo = GameplayProvider->GetObjectInfo(InAnimInstanceId);
+			Tab->SetLabel(FText::FromString(ObjectInfo.Name));
+		}
+
+		OpenDocumentTabs.Add(Tab);
+	}
+}
+
 #undef LOCTEXT_NAMESPACE
\ No newline at end of file
diff --git a/Engine/Plugins/Animation/GameplayInsights/Source/GameplayInsights/Private/AnimationSharedData.h b/Engine/Plugins/Animation/GameplayInsights/Source/GameplayInsights/Private/AnimationSharedData.h
--- a/Engine/Plugins/Animation/GameplayInsights/Source/GameplayInsights/Private/AnimationSharedData.h
+++ b/Engine/Plugins/Animation/GameplayInsights/Source/GameplayInsights/Private/AnimationSharedData.h
@@ -22,6 +22,8 @@ class FAnimationSharedData
 public:
 	FAnimationSharedData(FGameplaySharedData& InGameplaySharedData);
 
+	~FAnimationSharedData();
+
 	void OnBeginSession(Insights::ITimingViewSession& InTimingViewSession);
 	void OnEndSession(Insights::ITimingViewSession& InTimingViewSession);
 	void Tick(Insights::ITimingViewSession& InTimingViewSession, const Trace::IAnalysisSession& InAnalysisSession);
@@ -63,6 +65,9 @@ public:
 	// Find an anim nodes track with the specified anim instance ID
 	TSharedPtr<FAnimNodesTrack> FindAnimNodesTrack(uint64 InAnimInstanceId) const;
 
+	/** Open an anim graph tab to see its schematic view */
+	void OpenAnimGraphTab(uint64 InAnimInstanceId) const;
+
 private:
 	// UI handlers
 	void ToggleAnimationTracks();
@@ -79,6 +84,9 @@ private:
 	// Cached analysis session, set in Tick()
 	const Trace::IAnalysisSession* AnalysisSession;
 
+	// Cached timing view session, set in OnBeginSession/OnEndSession
+	Insights::ITimingViewSession* TimingViewSession;
+
 	// All the tracks we manage
 	TArray<TSharedRef<FSkeletalMeshPoseTrack>> SkeletalMeshPoseTracks;
 	TArray<TSharedRef<FSkeletalMeshCurvesTrack>> SkeletalMeshCurvesTracks;
@@ -91,6 +99,9 @@ private:
 	/** Various times and ranges */
 	double MarkerTime;
 
+	/** The open document tabs we reference */
+	mutable TArray<TWeakPtr<SDockTab>> OpenDocumentTabs;
+
 	/** Validity flags for pose times/ranges */
 	bool bTimeMarkerValid;
 
diff --git a/Engine/Plugins/Animation/GameplayInsights/Source/GameplayInsights/Private/AnimationTickRecordsTrack.cpp b/Engine/Plugins/Animation/GameplayInsights/Source/GameplayInsights/Private/AnimationTickRecordsTrack.cpp
--- a/Engine/Plugins/Animation/GameplayInsights/Source/GameplayInsights/Private/AnimationTickRecordsTrack.cpp
+++ b/Engine/Plugins/Animation/GameplayInsights/Source/GameplayInsights/Private/AnimationTickRecordsTrack.cpp
@@ -78,7 +78,7 @@ FAnimationTickRecordsTrack::FAnimationTickRecordsTrack(const FAnimationSharedDat
 			const FClassInfo* AnimInstanceClassInfo = GameplayProvider->FindClassInfo(AnimInstanceInfo->ClassId);
 			if(AnimInstanceClassInfo)
 			{
-				InstanceClass = FindObject<UAnimBlueprintGeneratedClass>(ANY_PACKAGE, AnimInstanceClassInfo->PathName);
+				InstanceClass = FSoftObjectPath(AnimInstanceClassInfo->PathName);
 			}
 		}
 	}
@@ -367,16 +367,16 @@ void FAnimationTickRecordsTrack::BuildContextMenu(FMenuBuilder& MenuBuilder)
 			FUIAction(
 				FExecuteAction::CreateLambda([this]()
 				{
-					if(InstanceClass)
+					if(InstanceClass.LoadSynchronous())
 					{
-						if(UAnimBlueprint* AnimBlueprint = Cast<UAnimBlueprint>(InstanceClass->ClassGeneratedBy))
+						if(UAnimBlueprint* AnimBlueprint = Cast<UAnimBlueprint>(InstanceClass.Get()->ClassGeneratedBy))
 						{
 							GEditor->GetEditorSubsystem<UAssetEditorSubsystem>()->OpenEditorForAsset(AnimBlueprint);
 
 							if(IAnimationBlueprintEditor* AnimBlueprintEditor = static_cast<IAnimationBlueprintEditor*>(GEditor->GetEditorSubsystem<UAssetEditorSubsystem>()->FindEditorForAsset(AnimBlueprint, true)))
 							{
-								int32 AnimNodeIndex = InstanceClass->AnimNodeProperties.Num() - NodeId - 1;
-								TWeakObjectPtr<const UEdGraphNode> GraphNode = InstanceClass->AnimBlueprintDebugData.NodePropertyIndexToNodeMap[AnimNodeIndex];
+								int32 AnimNodeIndex = InstanceClass.Get()->AnimNodeProperties.Num() - NodeId - 1;
+								TWeakObjectPtr<const UEdGraphNode> GraphNode = InstanceClass.Get()->AnimBlueprintDebugData.NodePropertyIndexToNodeMap[AnimNodeIndex];
 								if(GraphNode.Get())
 								{
 									AnimBlueprintEditor->JumpToHyperlink(GraphNode.Get());
diff --git a/Engine/Plugins/Animation/GameplayInsights/Source/GameplayInsights/Private/AnimationTickRecordsTrack.h b/Engine/Plugins/Animation/GameplayInsights/Source/GameplayInsights/Private/AnimationTickRecordsTrack.h
--- a/Engine/Plugins/Animation/GameplayInsights/Source/GameplayInsights/Private/AnimationTickRecordsTrack.h
+++ b/Engine/Plugins/Animation/GameplayInsights/Source/GameplayInsights/Private/AnimationTickRecordsTrack.h
@@ -74,7 +74,7 @@ private:
 
 #if WITH_EDITOR
 	/** The class that output this tick record */
-	UAnimBlueprintGeneratedClass* InstanceClass;
+	TSoftObjectPtr<UAnimBlueprintGeneratedClass> InstanceClass;
 #endif
 
 	/** Colors for line drawing */
diff --git a/Engine/Plugins/Animation/GameplayInsights/Source/GameplayInsights/Private/GameplayInsightsModule.cpp b/Engine/Plugins/Animation/GameplayInsights/Source/GameplayInsights/Private/GameplayInsightsModule.cpp
--- a/Engine/Plugins/Animation/GameplayInsights/Source/GameplayInsights/Private/GameplayInsightsModule.cpp
+++ b/Engine/Plugins/Animation/GameplayInsights/Source/GameplayInsights/Private/GameplayInsightsModule.cpp
@@ -1,12 +1,18 @@
 // Copyright Epic Games, Inc. All Rights Reserved.
 
-#include "Modules/ModuleInterface.h"
+#include "GameplayInsightsModule.h"
 #include "Features/IModularFeatures.h"
-#include "GameplayTraceModule.h"
-#include "GameplayTimingViewExtender.h"
 #include "Insights/ITimingViewExtender.h"
 #include "Containers/Ticker.h"
 #include "Modules/ModuleManager.h"
+#include "Framework/Docking/LayoutExtender.h"
+#include "SAnimGraphSchematicView.h"
+#include "Insights/IUnrealInsightsModule.h"
+#include "HAL/PlatformApplicationMisc.h"
+#include "WorkspaceMenuStructure.h"
+#include "WorkspaceMenuStructureModule.h"
+#include "TraceServices/ITraceServicesModule.h"
+#include "TraceServices/SessionService.h"
 
 #if WITH_EDITOR
 #include "IAnimationBlueprintEditorModule.h"
@@ -18,55 +24,147 @@
 
 #define LOCTEXT_NAMESPACE "GameplayInsightsModule"
 
-class FGameplayInsightsModule : public IModuleInterface
+namespace GameplayInsightsTabs
 {
-public:
-	// IModuleInterface interface
-	virtual void StartupModule() override
-	{
-		IModularFeatures::Get().RegisterModularFeature(Trace::ModuleFeatureName, &GameplayTraceModule);
-		IModularFeatures::Get().RegisterModularFeature(Insights::TimingViewExtenderFeatureName, &GameplayTimingViewExtender);
+	static const FName AnimGraphSchematicView("AnimGraphSchematicView");
+};
 
-		TickerHandle = FTicker::GetCoreTicker().AddTicker(TEXT("GameplayInsights"), 0.0f, [this](float DeltaTime)
-		{
-			GameplayTimingViewExtender.TickVisualizers(DeltaTime);
-			return true;
-		});
+void FGameplayInsightsModule::StartupModule()
+{
+	IModularFeatures::Get().RegisterModularFeature(Trace::ModuleFeatureName, &GameplayTraceModule);
+	IModularFeatures::Get().RegisterModularFeature(Insights::TimingViewExtenderFeatureName, &GameplayTimingViewExtender);
+
+	TickerHandle = FTicker::GetCoreTicker().AddTicker(TEXT("GameplayInsights"), 0.0f, [this](float DeltaTime)
+	{
+		GameplayTimingViewExtender.TickVisualizers(DeltaTime);
+		return true;
+	});
 
 #if WITH_EDITOR
-		IAnimationBlueprintEditorModule& AnimationBlueprintEditorModule = FModuleManager::LoadModuleChecked<IAnimationBlueprintEditorModule>("AnimationBlueprintEditor");
-		CustomDebugObjectHandle = AnimationBlueprintEditorModule.OnGetCustomDebugObjects().AddLambda([this](const IAnimationBlueprintEditor& InAnimationBlueprintEditor, TArray<FCustomDebugObject>& OutDebugList)
+	IAnimationBlueprintEditorModule& AnimationBlueprintEditorModule = FModuleManager::LoadModuleChecked<IAnimationBlueprintEditorModule>("AnimationBlueprintEditor");
+	CustomDebugObjectHandle = AnimationBlueprintEditorModule.OnGetCustomDebugObjects().AddLambda([this](const IAnimationBlueprintEditor& InAnimationBlueprintEditor, TArray<FCustomDebugObject>& OutDebugList)
+	{
+		GameplayTimingViewExtender.GetCustomDebugObjects(InAnimationBlueprintEditor, OutDebugList);
+	});
+
+	ITraceServicesModule& TraceServicesModule = FModuleManager::LoadModuleChecked<ITraceServicesModule>("TraceServices");
+
+	// Connect to loopback session
+	TSharedPtr<Trace::ISessionService> SessionService = TraceServicesModule.GetSessionService();
+	if(SessionService.IsValid())
+	{
+		// Connect session
+		SessionService->ConnectSession(TEXT("127.0.0.1"));
+
+		// Wait a second then attempt a connection - it takes a bit of time for sessions to get populated
+		FTicker::GetCoreTicker().AddTicker(FTickerDelegate::CreateLambda([](float InDeltaSeconds)
 		{
-			GameplayTimingViewExtender.GetCustomDebugObjects(InAnimationBlueprintEditor, OutDebugList);
-		});
-#endif
+			IUnrealInsightsModule& UnrealInsightsModule = FModuleManager::LoadModuleChecked<IUnrealInsightsModule>("TraceInsights");
+
+			// Start analysis on the latest live session
+			UnrealInsightsModule.StartAnalysisForLastLiveSession();
+
+			return false;
+		}), 1.0f);
 	}
 
-	virtual void ShutdownModule() override
+	IUnrealInsightsModule& UnrealInsightsModule = FModuleManager::LoadModuleChecked<IUnrealInsightsModule>("TraceInsights");
+
+	const float DPIScaleFactor = FPlatformApplicationMisc::GetDPIScaleFactorAtPoint(10.0f, 10.0f);
+
+	TSharedRef<FTabManager::FLayout> MajorTabsLayout = FTabManager::NewLayout("GameplayInsightsMajorLayout_v1.0")
+	->AddArea
+	(
+		FTabManager::NewArea(1280.f * DPIScaleFactor, 720.0f * DPIScaleFactor)
+		->Split
+		(
+			FTabManager::NewStack()
+			->AddTab(FInsightsManagerTabs::TimingProfilerTabId, ETabState::ClosedTab)
+		)
+	);
+
+	FInsightsMajorTabConfig TimingProfilerConfig;
+	TimingProfilerConfig.TabLabel = LOCTEXT("GameplayInsightsTabName", "Gameplay Insights");
+	TimingProfilerConfig.TabTooltip = LOCTEXT("GameplayInsightsTabTooltip", "Open the Gameplay Insights tab.");
+	TimingProfilerConfig.Layout = FTabManager::NewLayout("GameplayInsightsTimingLayout_v1.0")
+	->AddArea
+	(
+		FTabManager::NewPrimaryArea()
+		->SetOrientation(Orient_Horizontal)
+		->Split
+		(
+			FTabManager::NewSplitter()
+			->SetOrientation(Orient_Vertical)
+			->SetSizeCoefficient(0.7f)
+			->Split
+			(
+				FTabManager::NewStack()
+				->SetSizeCoefficient(0.1f)
+				->SetHideTabWell(true)
+				->AddTab(FTimingProfilerTabs::FramesTrackID, ETabState::OpenedTab)
+			)
+			->Split
+			(
+				FTabManager::NewStack()
+				->SetSizeCoefficient(0.9f)
+				->SetHideTabWell(true)
+				->AddTab(FTimingProfilerTabs::TimingViewID, ETabState::OpenedTab)
+			)
+		)
+		->Split
+		(
+			FTabManager::NewStack()
+			->SetSizeCoefficient(0.3f)
+			->AddTab(GameplayInsightsTabs::AnimGraphSchematicView, ETabState::ClosedTab)
+		)
+	);
+
+	UnrealInsightsModule.OnMajorTabCreated().AddLambda([this](const FName& InMajorTabId, TSharedPtr<FTabManager> InTabManager)
 	{
-#if WITH_EDITOR
-		IAnimationBlueprintEditorModule* AnimationBlueprintEditorModule = FModuleManager::GetModulePtr<IAnimationBlueprintEditorModule>("AnimationBlueprintEditor");
-		if(AnimationBlueprintEditorModule)
+		if(InMajorTabId == FInsightsManagerTabs::TimingProfilerTabId)
 		{
-			AnimationBlueprintEditorModule->OnGetCustomDebugObjects().Remove(CustomDebugObjectHandle);
+			WeakTimingProfilerTabManager = InTabManager;
 		}
-#endif
-
-		FTicker::GetCoreTicker().RemoveTicker(TickerHandle);
+	});
 
-		IModularFeatures::Get().UnregisterModularFeature(Trace::ModuleFeatureName, &GameplayTraceModule);
-		IModularFeatures::Get().UnregisterModularFeature(Insights::TimingViewExtenderFeatureName, &GameplayTimingViewExtender);
-	}
+	TimingProfilerConfig.WorkspaceGroup = WorkspaceMenu::GetMenuStructure().GetDeveloperToolsProfilingCategory();
 
-	FGameplayTraceModule GameplayTraceModule;
-	FGameplayTimingViewExtender GameplayTimingViewExtender;
+	UnrealInsightsModule.RegisterMajorTabConfig(FInsightsManagerTabs::TimingProfilerTabId, TimingProfilerConfig);
+	UnrealInsightsModule.RegisterMajorTabConfig(FInsightsManagerTabs::StartPageTabId, FInsightsMajorTabConfig::Unavailable());
+	UnrealInsightsModule.RegisterMajorTabConfig(FInsightsManagerTabs::SessionInfoTabId, FInsightsMajorTabConfig::Unavailable());
+	UnrealInsightsModule.RegisterMajorTabConfig(FInsightsManagerTabs::LoadingProfilerTabId, FInsightsMajorTabConfig::Unavailable());
+	UnrealInsightsModule.RegisterMajorTabConfig(FInsightsManagerTabs::NetworkingProfilerTabId, FInsightsMajorTabConfig::Unavailable());
 
-	FDelegateHandle TickerHandle;
+	UnrealInsightsModule.CreateSessionViewer(false);
+#endif
+}
 
+void FGameplayInsightsModule::ShutdownModule()
+{
 #if WITH_EDITOR
-	FDelegateHandle CustomDebugObjectHandle;
+	IAnimationBlueprintEditorModule* AnimationBlueprintEditorModule = FModuleManager::GetModulePtr<IAnimationBlueprintEditorModule>("AnimationBlueprintEditor");
+	if(AnimationBlueprintEditorModule)
+	{
+		AnimationBlueprintEditorModule->OnGetCustomDebugObjects().Remove(CustomDebugObjectHandle);
+	}
 #endif
-};
+
+	FTicker::GetCoreTicker().RemoveTicker(TickerHandle);
+
+	IModularFeatures::Get().UnregisterModularFeature(Trace::ModuleFeatureName, &GameplayTraceModule);
+	IModularFeatures::Get().UnregisterModularFeature(Insights::TimingViewExtenderFeatureName, &GameplayTimingViewExtender);
+}
+
+TSharedRef<SDockTab> FGameplayInsightsModule::SpawnTimingProfilerDocumentTab(const FTabManager::FSearchPreference& InSearchPreference)
+{
+	TSharedRef<SDockTab> NewTab = SNew(SDockTab);
+	TSharedPtr<FTabManager> TimingProfilerTabManager = WeakTimingProfilerTabManager.Pin();
+	if(TimingProfilerTabManager.IsValid())
+	{
+		TimingProfilerTabManager->InsertNewDocumentTab(GameplayInsightsTabs::AnimGraphSchematicView, InSearchPreference, NewTab);
+	}
+	return NewTab;
+}
 
 IMPLEMENT_MODULE(FGameplayInsightsModule, GameplayInsights);
 
diff --git a/Engine/Plugins/Animation/GameplayInsights/Source/GameplayInsights/Private/GameplayInsightsModule.h b/Engine/Plugins/Animation/GameplayInsights/Source/GameplayInsights/Private/GameplayInsightsModule.h
new file mode 100644
--- /dev/null
+++ b/Engine/Plugins/Animation/GameplayInsights/Source/GameplayInsights/Private/GameplayInsightsModule.h
@@ -0,0 +1,33 @@
+// Copyright 1998-2019 Epic Games, Inc. All Rights Reserved.
+
+#pragma once
+
+#include "CoreMinimal.h"
+#include "Modules/ModuleInterface.h"
+#include "GameplayTraceModule.h"
+#include "GameplayTimingViewExtender.h"
+#include "Framework/Docking/TabManager.h"
+
+class FGameplayInsightsModule : public IModuleInterface
+{
+public:
+	// IModuleInterface interface
+	virtual void StartupModule() override;
+	virtual void ShutdownModule() override;
+
+	// Spawn a document tab
+	TSharedRef<SDockTab> SpawnTimingProfilerDocumentTab(const FTabManager::FSearchPreference& InSearchPreference);
+
+private:
+	FGameplayTraceModule GameplayTraceModule;
+
+	FGameplayTimingViewExtender GameplayTimingViewExtender;
+
+	FDelegateHandle TickerHandle;
+
+#if WITH_EDITOR
+	FDelegateHandle CustomDebugObjectHandle;
+#endif
+
+	TWeakPtr<FTabManager> WeakTimingProfilerTabManager;
+};
\ No newline at end of file
diff --git a/Engine/Plugins/Animation/GameplayInsights/Source/GameplayInsights/Private/GameplayInsightsStyle.h b/Engine/Plugins/Animation/GameplayInsights/Source/GameplayInsights/Private/GameplayInsightsStyle.h
new file mode 100644
--- /dev/null
+++ b/Engine/Plugins/Animation/GameplayInsights/Source/GameplayInsights/Private/GameplayInsightsStyle.h
@@ -0,0 +1,72 @@
+// Copyright 1998-2019 Epic Games, Inc. All Rights Reserved.
+
+#pragma once
+
+#include "Styling/SlateStyle.h"
+#include "Styling/SlateStyleRegistry.h"
+
+#define IMAGE_BRUSH(RelativePath, ...) FSlateImageBrush(RootToContentDir(RelativePath, TEXT(".png")), __VA_ARGS__)
+#define BOX_BRUSH(RelativePath, ...) FSlateBoxBrush(RootToContentDir(RelativePath, TEXT(".png")), __VA_ARGS__)
+#define BORDER_BRUSH(RelativePath, ...) FSlateBorderBrush(RootToContentDir(RelativePath, TEXT(".png")), __VA_ARGS__)
+#define TTF_FONT(RelativePath, ...) FSlateFontInfo(RootToContentDir(RelativePath, TEXT(".ttf")), __VA_ARGS__)
+#define OTF_FONT(RelativePath, ...) FSlateFontInfo(RootToContentDir(RelativePath, TEXT(".otf")), __VA_ARGS__)
+
+class FGameplayInsightsStyle
+	: public FSlateStyleSet
+{
+public:
+	FGameplayInsightsStyle()
+		: FSlateStyleSet("GameplayInsightsStyle")
+		, SelectionColor_LinearRef(MakeShareable(new FLinearColor( 0.728f, 0.364f, 0.003f)))
+		, SelectionColor_Pressed_LinearRef(MakeShareable(new FLinearColor( 0.701f, 0.225f, 0.003f)))
+		, SelectionColor(SelectionColor_LinearRef)
+		, SelectionColor_Pressed(SelectionColor_Pressed_LinearRef)
+	{
+		const FVector2D Icon10x10(10.0f, 10.0f);
+		const FVector2D Icon16x16(16.0f, 16.0f);
+		const FVector2D Icon20x20(20.0f, 20.0f);
+		const FVector2D Icon24x24(24.0f, 24.0f);
+		const FVector2D Icon32x32(32.0f, 32.0f);
+		const FVector2D Icon40x40(40.0f, 40.0f);
+		SetContentRoot(FPaths::EnginePluginsDir() / TEXT("Animation/GameplayInsights/Content"));
+
+		const FSlateColor DefaultForeground(FLinearColor(0.72f, 0.72f, 0.72f, 1.f));
+
+		Set("SchematicViewRootLeft", new BOX_BRUSH("SchematicViewRootLeft", FMargin(4/16.0f, 8.0f/16.0f, 4/16.0f, 4/16.0f)));
+		Set("SchematicViewRootMid", new BOX_BRUSH("SchematicViewRootMid", FMargin(4/16.0f, 8.0f/16.0f, 4/16.0f, 4/16.0f)));
+
+		Set("SchematicViewViewButton", FButtonStyle()
+			.SetNormal(FSlateNoResource())
+			.SetHovered(BOX_BRUSH("RoundedSelection_16x", 4.0f/16.0f, SelectionColor))
+			.SetPressed(BOX_BRUSH("RoundedSelection_16x", 4.0f/16.0f, SelectionColor_Pressed))
+			.SetNormalPadding(FMargin(2, 2, 2, 2))
+			.SetPressedPadding(FMargin(2, 3, 2, 1))
+		);
+
+		Set("SchematicViewViewButtonIcon", new IMAGE_BRUSH("view_button", Icon20x20));
+
+		FSlateStyleRegistry::RegisterSlateStyle(*this);
+	}
+
+	static FGameplayInsightsStyle& Get()
+	{
+		static FGameplayInsightsStyle Inst;
+		return Inst;
+	}
+	
+	~FGameplayInsightsStyle()
+	{
+		FSlateStyleRegistry::UnRegisterSlateStyle(*this);
+	}
+
+	const TSharedRef<FLinearColor> SelectionColor_LinearRef;
+	const TSharedRef<FLinearColor> SelectionColor_Pressed_LinearRef;
+	const FSlateColor SelectionColor;
+	const FSlateColor SelectionColor_Pressed;
+};
+
+#undef IMAGE_BRUSH
+#undef BOX_BRUSH
+#undef BORDER_BRUSH
+#undef TTF_FONT
+#undef OTF_FONT
diff --git a/Engine/Plugins/Animation/GameplayInsights/Source/GameplayInsights/Private/IAnimationProvider.h b/Engine/Plugins/Animation/GameplayInsights/Source/GameplayInsights/Private/IAnimationProvider.h
--- a/Engine/Plugins/Animation/GameplayInsights/Source/GameplayInsights/Private/IAnimationProvider.h
+++ b/Engine/Plugins/Animation/GameplayInsights/Source/GameplayInsights/Private/IAnimationProvider.h
@@ -204,4 +204,6 @@ public:
 	virtual bool ReadStateMachinesTimeline(uint64 InObjectId, TFunctionRef<void(const StateMachinesTimeline&)> Callback) const = 0;
 	virtual const FSkeletalMeshInfo* FindSkeletalMeshInfo(uint64 InObjectId) const = 0;
 	virtual const TCHAR* GetName(uint32 InId) const = 0;
+	virtual FText FormatNodeKeyValue(const FAnimNodeValueMessage& InMessage) const = 0;
+	virtual FText FormatNodeValue(const FAnimNodeValueMessage& InMessage) const = 0;
 };
\ No newline at end of file
diff --git a/Engine/Plugins/Animation/GameplayInsights/Source/GameplayInsights/Private/SAnimGraphSchematicView.cpp b/Engine/Plugins/Animation/GameplayInsights/Source/GameplayInsights/Private/SAnimGraphSchematicView.cpp
new file mode 100644
--- /dev/null
+++ b/Engine/Plugins/Animation/GameplayInsights/Source/GameplayInsights/Private/SAnimGraphSchematicView.cpp
@@ -0,0 +1,966 @@
+// Copyright 1998-2019 Epic Games, Inc. All Rights Reserved.
+
+#include "SAnimGraphSchematicView.h"
+#include "AnimationProvider.h"
+#include "Widgets/Input/SSearchBox.h"
+#include "TraceServices/Model/AnalysisSession.h"
+#include "Insights/ITimingViewSession.h"
+#include "GameplayProvider.h"
+#include "TraceServices/Model/Frames.h"
+#include "Widgets/Layout/SBorder.h"
+#include "Widgets/SBoxPanel.h"
+#include "Styling/CoreStyle.h"
+#include "GameplayInsightsStyle.h"
+#include "Widgets/Layout/SScrollBorder.h"
+#include "Framework/MultiBox/MultiBoxBuilder.h"
+#include "Widgets/Input/SEditableTextBox.h"
+#include "Widgets/Input/SHyperlink.h"
+#include "Subsystems/AssetEditorSubsystem.h"
+
+#define LOCTEXT_NAMESPACE "SAnimGraphSchematicView"
+
+namespace AnimGraphSchematicPropertyColumns
+{
+	static const FName Name("Name");
+	static const FName Value("Value");
+};
+
+namespace AnimGraphSchematicColumns
+{
+	static const FName Type("Type");
+	static const FName Name("Name");
+	static const FName Weight("Weight");
+	static const FName RootMotionWeight("Root Motion Weight");
+};
+
+enum class EAnimGraphSchematicFilterState
+{
+	Hidden,
+	Visible,
+	Highlighted,
+};
+
+// A in the tree of 'properties' for an animation node's debug info
+class FAnimGraphSchematicPropertyNode
+{
+public:
+	FAnimGraphSchematicPropertyNode(const FText& InName, const TSharedPtr<FAnimNodeValueMessage>& InValue, const Trace::IAnalysisSession& InAnalysisSession)
+		: AnalysisSession(InAnalysisSession)
+		, Name(InName)
+		, Value(InValue)
+	{}
+
+	TSharedRef<SWidget> MakeValueWidget()
+	{
+		TSharedPtr<FAnimNodeValueMessage> PinnedValue = Value.Pin();
+		if(PinnedValue.IsValid())
+		{
+			return StaticMakeValueWidget(AnalysisSession, PinnedValue.ToSharedRef());
+		}
+		else
+		{
+			return SNullWidget::NullWidget;
+		}
+	}
+
+	static TSharedRef<SWidget> StaticMakeValueWidget(const Trace::IAnalysisSession& InAnalysisSession, const TSharedRef<FAnimNodeValueMessage>& InValue) 
+	{ 
+		switch(InValue->Type)
+		{
+		case EAnimNodeValueType::Bool:
+			return 
+				SNew(SCheckBox)
+				.IsEnabled(false)
+				.IsChecked(InValue->Bool.bValue ? ECheckBoxState::Checked : ECheckBoxState::Unchecked);
+
+		case EAnimNodeValueType::Int32:
+			return
+				SNew(SBox)
+				.WidthOverride(125.0f)
+				[
+					SNew(SEditableTextBox)
+					.IsEnabled(false)
+					.Font(FCoreStyle::Get().GetFontStyle("SmallFont"))
+					.Text(FText::AsNumber(InValue->Int32.Value))
+				];
+
+		case EAnimNodeValueType::Float:
+			return 
+				SNew(SBox)
+				.WidthOverride(125.0f)
+				[
+					SNew(SEditableTextBox)
+					.IsEnabled(false)
+					.Font(FCoreStyle::Get().GetFontStyle("SmallFont"))
+					.Text(FText::AsNumber(InValue->Float.Value))
+				];
+
+		case EAnimNodeValueType::Vector:
+			return SNew(SHorizontalBox)
+				+SHorizontalBox::Slot()
+				.AutoWidth()
+				[
+					SNew(SBox)
+					.WidthOverride(125.0f)
+					[
+						SNew(SEditableTextBox)
+						.IsEnabled(false)
+						.Font(FCoreStyle::Get().GetFontStyle("SmallFont"))
+						.Text(FText::AsNumber(InValue->Vector.Value.X))
+					]
+				]
+				+SHorizontalBox::Slot()
+				.AutoWidth()
+				[
+					SNew(SBox)
+					.WidthOverride(125.0f)
+					[
+						SNew(SEditableTextBox)
+						.IsEnabled(false)
+						.Font(FCoreStyle::Get().GetFontStyle("SmallFont"))
+						.Text(FText::AsNumber(InValue->Vector.Value.Y))
+					]
+				]
+				+SHorizontalBox::Slot()
+				.AutoWidth()
+				[
+					SNew(SBox)
+					.WidthOverride(125.0f)
+					[
+						SNew(SEditableTextBox)
+						.IsEnabled(false)
+						.Font(FCoreStyle::Get().GetFontStyle("SmallFont"))
+						.Text(FText::AsNumber(InValue->Vector.Value.Z))
+					]
+				];
+
+		case EAnimNodeValueType::String:
+			return 
+				SNew(STextBlock)
+				.Font(FCoreStyle::Get().GetFontStyle("SmallFont"))
+				.Text(FText::FromString(InValue->String.Value));
+
+		case EAnimNodeValueType::Object:
+		{
+			const FGameplayProvider* GameplayProvider = InAnalysisSession.ReadProvider<FGameplayProvider>(FGameplayProvider::ProviderName);
+			if(GameplayProvider)
+			{
+				Trace::FAnalysisSessionReadScope SessionReadScope(InAnalysisSession);
+
+				const FObjectInfo& ObjectInfo = GameplayProvider->GetObjectInfo(InValue->Object.Value);
+#if WITH_EDITOR
+				return 
+					SNew(SHyperlink)
+					.Text(FText::FromString(ObjectInfo.Name))
+					.TextStyle(&FCoreStyle::Get().GetWidgetStyle<FTextBlockStyle>("SmallText"))
+					.ToolTipText(FText::Format(LOCTEXT("AssetHyperlinkTooltipFormat", "Open asset '{0}'"), FText::FromString(ObjectInfo.PathName)))
+					.OnNavigate_Lambda([ObjectInfo]()
+					{
+						GEditor->GetEditorSubsystem<UAssetEditorSubsystem>()->OpenEditorForAsset(ObjectInfo.PathName);
+					});
+#else
+				return 
+					SNew(STextBlock)
+					.Font(FCoreStyle::Get().GetFontStyle("SmallFont"))
+					.Text(FText::FromString(ObjectInfo.Name))
+					.ToolTipText(FText::FromString(ObjectInfo.PathName));
+#endif
+			}
+		}
+		break;
+
+		case EAnimNodeValueType::Class:
+		{
+			const FGameplayProvider* GameplayProvider = InAnalysisSession.ReadProvider<FGameplayProvider>(FGameplayProvider::ProviderName);
+			if(GameplayProvider)
+			{
+				Trace::FAnalysisSessionReadScope SessionReadScope(InAnalysisSession);
+
+				const FClassInfo& ClassInfo = GameplayProvider->GetClassInfo(InValue->Class.Value);
+#if WITH_EDITOR
+				return 
+					SNew(SHyperlink)
+					.Text(FText::FromString(ClassInfo.Name))
+					.TextStyle(&FCoreStyle::Get().GetWidgetStyle<FTextBlockStyle>("SmallText"))
+					.ToolTipText(FText::Format(LOCTEXT("ClassHyperlinkTooltipFormat", "Open class '{0}'"), FText::FromString(ClassInfo.PathName)))
+					.OnNavigate_Lambda([ClassInfo]()
+					{
+						GEditor->GetEditorSubsystem<UAssetEditorSubsystem>()->OpenEditorForAsset(ClassInfo.PathName);
+					});
+#else
+				return 
+					SNew(STextBlock)
+					.Font(FCoreStyle::Get().GetFontStyle("SmallFont"))
+					.Text(FText::FromString(ClassInfo.Name))
+					.ToolTipText(FText::FromString(ClassInfo.PathName));
+#endif			
+			}
+		}
+		break;
+		}
+
+		return SNullWidget::NullWidget; 
+	}
+
+	const Trace::IAnalysisSession& AnalysisSession;
+
+	FText Name;
+
+	TWeakPtr<FAnimNodeValueMessage> Value;
+
+	TWeakPtr<FAnimGraphSchematicPropertyNode> Parent;
+
+	TArray<TSharedRef<FAnimGraphSchematicPropertyNode>> Children;
+};
+
+// Container for an entry in the property view
+class SAnimGraphSchematicPropertyNode : public SMultiColumnTableRow<TSharedRef<FAnimGraphSchematicPropertyNode>>
+{
+public:
+	void Construct(const FArguments& InArgs, const TSharedRef<STableViewBase>& InOwnerTable, TSharedRef<FAnimGraphSchematicPropertyNode> InNode)
+	{
+		Node = InNode;
+
+		SMultiColumnTableRow<TSharedRef<FAnimGraphSchematicPropertyNode>>::Construct(
+			FSuperRowType::FArguments()
+			.Padding(1.0f),
+			InOwnerTable
+		);
+	}
+
+	// SMultiColumnTableRow interface
+	virtual TSharedRef<SWidget> GenerateWidgetForColumn(const FName& InColumnName) override
+	{
+		const bool bIsRoot = !Node->Parent.IsValid();
+
+		if (InColumnName == AnimGraphSchematicPropertyColumns::Name)
+		{
+			return 
+				SNew(SBorder)
+				.BorderImage(bIsRoot ? FGameplayInsightsStyle::Get().GetBrush("SchematicViewRootLeft") : FCoreStyle::Get().GetBrush("NoBorder"))
+				[
+					SNew(SHorizontalBox)
+					+SHorizontalBox::Slot()
+					.AutoWidth()
+					.Padding(6, 0, 0, 0)
+					.VAlign(VAlign_Center)
+					[
+						SNew(SExpanderArrow, SharedThis(this))
+						.IndentAmount(0)
+					]
+					+SHorizontalBox::Slot()
+					.FillWidth(1.0f)
+					.VAlign(VAlign_Center)
+					[
+						SNew(STextBlock)
+						.Font(FCoreStyle::Get().GetFontStyle(bIsRoot ? "ExpandableArea.TitleFont" : "SmallFont"))
+						.Text(Node->Name)
+					]
+				];
+		}
+		else if(InColumnName == AnimGraphSchematicPropertyColumns::Value)
+		{
+			return
+				SNew(SBorder)
+				.BorderImage(bIsRoot ? FGameplayInsightsStyle::Get().GetBrush("SchematicViewRootMid") : FCoreStyle::Get().GetBrush("NoBorder"))
+				[
+					SNew(SHorizontalBox)
+					+SHorizontalBox::Slot()
+					.AutoWidth()
+					.VAlign(VAlign_Center)
+					[
+						Node->MakeValueWidget()
+					]
+				];
+		}
+
+		return SNullWidget::NullWidget;
+	}
+
+	TSharedPtr<FAnimGraphSchematicPropertyNode> Node;
+};
+
+// Node representing debug data for an anim node
+class FAnimGraphSchematicNode : public TSharedFromThis<FAnimGraphSchematicNode>
+{
+public:
+	FAnimGraphSchematicNode(int32 InNodeId, const FText& InType, const Trace::IAnalysisSession& InAnalysisSession)
+		: AnalysisSession(InAnalysisSession)
+		, NodeId(InNodeId)
+		, Type(InType)
+		, FilterState(EAnimGraphSchematicFilterState::Hidden)
+		, bLinearized(false)
+	{
+	}
+
+	void MakePropertyNodes(TArray<TSharedRef<FAnimGraphSchematicPropertyNode>>& OutNodes)
+	{
+		// Add 'root' representing the node itself
+		TSharedRef<FAnimGraphSchematicPropertyNode> Root = MakeShared<FAnimGraphSchematicPropertyNode>(Type, nullptr, AnalysisSession);
+		OutNodes.Add(Root);
+
+		for(TSharedRef<FAnimNodeValueMessage> Value : Values)
+		{
+			TSharedRef<FAnimGraphSchematicPropertyNode> NewNode = MakeShared<FAnimGraphSchematicPropertyNode>(FText::FromString(Value->Key), Value, AnalysisSession);
+			Root->Children.Add(NewNode);
+			NewNode->Parent = Root;
+		}
+	}
+
+	const Trace::IAnalysisSession& AnalysisSession;
+
+	int32 NodeId;
+
+	FText Type;
+
+	TMap<FName, TSharedRef<FAnimNodeValueMessage>> KeysAndValues;
+
+	TArray<TSharedRef<FAnimNodeValueMessage>> Values;
+
+	TWeakPtr<FAnimGraphSchematicNode> Parent;
+
+	TArray<TSharedRef<FAnimGraphSchematicNode>> Children;
+
+	TArray<TSharedRef<FAnimGraphSchematicNode>> FlattenedLinearChildren;
+
+	EAnimGraphSchematicFilterState FilterState;
+
+	bool bLinearized;
+};
+
+class SAnimGraphSchematicNode : public SMultiColumnTableRow<TSharedRef<FAnimGraphSchematicNode>>
+{
+	SLATE_BEGIN_ARGS(SAnimGraphSchematicNode) {}
+
+	SLATE_ATTRIBUTE(FText, FilterText)
+
+	SLATE_END_ARGS()
+
+	void Construct(const FArguments& InArgs, const TSharedRef<STableViewBase>& InOwnerTable, TSharedRef<FAnimGraphSchematicNode> InNode)
+	{
+		Node = InNode;
+		FilterText = InArgs._FilterText;
+
+		SMultiColumnTableRow<TSharedRef<FAnimGraphSchematicNode>>::Construct(
+			FSuperRowType::FArguments()
+			.Padding(1.0f),
+			InOwnerTable
+		);
+	}
+
+	// SMultiColumnTableRow interface
+	virtual TSharedRef<SWidget> GenerateWidgetForColumn(const FName& InColumnName) override
+	{
+		const bool bIsRoot = !Node->Parent.IsValid();
+
+		if (InColumnName == AnimGraphSchematicColumns::Type)
+		{
+			return 
+				SNew(SBorder)
+				.BorderImage(bIsRoot ? FGameplayInsightsStyle::Get().GetBrush("SchematicViewRootLeft") : FCoreStyle::Get().GetBrush("NoBorder"))
+				[
+					SNew(SHorizontalBox)
+					+SHorizontalBox::Slot()
+					.AutoWidth()
+					.VAlign(VAlign_Center)
+					.Padding(6, 0, 0, 0)
+					[
+						SNew(SExpanderArrow, SharedThis(this))
+						.IndentAmount(12)
+					]
+					+ SHorizontalBox::Slot()
+					.FillWidth(1.0f)
+					.VAlign(VAlign_Center)
+					[
+						SNew(STextBlock)
+						.Font(FCoreStyle::Get().GetFontStyle(bIsRoot ? "ExpandableArea.TitleFont" : "SmallFont"))
+						.Text(Node->Type)
+						.HighlightText(FilterText)
+					]
+				];
+		}
+		else
+		{
+			TSharedPtr<SWidget> ValueWidget = SNullWidget::NullWidget;
+			TSharedRef<FAnimNodeValueMessage>* ValuePtr = Node->KeysAndValues.Find(InColumnName);
+			if(ValuePtr)
+			{
+				ValueWidget = FAnimGraphSchematicPropertyNode::StaticMakeValueWidget(Node->AnalysisSession, *ValuePtr);
+			}
+
+			return
+				SNew(SBorder)
+				.BorderImage(bIsRoot ? FGameplayInsightsStyle::Get().GetBrush("SchematicViewRootMid") : FCoreStyle::Get().GetBrush("NoBorder"))
+				[
+					SNew(SHorizontalBox)
+					+SHorizontalBox::Slot()
+					.VAlign(VAlign_Center)
+					.AutoWidth()
+					[
+						ValueWidget.ToSharedRef()
+					]
+				];
+		}
+
+		return SNullWidget::NullWidget;
+	}
+
+	TSharedPtr<FAnimGraphSchematicNode> Node;
+
+	TAttribute<FText> FilterText;
+};
+
+void SAnimGraphSchematicView::Construct(const FArguments& InArgs, uint64 InAnimInstanceId, Insights::ITimingViewSession& InTimingViewSession, const Trace::IAnalysisSession& InAnalysisSession)
+{
+	AnimInstanceId = InAnimInstanceId;
+	TimingViewSession = &InTimingViewSession;
+	AnalysisSession = &InAnalysisSession;
+	TimeMarker = InTimingViewSession.GetTimeMarker();
+
+	InTimingViewSession.OnTimeMarkerChanged().AddSP(this, &SAnimGraphSchematicView::HandleTimeMarkerChanged);
+
+	// Create header row and add default columns
+	HeaderRow = SNew(SHeaderRow);
+
+	// Make default columns
+	Columns.Add(AnimGraphSchematicColumns::Type, { 0, true });
+	Columns.Add(AnimGraphSchematicColumns::Name, { 1, true });
+	Columns.Add(AnimGraphSchematicColumns::Weight, { 99999, true });
+	Columns.Add(AnimGraphSchematicColumns::RootMotionWeight, { 100000, false });
+
+	TreeView = SNew(STreeView<TSharedRef<FAnimGraphSchematicNode>>)
+			.TreeItemsSource(&FilteredNodes)
+			.OnGenerateRow(this, &SAnimGraphSchematicView::HandleGenerateRow)
+			.OnGetChildren(this, &SAnimGraphSchematicView::HandleGetChildren)
+			.HeaderRow(HeaderRow.ToSharedRef())
+			.OnSelectionChanged(this, &SAnimGraphSchematicView::HandleSelectionChanged);
+
+	ChildSlot
+	[
+		SNew(SVerticalBox)
+		+SVerticalBox::Slot()
+		.AutoHeight()
+		[
+			SAssignNew(SearchBox, SSearchBox)
+			.OnTextChanged_Lambda([this](const FText& InText){ FilterText = InText; RefreshFilter(); })
+		]
+		+SVerticalBox::Slot()
+		.FillHeight(1.0f)
+		[
+			SAssignNew(Splitter, SSplitter)
+			.Orientation(Orient_Vertical)
+			+SSplitter::Slot()
+			.Value(0.7f)
+			[
+				SNew(SScrollBorder, TreeView.ToSharedRef())
+				[
+					TreeView.ToSharedRef()
+				]
+			]
+		]
+		+SVerticalBox::Slot()
+		.AutoHeight()
+		.HAlign(HAlign_Right)
+		[
+			SAssignNew(ViewButton, SComboButton)
+			.ContentPadding(0)
+			.ForegroundColor(this, &SAnimGraphSchematicView::GetViewButtonForegroundColor)
+			.ButtonStyle(FGameplayInsightsStyle::Get(), "SchematicViewViewButton")
+			.OnGetMenuContent(this, &SAnimGraphSchematicView::HandleGetViewMenuContent)
+			.ButtonContent()
+			[
+				SNew(SHorizontalBox)
+				+SHorizontalBox::Slot()
+				.AutoWidth()
+				.VAlign(VAlign_Center)
+				[
+					SNew(SImage)
+					.Image(FGameplayInsightsStyle::Get().GetBrush("SchematicViewViewButtonIcon"))
+				]
+				+SHorizontalBox::Slot()
+				.AutoWidth()
+				.Padding(2, 0, 0, 0)
+				.VAlign(VAlign_Center)
+				[
+					SNew(STextBlock)
+					.Text(LOCTEXT("ViewButton", "View Options"))
+				]
+			]
+		]
+	];
+
+	RefreshColumns();
+	RefreshNodes();
+}
+
+TSharedRef<ITableRow> SAnimGraphSchematicView::HandleGenerateRow(TSharedRef<FAnimGraphSchematicNode> Item, const TSharedRef<STableViewBase>& OwnerTable)
+{
+	return 
+		SNew(SAnimGraphSchematicNode, OwnerTable, Item)
+		.FilterText_Lambda([this](){ return FilterText; });
+}
+
+void SAnimGraphSchematicView::HandleGetChildren(TSharedRef<FAnimGraphSchematicNode> InItem, TArray<TSharedRef<FAnimGraphSchematicNode>>& OutChildren)
+{
+	for(const TSharedRef<FAnimGraphSchematicNode>& Child : InItem->FlattenedLinearChildren)
+	{
+		if(Child->FilterState != EAnimGraphSchematicFilterState::Hidden)
+		{
+			OutChildren.Add(Child);
+		}
+	}
+}
+
+TSharedRef<ITableRow> SAnimGraphSchematicView::HandleGeneratePropertyRow(TSharedRef<FAnimGraphSchematicPropertyNode> Item, const TSharedRef<STableViewBase>& OwnerTable)
+{
+	return 
+		SNew(SAnimGraphSchematicPropertyNode, OwnerTable, Item);
+}
+
+void SAnimGraphSchematicView::HandleGetPropertyChildren(TSharedRef<FAnimGraphSchematicPropertyNode> InItem, TArray<TSharedRef<FAnimGraphSchematicPropertyNode>>& OutChildren)
+{
+	OutChildren.Append(InItem->Children);
+}
+
+// Add children recursively if only one child is present, linearizing runs of single children
+static void AddChildren_Helper(TSharedRef<FAnimGraphSchematicNode> InItem, TArray<TSharedRef<FAnimGraphSchematicNode>>& OutChildren)
+{
+	if(InItem->FlattenedLinearChildren.Num() == 0)
+	{
+		if(InItem->Children.Num() == 1)
+		{
+			if(!InItem->Children[0]->bLinearized)
+			{
+				InItem->Children[0]->bLinearized = true;
+				OutChildren.Add(InItem->Children[0]);
+				AddChildren_Helper(InItem->Children[0], OutChildren);
+			}
+		}
+		else if(InItem->Children.Num() > 1)
+		{
+			for(const TSharedRef<FAnimGraphSchematicNode>& Child : InItem->Children)
+			{
+				if(!Child->bLinearized)
+				{
+					Child->bLinearized = true;
+					OutChildren.Add(Child);
+				}
+			}
+		}
+	}
+}
+
+void SAnimGraphSchematicView::RefreshNodes()
+{
+	UnfilteredNodes.Reset();
+	FilteredNodes.Reset();
+	LinearNodes.Reset();
+	
+	TMap<int32, TSharedRef<FAnimGraphSchematicNode>> NodeMap;
+
+	const FAnimationProvider* AnimationProvider = AnalysisSession->ReadProvider<FAnimationProvider>(FAnimationProvider::ProviderName);
+	const FGameplayProvider* GameplayProvider = AnalysisSession->ReadProvider<FGameplayProvider>(FGameplayProvider::ProviderName);
+
+	if(AnimationProvider && GameplayProvider)
+	{
+		Trace::FAnalysisSessionReadScope SessionReadScope(*AnalysisSession);
+
+		AnimationProvider->ReadAnimGraphTimeline(AnimInstanceId, [this, &NodeMap, AnimationProvider, GameplayProvider](const FAnimationProvider::AnimGraphTimeline& InGraphTimeline)
+		{
+			const Trace::IFrameProvider& FramesProvider = Trace::ReadFrameProvider(*AnalysisSession);
+
+			Trace::FFrame Frame;
+			if(FramesProvider.GetFrameFromTime(ETraceFrameType::TraceFrameType_Game, TimeMarker, Frame))
+			{
+				InGraphTimeline.EnumerateEvents(Frame.StartTime, Frame.EndTime, [this, &NodeMap, AnimationProvider, GameplayProvider](double InGraphStartTime, double InGraphEndTime, uint32 InDepth, const FAnimGraphMessage& InMessage)
+				{
+					// Check for an update phase (which contains weights)
+					if(InMessage.Phase == EAnimGraphPhase::Update)
+					{
+						AnimationProvider->ReadAnimNodesTimeline(AnimInstanceId, [this, &NodeMap, InGraphStartTime, InGraphEndTime](const FAnimationProvider::AnimNodesTimeline& InNodesTimeline)
+						{
+							InNodesTimeline.EnumerateEvents(InGraphStartTime, InGraphEndTime, [this, &NodeMap](double InStartTime, double InEndTime, uint32 InDepth, const FAnimNodeMessage& InMessage)
+							{
+								if(InMessage.NodeId != INDEX_NONE && InMessage.Phase == EAnimGraphPhase::Update)
+								{
+									TSharedPtr<FAnimGraphSchematicNode> Node;
+									TSharedRef<FAnimGraphSchematicNode>* ExistingNodePtr = NodeMap.Find(InMessage.NodeId);
+									if(ExistingNodePtr == nullptr)
+									{
+										Node = NodeMap.Add(InMessage.NodeId, MakeShared<FAnimGraphSchematicNode>(InMessage.NodeId, FText::FromString(InMessage.NodeName), *AnalysisSession));
+
+										// Add dummy values for weight and root motion weight
+										TSharedRef<FAnimNodeValueMessage> WeightMessage = MakeShared<FAnimNodeValueMessage>();
+										WeightMessage->NodeId = InMessage.NodeId;
+										WeightMessage->FrameCounter = InMessage.FrameCounter;
+										WeightMessage->Type = EAnimNodeValueType::Float;
+										WeightMessage->Key = TEXT("Weight");
+										WeightMessage->Float.Value = InMessage.Weight;
+
+										Node->KeysAndValues.Add(AnimGraphSchematicColumns::Weight, WeightMessage);
+										Node->Values.Add(WeightMessage);
+
+										TSharedRef<FAnimNodeValueMessage> RootMotionWeightMessage = MakeShared<FAnimNodeValueMessage>();
+										RootMotionWeightMessage->NodeId = InMessage.NodeId;
+										RootMotionWeightMessage->FrameCounter = InMessage.FrameCounter;
+										RootMotionWeightMessage->Type = EAnimNodeValueType::Float;
+										RootMotionWeightMessage->Key = TEXT("Root Motion Weight");
+										RootMotionWeightMessage->Float.Value = InMessage.RootMotionWeight;
+
+										Node->KeysAndValues.Add(AnimGraphSchematicColumns::RootMotionWeight, RootMotionWeightMessage);
+										Node->Values.Add(RootMotionWeightMessage);
+
+										LinearNodes.Add(Node.ToSharedRef());
+									}
+									else
+									{
+										Node = *ExistingNodePtr;
+									}
+
+									if(InMessage.PreviousNodeId != INDEX_NONE)
+									{
+										check(InMessage.NodeId != InMessage.PreviousNodeId);
+										TSharedRef<FAnimGraphSchematicNode>* ExistingPreviousNodePtr = NodeMap.Find(InMessage.PreviousNodeId);
+										if(ExistingPreviousNodePtr != nullptr)
+										{
+											TSharedRef<FAnimGraphSchematicNode> PreviousNode = *ExistingPreviousNodePtr;
+											PreviousNode->Children.Add(Node.ToSharedRef());
+											Node->Parent = PreviousNode;
+										}
+									}
+								}
+							});
+						});
+
+						// TODO: Add visualizations of the reset of the data below
+// 						AnimationProvider->ReadStateMachinesTimeline(AnimInstanceId, [InGraphStartTime, InGraphEndTime](const FAnimationProvider::StateMachinesTimeline& InStateMachinesTimeline)
+// 						{
+// 							InStateMachinesTimeline.EnumerateEvents(InGraphStartTime, InGraphEndTime, [](double InStartTime, double InEndTime, uint32 InDepth, const FAnimStateMachineMessage& InMessage)
+// 							{
+// 								//DebugData.RecordStateData(InMessage.StateMachineIndex, InMessage.StateIndex, InMessage.StateWeight, InMessage.ElapsedTime);
+// 							});
+// 						});
+//						
+// 						AnimationProvider->ReadAnimSequencePlayersTimeline(AnimInstanceId, [InGraphStartTime, InGraphEndTime, GameplayProvider](const FAnimationProvider::AnimSequencePlayersTimeline& InSequencePlayersTimeline)
+// 						{
+// 							InSequencePlayersTimeline.EnumerateEvents(InGraphStartTime, InGraphEndTime, [](double InStartTime, double InEndTime, uint32 InDepth, const FAnimSequencePlayerMessage& InMessage)
+// 							{
+// 								//DebugData.RecordSequencePlayer(InMessage.NodeId, InMessage.Position, InMessage.Length, InMessage.FrameCount);
+// 							});
+// 						});
+// 
+// 						AnimationProvider->ReadAnimBlendSpacePlayersTimeline(AnimInstanceId, [InGraphStartTime, InGraphEndTime, GameplayProvider](const FAnimationProvider::BlendSpacePlayersTimeline& InBlendSpacePlayersTimeline)
+// 						{
+// 							InBlendSpacePlayersTimeline.EnumerateEvents(InGraphStartTime, InGraphEndTime, [GameplayProvider](double InStartTime, double InEndTime, uint32 InDepth, const FBlendSpacePlayerMessage& InMessage)
+// 							{
+// // 								UBlendSpaceBase* BlendSpaceBase = nullptr;
+// // 								const FObjectInfo* BlendSpaceInfo = GameplayProvider->FindObjectInfo(InMessage.BlendSpaceId);
+// // 								if(BlendSpaceInfo)
+// // 								{
+// // 									BlendSpaceBase = TSoftObjectPtr<UBlendSpaceBase>(FSoftObjectPath(BlendSpaceInfo->PathName)).LoadSynchronous();
+// // 								}
+// // 
+// // 								DebugData.RecordBlendSpacePlayer(InMessage.NodeId, BlendSpaceBase, InMessage.PositionX, InMessage.PositionY, InMessage.PositionZ);
+// 							});
+// 						});
+
+					}
+
+					AnimationProvider->ReadAnimNodeValuesTimeline(AnimInstanceId, [this, InGraphStartTime, InGraphEndTime, AnimationProvider, &NodeMap](const FAnimationProvider::AnimNodeValuesTimeline& InNodeValuesTimeline)
+					{
+						InNodeValuesTimeline.EnumerateEvents(InGraphStartTime, InGraphEndTime, [this, AnimationProvider, &NodeMap](double InStartTime, double InEndTime, uint32 InDepth, const FAnimNodeValueMessage& InMessage)
+						{
+							TSharedRef<FAnimGraphSchematicNode>* ExistingNodePtr = NodeMap.Find(InMessage.NodeId);
+							if(ExistingNodePtr != nullptr)
+							{
+								TSharedRef<FAnimGraphSchematicNode> ExistingNode = *ExistingNodePtr;
+
+								FName Key = InMessage.Key;
+
+								// Add to columns (+2 to leave space for the fixed type and name columns in the index range)
+								Columns.FindOrAdd(Key, { Columns.Num() + 2, false });
+								TSharedRef<FAnimNodeValueMessage> SharedMessage = MakeShared<FAnimNodeValueMessage>();
+								SharedMessage.Get() = InMessage;
+								ExistingNode->KeysAndValues.Add(Key, SharedMessage);
+								ExistingNode->Values.Add(SharedMessage);
+							}
+						});
+					});
+				});
+			};
+		});
+	}
+
+	// Re-sort columns
+	Columns.ValueSort([](const FColumnState& Column0, const FColumnState& Column1)
+	{
+		return Column0.SortIndex < Column1.SortIndex;
+	});
+
+	// Add root nodes to unfiltered list
+	for(const TSharedRef<FAnimGraphSchematicNode>& Node : LinearNodes)
+	{
+		if(!Node->Parent.IsValid())
+		{
+			UnfilteredNodes.Add(Node);
+		}
+	}
+
+	// Build flattened linear children from runs of single-child nodes
+	for(const TSharedRef<FAnimGraphSchematicNode>& Node : LinearNodes)
+	{
+		AddChildren_Helper(Node, Node->FlattenedLinearChildren);
+	}
+
+	RefreshFilter();
+}
+
+EAnimGraphSchematicFilterState SAnimGraphSchematicView::RefreshFilter_Helper(const TSharedRef<FAnimGraphSchematicNode>& InNode)
+{
+	InNode->FilterState = EAnimGraphSchematicFilterState::Hidden;
+
+	for(const TSharedRef<FAnimGraphSchematicNode>& ChildNode : InNode->FlattenedLinearChildren)
+	{
+		EAnimGraphSchematicFilterState ChildFilterState = RefreshFilter_Helper(ChildNode);
+		InNode->FilterState = FMath::Max(ChildFilterState, InNode->FilterState);
+	}
+	
+	if(InNode->FilterState == EAnimGraphSchematicFilterState::Hidden)
+	{
+		if(FilterText.IsEmpty())
+		{
+			InNode->FilterState = EAnimGraphSchematicFilterState::Visible;
+		}
+		else if (InNode->Type.ToString().Contains(FilterText.ToString()))
+		{
+			InNode->FilterState = EAnimGraphSchematicFilterState::Highlighted;
+		}
+		else
+		{
+			InNode->FilterState = EAnimGraphSchematicFilterState::Hidden;
+		}
+	}
+
+	bool bAutoExpand = (FilterText.IsEmpty() && InNode->FilterState != EAnimGraphSchematicFilterState::Hidden) ||
+					  (!FilterText.IsEmpty() && InNode->FilterState == EAnimGraphSchematicFilterState::Highlighted);
+	TreeView->SetItemExpansion(InNode, bAutoExpand);
+
+	return InNode->FilterState;
+}
+
+void SAnimGraphSchematicView::RefreshFilter()
+{
+	FilteredNodes.Reset();
+
+	for(const TSharedRef<FAnimGraphSchematicNode>& RootNode : UnfilteredNodes)
+	{
+		EAnimGraphSchematicFilterState FilterState = RefreshFilter_Helper(RootNode);
+		if(FilterState != EAnimGraphSchematicFilterState::Hidden)
+		{
+			FilteredNodes.Add(RootNode);
+		}
+	}
+
+	// re-select any previously selected nodes
+	TArray<TSharedRef<FAnimGraphSchematicNode>> SelectedItems;
+	for(const TSharedRef<FAnimGraphSchematicNode>& Node : LinearNodes)
+	{
+		if(SelectedNodeIds.Contains(Node->NodeId))
+		{
+			SelectedItems.Add(Node);
+		}
+	}
+
+	if(SelectedItems.Num() > 0)
+	{
+		TreeView->SetItemSelection(SelectedItems, true);
+		TreeView->RequestScrollIntoView(SelectedItems.Last());
+	}
+
+	TreeView->RequestTreeRefresh();
+}
+
+void SAnimGraphSchematicView::HandleTimeMarkerChanged(Insights::ETimeChangedFlags InFlags, double InTimeMarker)
+{
+	TimeMarker = InTimeMarker;
+
+	RefreshNodes();
+}
+
+TSharedRef<SWidget> SAnimGraphSchematicView::HandleGetViewMenuContent()
+{
+	FMenuBuilder MenuBuilder(true, nullptr);
+
+	MenuBuilder.BeginSection("Columns", LOCTEXT("ColumnsMenuHeader", "Columns"));
+	{
+		if(Columns.Num() > 0)
+		{
+			for(const TPair<FName, FColumnState>& ColumnPair : Columns)
+			{
+				const FName ColumnId = ColumnPair.Key;
+
+				MenuBuilder.AddMenuEntry(
+					FText::FromName(ColumnId), 
+					LOCTEXT("ColumnTooltip", "Enable/disable this column"),
+					FSlateIcon(), 
+					FUIAction(
+						FExecuteAction::CreateLambda([this, ColumnId]()
+						{
+							FColumnState* StatePtr = Columns.Find(ColumnId);
+							if(StatePtr != nullptr && StatePtr->bEnabled)
+							{
+								StatePtr->bEnabled = false;
+							}
+							else
+							{
+								StatePtr->bEnabled = true;
+							}
+
+							RefreshColumns();
+						}),
+						FCanExecuteAction(),
+						FIsActionChecked::CreateLambda([this, ColumnId]()
+						{
+							return Columns.FindRef(ColumnId).bEnabled;
+						})
+					),
+					NAME_None,
+					EUserInterfaceActionType::Check
+				);
+			}
+		}
+		else
+		{
+			MenuBuilder.AddMenuEntry(
+				LOCTEXT("NoColumns", "No Optional Columns Found"), 
+				LOCTEXT("NoColumnsTooltip", "No Optional Columns Found in the Current Session"), 
+				FSlateIcon(), 
+				FUIAction()
+			);
+		}
+	}
+	MenuBuilder.EndSection();
+
+	return MenuBuilder.MakeWidget();
+}
+
+void SAnimGraphSchematicView::RefreshColumns()
+{
+	HeaderRow->ClearColumns();
+
+	Columns.ValueSort([](const FColumnState& Column0, const FColumnState& Column1)
+	{
+		return Column0.SortIndex < Column1.SortIndex;
+	});
+
+	for(const TPair<FName, FColumnState>& ColumnPair : Columns)
+	{
+		if(ColumnPair.Value.bEnabled)
+		{
+			HeaderRow->AddColumn(
+				SHeaderRow::FColumn::FArguments()
+				.ColumnId(ColumnPair.Key)
+				.DefaultLabel(FText::FromName(ColumnPair.Key))
+			);
+		}
+	}
+}
+
+FSlateColor SAnimGraphSchematicView::GetViewButtonForegroundColor() const
+{
+	static const FName InvertedForegroundName("InvertedForeground");
+	static const FName DefaultForegroundName("DefaultForeground");
+
+	return ViewButton->IsHovered() ? FCoreStyle::Get().GetSlateColor(InvertedForegroundName) : FCoreStyle::Get().GetSlateColor(DefaultForegroundName);
+}
+
+void SAnimGraphSchematicView::HandleSelectionChanged(TSharedPtr<FAnimGraphSchematicNode> InNode, ESelectInfo::Type InSelectInfo)
+{
+	TArray<TSharedRef<FAnimGraphSchematicNode>> SelectedNodes;
+	TreeView->GetSelectedItems(SelectedNodes);
+
+	// Preserve selection of node IDs so scrubbing will re-select them
+	if(InSelectInfo != ESelectInfo::Direct)
+	{
+		SelectedNodeIds.Reset();
+		for(const TSharedRef<FAnimGraphSchematicNode>& Item : SelectedNodes)
+		{
+			SelectedNodeIds.Add(Item->NodeId);
+		}
+	}
+
+	RefreshDetails(SelectedNodes);
+}
+
+void SAnimGraphSchematicView::RefreshDetails(const TArray<TSharedRef<FAnimGraphSchematicNode>>& InNodes)
+{
+	if(InNodes.Num() > 0)
+	{
+		// Nodes to show, so create/re-use content box in splitter
+		TSharedPtr<SVerticalBox> VerticalBox;
+
+		if(DetailsContentBox.IsValid())
+		{
+			VerticalBox = DetailsContentBox.Pin();
+			VerticalBox->ClearChildren();
+		}
+		else
+		{
+			VerticalBox = SNew(SVerticalBox);
+			DetailsContentBox = VerticalBox;
+
+			Splitter->AddSlot()
+			.Value(0.3f)
+			[
+				VerticalBox.ToSharedRef()
+			];
+		}
+
+		// Add content for each node
+		PropertyNodes.Reset();
+
+		for(const TSharedRef<FAnimGraphSchematicNode>& Node : InNodes)
+		{
+			Node->MakePropertyNodes(PropertyNodes);
+		}
+
+		PropertyTreeView = SNew(STreeView<TSharedRef<FAnimGraphSchematicPropertyNode>>)
+			.SelectionMode(ESelectionMode::None)
+			.OnGenerateRow(this, &SAnimGraphSchematicView::HandleGeneratePropertyRow)
+			.OnGetChildren(this, &SAnimGraphSchematicView::HandleGetPropertyChildren)
+			.TreeItemsSource(&PropertyNodes)
+			.HeaderRow(
+				SNew(SHeaderRow)
+				.Visibility(EVisibility::Collapsed)
+				+SHeaderRow::Column(AnimGraphSchematicPropertyColumns::Name)
+				.DefaultLabel(LOCTEXT("PropertiesNameColumn", "Name"))
+				+SHeaderRow::Column(AnimGraphSchematicPropertyColumns::Value)
+				.DefaultLabel(LOCTEXT("PropertiesValueColumn", "Value"))
+			);
+
+		for(const TSharedRef<FAnimGraphSchematicPropertyNode>& RootNode : PropertyNodes)
+		{
+			PropertyTreeView->SetItemExpansion(RootNode, true);
+		}
+
+		VerticalBox->AddSlot()
+			.FillHeight(1.0f)
+			[
+				SNew(SScrollBorder, PropertyTreeView.ToSharedRef())
+				[
+					PropertyTreeView.ToSharedRef()
+				]
+			];
+	}
+	else
+	{
+		// no nodes to show, so remove splitter (and therefore content) if valid
+		if(DetailsContentBox.IsValid())
+		{
+			DetailsContentBox.Pin()->ClearChildren();
+			Splitter->RemoveAt(1);
+		}
+	}
+}
+
+#undef LOCTEXT_NAMESPACE
\ No newline at end of file
diff --git a/Engine/Plugins/Animation/GameplayInsights/Source/GameplayInsights/Private/SAnimGraphSchematicView.h b/Engine/Plugins/Animation/GameplayInsights/Source/GameplayInsights/Private/SAnimGraphSchematicView.h
new file mode 100644
--- /dev/null
+++ b/Engine/Plugins/Animation/GameplayInsights/Source/GameplayInsights/Private/SAnimGraphSchematicView.h
@@ -0,0 +1,115 @@
+// Copyright 1998-2019 Epic Games, Inc. All Rights Reserved.
+
+#pragma once
+
+#include "CoreMinimal.h"
+#include "Widgets/SCompoundWidget.h"
+#include "Widgets/Views/STreeView.h"
+
+namespace Insights { class ITimingViewSession; }
+namespace Insights { enum class ETimeChangedFlags; }
+namespace Trace { class IAnalysisSession; }
+class FAnimationSharedData;
+class IInsightsManager;
+class FAnimGraphSchematicNode;
+class FAnimGraphSchematicPropertyNode;
+class SSearchBox;
+class SHeaderRow;
+class SComboButton;
+class SSplitter;
+enum class EAnimGraphSchematicFilterState;
+
+class SAnimGraphSchematicView : public SCompoundWidget
+{
+public:
+	SLATE_BEGIN_ARGS(SAnimGraphSchematicView) {}
+
+	SLATE_END_ARGS()
+
+	void Construct(const FArguments& InArgs, uint64 InAnimInstanceId, Insights::ITimingViewSession& InTimingViewSession, const Trace::IAnalysisSession& InAnalysisSession);
+
+	uint64 GetAnimInstanceId() const { return AnimInstanceId; }
+
+private:
+	// Generate a row widget for an item
+	TSharedRef<ITableRow> HandleGenerateRow(TSharedRef<FAnimGraphSchematicNode> Item, const TSharedRef<STableViewBase>& OwnerTable);
+
+	// Get the children of an item
+	void HandleGetChildren(TSharedRef<FAnimGraphSchematicNode> InItem, TArray<TSharedRef<FAnimGraphSchematicNode>>& OutChildren);
+
+	// Generate a row widget for a property item
+	TSharedRef<ITableRow> HandleGeneratePropertyRow(TSharedRef<FAnimGraphSchematicPropertyNode> Item, const TSharedRef<STableViewBase>& OwnerTable);
+
+	// Get the children of a property item
+	void HandleGetPropertyChildren(TSharedRef<FAnimGraphSchematicPropertyNode> InItem, TArray<TSharedRef<FAnimGraphSchematicPropertyNode>>& OutChildren);
+
+	// Handle the time marker being scrubbed
+	void HandleTimeMarkerChanged(Insights::ETimeChangedFlags InFlags, double InTimeMarker);
+
+	// Handle tree selection changing
+	void HandleSelectionChanged(TSharedPtr<FAnimGraphSchematicNode> InNode, ESelectInfo::Type InSelectInfo);
+
+	// Refresh the unfiltered and filtered nodes
+	void RefreshNodes();
+
+	// Recursive helper to filter tree
+	EAnimGraphSchematicFilterState RefreshFilter_Helper(const TSharedRef<FAnimGraphSchematicNode>& InNode);
+
+	// Refresh the displayed columns
+	void RefreshColumns();
+
+	// Refresh the filtered nodes
+	void RefreshFilter();
+
+	// Refresh the details panel
+	void RefreshDetails(const TArray<TSharedRef<FAnimGraphSchematicNode>>& InNodes);
+
+	// Handle getting content for the view menu
+	TSharedRef<SWidget> HandleGetViewMenuContent();
+
+	// Handle inverting the color of the button text when hovered
+	FSlateColor GetViewButtonForegroundColor() const;
+
+private:
+	const Trace::IAnalysisSession* AnalysisSession;
+
+	Insights::ITimingViewSession* TimingViewSession;
+
+	TSharedPtr<SSearchBox> SearchBox;
+
+	TSharedPtr<SHeaderRow> HeaderRow;
+
+	TSharedPtr<SSplitter> Splitter;
+
+	TSharedPtr<STreeView<TSharedRef<FAnimGraphSchematicNode>>> TreeView;
+
+	TSharedPtr<STreeView<TSharedRef<FAnimGraphSchematicPropertyNode>>> PropertyTreeView;
+
+	TSharedPtr<SComboButton> ViewButton;
+
+	TWeakPtr<SVerticalBox> DetailsContentBox;
+
+	TArray<TSharedRef<FAnimGraphSchematicNode>> UnfilteredNodes;
+
+	TArray<TSharedRef<FAnimGraphSchematicNode>> FilteredNodes;
+
+	TArray<TSharedRef<FAnimGraphSchematicNode>> LinearNodes;
+
+	TArray<TSharedRef<FAnimGraphSchematicPropertyNode>> PropertyNodes;
+
+	TSet<int32> SelectedNodeIds;
+
+	struct FColumnState
+	{
+		int32 SortIndex = 0;
+		bool bEnabled = false;
+	};
+
+	TMap<FName, FColumnState> Columns;
+
+	FText FilterText;
+
+	double TimeMarker;
+
+	uint64 AnimInstanceId;
+};
\ No newline at end of file
diff --git a/Engine/Source/Editor/KismetCompiler/Private/AnimBlueprintCompiler.cpp b/Engine/Source/Editor/KismetCompiler/Private/AnimBlueprintCompiler.cpp
--- a/Engine/Source/Editor/KismetCompiler/Private/AnimBlueprintCompiler.cpp
+++ b/Engine/Source/Editor/KismetCompiler/Private/AnimBlueprintCompiler.cpp
@@ -556,6 +556,11 @@ void FAnimBlueprintCompilerContext::ProcessAnimationNode(UAnimGraphNode_Base* Vi
 		// Process root nodes
 		ProcessRoot(RootNode);
 	}
+	else if(UAnimGraphNode_StateResult* StateResultNode = Cast<UAnimGraphNode_StateResult>(VisualAnimNode))
+	{
+		// Process state result nodes
+		ProcessStateResult(StateResultNode);
+	}
 	else if (UAnimGraphNode_AssetPlayerBase* AssetPlayerBase = Cast<UAnimGraphNode_AssetPlayerBase>(VisualAnimNode))
 	{			
 		// Process Asset Player nodes to, if necessary cache off their node index for retrieval at runtime (used for evaluating Automatic Rule Transitions when using Layer nodes)
@@ -748,6 +753,13 @@ void FAnimBlueprintCompilerContext::ProcessRoot(UAnimGraphNode_Root* Root)
 	Root->Node.Name = TrueNode->GetGraph()->GetFName();
 }
 
+void FAnimBlueprintCompilerContext::ProcessStateResult(UAnimGraphNode_StateResult* StateResult)
+{
+	UAnimGraphNode_StateResult* TrueNode = MessageLog.FindSourceObjectTypeChecked<UAnimGraphNode_StateResult>(StateResult);
+
+	StateResult->Node.Name = TrueNode->GetGraph()->GetFName();
+}
+
 void FAnimBlueprintCompilerContext::ProcessUseCachedPose(UAnimGraphNode_UseCachedPose* UseCachedPose)
 {
 	bool bSuccessful = false;
diff --git a/Engine/Source/Editor/KismetCompiler/Private/AnimBlueprintCompiler.h b/Engine/Source/Editor/KismetCompiler/Private/AnimBlueprintCompiler.h
--- a/Engine/Source/Editor/KismetCompiler/Private/AnimBlueprintCompiler.h
+++ b/Engine/Source/Editor/KismetCompiler/Private/AnimBlueprintCompiler.h
@@ -329,6 +329,9 @@ private:
 	// Compiles one root node
 	void ProcessRoot(UAnimGraphNode_Root* Root);
 
+	// Compiles one state result node
+	void ProcessStateResult(UAnimGraphNode_StateResult* StateResult);
+
 	// Traverses linked anim graph links looking for slot names and state machine names, returning their count in a name map
 	typedef TMap<FName, int32> NameToCountMap;
 	void GetDuplicatedSlotAndStateNames(UAnimGraphNode_LinkedAnimGraphBase* InLinkedAnimGraph, NameToCountMap& OutStateMachineNameToCountMap, NameToCountMap& OutSlotNameToCountMap);
diff --git a/Engine/Source/Runtime/AnimGraphRuntime/Private/AnimNodes/AnimNode_AimOffsetLookAt.cpp b/Engine/Source/Runtime/AnimGraphRuntime/Private/AnimNodes/AnimNode_AimOffsetLookAt.cpp
--- a/Engine/Source/Runtime/AnimGraphRuntime/Private/AnimNodes/AnimNode_AimOffsetLookAt.cpp
+++ b/Engine/Source/Runtime/AnimGraphRuntime/Private/AnimNodes/AnimNode_AimOffsetLookAt.cpp
@@ -65,7 +65,7 @@ void FAnimNode_AimOffsetLookAt::UpdateAssetPlayer(const FAnimationUpdateContext&
 
 	BasePose.Update(Context);
 
-	TRACE_ANIM_NODE_VALUE(Context, TEXT("Play Time"), InternalTimeAccumulator);
+	TRACE_ANIM_NODE_VALUE(Context, TEXT("Playback Time"), InternalTimeAccumulator);
 }
 
 void FAnimNode_AimOffsetLookAt::CacheBones_AnyThread(const FAnimationCacheBonesContext& Context)
diff --git a/Engine/Source/Runtime/AnimGraphRuntime/Private/AnimNodes/AnimNode_BlendSpaceEvaluator.cpp b/Engine/Source/Runtime/AnimGraphRuntime/Private/AnimNodes/AnimNode_BlendSpaceEvaluator.cpp
--- a/Engine/Source/Runtime/AnimGraphRuntime/Private/AnimNodes/AnimNode_BlendSpaceEvaluator.cpp
+++ b/Engine/Source/Runtime/AnimGraphRuntime/Private/AnimNodes/AnimNode_BlendSpaceEvaluator.cpp
@@ -21,8 +21,9 @@ void FAnimNode_BlendSpaceEvaluator::UpdateAssetPlayer(const FAnimationUpdateCont
 
 	UpdateInternal(Context);
 
-	TRACE_ANIM_NODE_VALUE(Context, TEXT("Blend Space Name"), BlendSpace ? *BlendSpace->GetName() : TEXT("None"));
-	TRACE_ANIM_NODE_VALUE(Context, TEXT("Play Time"), InternalTimeAccumulator);
+	TRACE_ANIM_NODE_VALUE(Context, TEXT("Name"), BlendSpace ? *BlendSpace->GetName() : TEXT("None"));
+	TRACE_ANIM_NODE_VALUE(Context, TEXT("Blend Space"), BlendSpace);
+	TRACE_ANIM_NODE_VALUE(Context, TEXT("Playback Time"), InternalTimeAccumulator);
 }
 
 void FAnimNode_BlendSpaceEvaluator::GatherDebugData(FNodeDebugData& DebugData)
diff --git a/Engine/Source/Runtime/AnimGraphRuntime/Private/AnimNodes/AnimNode_BlendSpacePlayer.cpp b/Engine/Source/Runtime/AnimGraphRuntime/Private/AnimNodes/AnimNode_BlendSpacePlayer.cpp
--- a/Engine/Source/Runtime/AnimGraphRuntime/Private/AnimNodes/AnimNode_BlendSpacePlayer.cpp
+++ b/Engine/Source/Runtime/AnimGraphRuntime/Private/AnimNodes/AnimNode_BlendSpacePlayer.cpp
@@ -116,8 +116,9 @@ void FAnimNode_BlendSpacePlayer::UpdateInternal(const FAnimationUpdateContext& C
 	}
 
 	TRACE_BLENDSPACE_PLAYER(Context, *this);
-	TRACE_ANIM_NODE_VALUE(Context, TEXT("Blend Space Name"), BlendSpace ? *BlendSpace->GetName() : TEXT("None"));
-	TRACE_ANIM_NODE_VALUE(Context, TEXT("Play Time"), InternalTimeAccumulator);
+	TRACE_ANIM_NODE_VALUE(Context, TEXT("Name"), BlendSpace ? *BlendSpace->GetName() : TEXT("None"));
+	TRACE_ANIM_NODE_VALUE(Context, TEXT("Blend Space"), BlendSpace);
+	TRACE_ANIM_NODE_VALUE(Context, TEXT("Playback Time"), InternalTimeAccumulator);
 }
 
 void FAnimNode_BlendSpacePlayer::Evaluate_AnyThread(FPoseContext& Output)
diff --git a/Engine/Source/Runtime/AnimGraphRuntime/Private/AnimNodes/AnimNode_PoseHandler.cpp b/Engine/Source/Runtime/AnimGraphRuntime/Private/AnimNodes/AnimNode_PoseHandler.cpp
--- a/Engine/Source/Runtime/AnimGraphRuntime/Private/AnimNodes/AnimNode_PoseHandler.cpp
+++ b/Engine/Source/Runtime/AnimGraphRuntime/Private/AnimNodes/AnimNode_PoseHandler.cpp
@@ -79,6 +79,7 @@ void FAnimNode_PoseHandler::UpdateAssetPlayer(const FAnimationUpdateContext& Con
 		UpdatePoseAssetProperty(Context.AnimInstanceProxy);
 	}
 
+	TRACE_ANIM_NODE_VALUE(Context, TEXT("Name"), CurrentPoseAsset.IsValid() ? *CurrentPoseAsset.Get()->GetName() : TEXT("None"));
 	TRACE_ANIM_NODE_VALUE(Context, TEXT("Pose Asset"), CurrentPoseAsset.IsValid() ? *CurrentPoseAsset.Get()->GetName() : TEXT("None"));
 }
 
diff --git a/Engine/Source/Runtime/AnimGraphRuntime/Private/AnimNodes/AnimNode_RotationOffsetBlendSpace.cpp b/Engine/Source/Runtime/AnimGraphRuntime/Private/AnimNodes/AnimNode_RotationOffsetBlendSpace.cpp
--- a/Engine/Source/Runtime/AnimGraphRuntime/Private/AnimNodes/AnimNode_RotationOffsetBlendSpace.cpp
+++ b/Engine/Source/Runtime/AnimGraphRuntime/Private/AnimNodes/AnimNode_RotationOffsetBlendSpace.cpp
@@ -70,7 +70,7 @@ void FAnimNode_RotationOffsetBlendSpace::UpdateAssetPlayer(const FAnimationUpdat
 	BasePose.Update(Context);
 
 	TRACE_ANIM_NODE_VALUE(Context, TEXT("Alpha"), ActualAlpha);
-	TRACE_ANIM_NODE_VALUE(Context, TEXT("PlayTime"), InternalTimeAccumulator);
+	TRACE_ANIM_NODE_VALUE(Context, TEXT("Playback Time"), InternalTimeAccumulator);
 }
 
 void FAnimNode_RotationOffsetBlendSpace::Evaluate_AnyThread(FPoseContext& Context)
diff --git a/Engine/Source/Runtime/AnimGraphRuntime/Private/AnimNodes/AnimNode_SequenceEvaluator.cpp b/Engine/Source/Runtime/AnimGraphRuntime/Private/AnimNodes/AnimNode_SequenceEvaluator.cpp
--- a/Engine/Source/Runtime/AnimGraphRuntime/Private/AnimNodes/AnimNode_SequenceEvaluator.cpp
+++ b/Engine/Source/Runtime/AnimGraphRuntime/Private/AnimNodes/AnimNode_SequenceEvaluator.cpp
@@ -87,7 +87,8 @@ void FAnimNode_SequenceEvaluator::UpdateAssetPlayer(const FAnimationUpdateContex
 
 	bReinitialized = false;
 
-	TRACE_ANIM_NODE_VALUE(Context, TEXT("Sequence"), Sequence != nullptr ? Sequence->GetFName() : NAME_None);
+	TRACE_ANIM_NODE_VALUE(Context, TEXT("Name"), Sequence != nullptr ? Sequence->GetFName() : NAME_None);
+	TRACE_ANIM_NODE_VALUE(Context, TEXT("Sequence"), Sequence);
 	TRACE_ANIM_NODE_VALUE(Context, TEXT("InputTime"), ExplicitTime);
 	TRACE_ANIM_NODE_VALUE(Context, TEXT("Time"), InternalTimeAccumulator);
 }
diff --git a/Engine/Source/Runtime/AnimGraphRuntime/Private/AnimNodes/AnimNode_Slot.cpp b/Engine/Source/Runtime/AnimGraphRuntime/Private/AnimNodes/AnimNode_Slot.cpp
--- a/Engine/Source/Runtime/AnimGraphRuntime/Private/AnimNodes/AnimNode_Slot.cpp
+++ b/Engine/Source/Runtime/AnimGraphRuntime/Private/AnimNodes/AnimNode_Slot.cpp
@@ -46,8 +46,8 @@ void FAnimNode_Slot::Update_AnyThread(const FAnimationUpdateContext& Context)
 	}
 
 #if ANIM_TRACE_ENABLED
-	TRACE_ANIM_NODE_VALUE(Context, TEXT("Slot Name"), SlotName);
-	TRACE_ANIM_NODE_VALUE(Context, TEXT("Weight"), WeightData.SlotNodeWeight);
+	TRACE_ANIM_NODE_VALUE(Context, TEXT("Name"), SlotName);
+	TRACE_ANIM_NODE_VALUE(Context, TEXT("Slot Weight"), WeightData.SlotNodeWeight);
 	TRACE_ANIM_NODE_VALUE(Context, TEXT("Pose Source"), (WeightData.SourceWeight <= ZERO_ANIMWEIGHT_THRESH));
 
 	Context.AnimInstanceProxy->TraceMontageEvaluationData(Context, SlotName);
diff --git a/Engine/Source/Runtime/AnimGraphRuntime/Private/BoneControllers/AnimNode_HandIKRetargeting.cpp b/Engine/Source/Runtime/AnimGraphRuntime/Private/BoneControllers/AnimNode_HandIKRetargeting.cpp
--- a/Engine/Source/Runtime/AnimGraphRuntime/Private/BoneControllers/AnimNode_HandIKRetargeting.cpp
+++ b/Engine/Source/Runtime/AnimGraphRuntime/Private/BoneControllers/AnimNode_HandIKRetargeting.cpp
@@ -98,10 +98,6 @@ void FAnimNode_HandIKRetargeting::EvaluateSkeletalControl_AnyThread(FComponentSp
 
 #if ANIM_TRACE_ENABLED
 	TRACE_ANIM_NODE_VALUE(Output, TEXT("Hand FK Weight"), HandFKWeight);
-	for (int32 BoneIndex = 0; BoneIndex < IKBonesToMove.Num(); BoneIndex++)
-	{
-		TRACE_ANIM_NODE_VALUE(Output, *FString::Printf(TEXT("Bone %d"), BoneIndex), IKBonesToMove[BoneIndex].BoneName);
-	}
 #endif
 }
 
diff --git a/Engine/Source/Runtime/Engine/Classes/Animation/AnimNodeBase.h b/Engine/Source/Runtime/Engine/Classes/Animation/AnimNodeBase.h
--- a/Engine/Source/Runtime/Engine/Classes/Animation/AnimNodeBase.h
+++ b/Engine/Source/Runtime/Engine/Classes/Animation/AnimNodeBase.h
@@ -261,6 +261,11 @@ public:
 		FAnimationUpdateContext Result(*this);
 		Result.SharedContext = InSharedContext;
 
+#if ANIM_TRACE_ENABLED
+		// This is currently only used in the case of cached poses, where we dont want to preserve the previous node, so clear it here
+		Result.PreviousNodeId = INDEX_NONE;
+#endif
+
 		return Result;
 	}
 
diff --git a/Engine/Source/Runtime/Engine/Private/Animation/AnimInstanceProxy.cpp b/Engine/Source/Runtime/Engine/Private/Animation/AnimInstanceProxy.cpp
--- a/Engine/Source/Runtime/Engine/Private/Animation/AnimInstanceProxy.cpp
+++ b/Engine/Source/Runtime/Engine/Private/Animation/AnimInstanceProxy.cpp
@@ -38,6 +38,7 @@ const FName NAME_AnimGraph(TEXT("AnimGraph"));
 void FAnimInstanceProxy::UpdateAnimationNode(const FAnimationUpdateContext& InContext)
 {
 	TRACE_SCOPED_ANIM_GRAPH(InContext);
+	TRACE_SCOPED_ANIM_NODE(InContext);
 
 	UpdateAnimationNode_WithRoot(InContext, RootNode, NAME_AnimGraph);
 }
diff --git a/Engine/Source/Runtime/Engine/Private/Animation/AnimNodeBase.cpp b/Engine/Source/Runtime/Engine/Private/Animation/AnimNodeBase.cpp
--- a/Engine/Source/Runtime/Engine/Private/Animation/AnimNodeBase.cpp
+++ b/Engine/Source/Runtime/Engine/Private/Animation/AnimNodeBase.cpp
@@ -467,10 +467,13 @@ void FComponentSpacePoseLink::EvaluateComponentSpace(FComponentSpacePoseContext&
 
 	if (LinkedNode != NULL)
 	{
+		{
 #if ANIM_TRACE_ENABLED
-		Output.SetNodeId(LinkID);
+			Output.SetNodeId(LinkID);
+			TRACE_SCOPED_ANIM_NODE(Output);
 #endif
-		LinkedNode->EvaluateComponentSpace_AnyThread(Output);
+			LinkedNode->EvaluateComponentSpace_AnyThread(Output);
+		}
 
 #if WITH_EDITOR
 		Output.AnimInstanceProxy->RegisterWatchedPose(Output.Pose, LinkID);
diff --git a/Engine/Source/Runtime/Engine/Private/Animation/AnimNode_LinkedAnimGraph.cpp b/Engine/Source/Runtime/Engine/Private/Animation/AnimNode_LinkedAnimGraph.cpp
--- a/Engine/Source/Runtime/Engine/Private/Animation/AnimNode_LinkedAnimGraph.cpp
+++ b/Engine/Source/Runtime/Engine/Private/Animation/AnimNode_LinkedAnimGraph.cpp
@@ -136,6 +136,7 @@ void FAnimNode_LinkedAnimGraph::Update_AnyThread(const FAnimationUpdateContext&
 		PendingBlendDuration = -1.0f;
 	}
 
+	TRACE_ANIM_NODE_VALUE(InContext, TEXT("Name"), GetDynamicLinkFunctionName());
 	TRACE_ANIM_NODE_VALUE(InContext, TEXT("Target Class"), InstanceClass.Get());
 }
 
diff --git a/Engine/Source/Runtime/Engine/Private/Animation/AnimNode_Root.cpp b/Engine/Source/Runtime/Engine/Private/Animation/AnimNode_Root.cpp
--- a/Engine/Source/Runtime/Engine/Private/Animation/AnimNode_Root.cpp
+++ b/Engine/Source/Runtime/Engine/Private/Animation/AnimNode_Root.cpp
@@ -23,6 +23,8 @@ void FAnimNode_Root::CacheBones_AnyThread(const FAnimationCacheBonesContext& Con
 
 void FAnimNode_Root::Update_AnyThread(const FAnimationUpdateContext& Context)
 {
+	TRACE_ANIM_NODE_VALUE(Context, TEXT("Name"), Name);
+
 	GetEvaluateGraphExposedInputs().Execute(Context);
 	Result.Update(Context);
 }
diff --git a/Engine/Source/Runtime/Engine/Private/Animation/AnimNode_SaveCachedPose.cpp b/Engine/Source/Runtime/Engine/Private/Animation/AnimNode_SaveCachedPose.cpp
--- a/Engine/Source/Runtime/Engine/Private/Animation/AnimNode_SaveCachedPose.cpp
+++ b/Engine/Source/Runtime/Engine/Private/Animation/AnimNode_SaveCachedPose.cpp
@@ -54,7 +54,7 @@ void FAnimNode_SaveCachedPose::Update_AnyThread(const FAnimationUpdateContext& C
 	// Store this context for the post update
 	CachedUpdate.Context = Context.WithOtherSharedContext(CachedUpdate.SharedContext.Get());
 
-	TRACE_ANIM_NODE_VALUE(Context, TEXT("Cached Pose Name"), CachePoseName);
+	TRACE_ANIM_NODE_VALUE(Context, TEXT("Name"), CachePoseName);
 }
 
 void FAnimNode_SaveCachedPose::Evaluate_AnyThread(FPoseContext& Output)
@@ -107,7 +107,10 @@ void FAnimNode_SaveCachedPose::PostGraphUpdate()
 		}
 
 		// Update the max weighted pose node
-		Pose.Update(CachedUpdateContexts[MaxWeightIdx].Context);
+		{
+			TRACE_SCOPED_ANIM_NODE(CachedUpdateContexts[MaxWeightIdx].Context);
+			Pose.Update(CachedUpdateContexts[MaxWeightIdx].Context);
+		}
 
 		// Determine if any ancestors are interested in the other updates we'll be skipping
 		TArray<FAnimNode_Base*, TInlineAllocator<4>> AncestorsWithSkippedUpdateHandlers;
diff --git a/Engine/Source/Runtime/Engine/Private/Animation/AnimNode_SequencePlayer.cpp b/Engine/Source/Runtime/Engine/Private/Animation/AnimNode_SequencePlayer.cpp
--- a/Engine/Source/Runtime/Engine/Private/Animation/AnimNode_SequencePlayer.cpp
+++ b/Engine/Source/Runtime/Engine/Private/Animation/AnimNode_SequencePlayer.cpp
@@ -82,6 +82,7 @@ void FAnimNode_SequencePlayer::UpdateAssetPlayer(const FAnimationUpdateContext&
 #endif
 
 	TRACE_ANIM_SEQUENCE_PLAYER(Context, *this);
+	TRACE_ANIM_NODE_VALUE(Context, TEXT("Name"), Sequence != nullptr ? Sequence->GetFName() : NAME_None);
 	TRACE_ANIM_NODE_VALUE(Context, TEXT("Sequence"), Sequence);
 	TRACE_ANIM_NODE_VALUE(Context, TEXT("Playback Time"), InternalTimeAccumulator);
 }
diff --git a/Engine/Source/Runtime/Engine/Private/Animation/AnimNode_StateMachine.cpp b/Engine/Source/Runtime/Engine/Private/Animation/AnimNode_StateMachine.cpp
--- a/Engine/Source/Runtime/Engine/Private/Animation/AnimNode_StateMachine.cpp
+++ b/Engine/Source/Runtime/Engine/Private/Animation/AnimNode_StateMachine.cpp
@@ -561,13 +561,8 @@ void FAnimNode_StateMachine::Update_AnyThread(const FAnimationUpdateContext& Con
 	ElapsedTime += Context.GetDeltaTime();
 
 #if ANIM_TRACE_ENABLED
-	TRACE_ANIM_NODE_VALUE(Context, TEXT("State Machine Name"), GetMachineDescription()->MachineName);
+	TRACE_ANIM_NODE_VALUE(Context, TEXT("Name"), GetMachineDescription()->MachineName);
 	TRACE_ANIM_NODE_VALUE(Context, TEXT("Current State"), GetStateInfo().StateName);
-	for (int32 PoseIndex = 0; PoseIndex < StatePoseLinks.Num(); ++PoseIndex)
-	{
-		const FString StateName = FString::Printf(TEXT("State %s Weight"), *GetStateInfo(PoseIndex).StateName.ToString());
-		TRACE_ANIM_NODE_VALUE(Context, *StateName, GetStateWeight(PoseIndex));
-	}
 #endif
 }
 
diff --git a/Engine/Source/Runtime/Engine/Private/Animation/AnimNode_UseCachedPose.cpp b/Engine/Source/Runtime/Engine/Private/Animation/AnimNode_UseCachedPose.cpp
--- a/Engine/Source/Runtime/Engine/Private/Animation/AnimNode_UseCachedPose.cpp
+++ b/Engine/Source/Runtime/Engine/Private/Animation/AnimNode_UseCachedPose.cpp
@@ -24,9 +24,14 @@ void FAnimNode_UseCachedPose::CacheBones_AnyThread(const FAnimationCacheBonesCon
 
 void FAnimNode_UseCachedPose::Update_AnyThread(const FAnimationUpdateContext& Context)
 {
-	LinkToCachingNode.Update(Context);
-
-	TRACE_ANIM_NODE_VALUE(Context, TEXT("Cached Pose Name"), CachePoseName);
+	{
+		// This makes sure we dont see a 'save cached pose' entry in the debug data.
+		// This will be handled later on when the cached pose branch gets taken.
+		TRACE_SCOPED_ANIM_NODE_SUSPEND;	
+		LinkToCachingNode.Update(Context);
+	}
+
+	TRACE_ANIM_NODE_VALUE(Context, TEXT("Name"), CachePoseName);
 }
 
 void FAnimNode_UseCachedPose::Evaluate_AnyThread(FPoseContext& Output)
diff --git a/Engine/Source/Runtime/Engine/Private/Animation/AnimTrace.cpp b/Engine/Source/Runtime/Engine/Private/Animation/AnimTrace.cpp
--- a/Engine/Source/Runtime/Engine/Private/Animation/AnimTrace.cpp
+++ b/Engine/Source/Runtime/Engine/Private/Animation/AnimTrace.cpp
@@ -183,39 +183,78 @@ FUObjectAnnotationSparseBool GSkeletalMeshTraceAnnotations;
 // Map used for unique name output
 TMap<FName, uint32> GAnimTraceNames;
 
+class FSuspendCounter : public TThreadSingleton<FSuspendCounter>
+{
+public:
+	FSuspendCounter()
+		: SuspendCount(0)
+	{}
+
+	int32 SuspendCount;
+};
+
+FAnimTrace::FScopedAnimNodeTraceSuspend::FScopedAnimNodeTraceSuspend()
+{
+	FSuspendCounter::Get().SuspendCount++;
+}
+
+FAnimTrace::FScopedAnimNodeTraceSuspend::~FScopedAnimNodeTraceSuspend()
+{
+	FSuspendCounter::Get().SuspendCount--;
+	check(FSuspendCounter::Get().SuspendCount >= 0);
+}
+
 FAnimTrace::FScopedAnimNodeTrace::FScopedAnimNodeTrace(const FAnimationInitializeContext& InContext)
 	: Context(InContext)
 {
-	OutputAnimNodeStart(InContext, FPlatformTime::Cycles64(), InContext.GetPreviousNodeId(), InContext.GetCurrentNodeId(), 0.0f, 0.0f, (__underlying_type(EPhase))EPhase::Initialize);
+	if(FSuspendCounter::Get().SuspendCount == 0)
+	{
+		OutputAnimNodeStart(InContext, FPlatformTime::Cycles64(), InContext.GetPreviousNodeId(), InContext.GetCurrentNodeId(), 0.0f, 0.0f, (__underlying_type(EPhase))EPhase::Initialize);
+	}
 }
 
 FAnimTrace::FScopedAnimNodeTrace::FScopedAnimNodeTrace(const FAnimationUpdateContext& InContext)
 	: Context(InContext)
 {
-	OutputAnimNodeStart(InContext, FPlatformTime::Cycles64(), InContext.GetPreviousNodeId(), InContext.GetCurrentNodeId(), InContext.GetFinalBlendWeight(), InContext.GetRootMotionWeightModifier(), (__underlying_type(EPhase))EPhase::Update);
+	if(FSuspendCounter::Get().SuspendCount == 0)
+	{
+		OutputAnimNodeStart(InContext, FPlatformTime::Cycles64(), InContext.GetPreviousNodeId(), InContext.GetCurrentNodeId(), InContext.GetFinalBlendWeight(), InContext.GetRootMotionWeightModifier(), (__underlying_type(EPhase))EPhase::Update);
+	}
 }
 
 FAnimTrace::FScopedAnimNodeTrace::FScopedAnimNodeTrace(const FAnimationCacheBonesContext& InContext)
 	: Context(InContext)
 {
-	OutputAnimNodeStart(InContext, FPlatformTime::Cycles64(), InContext.GetPreviousNodeId(), InContext.GetCurrentNodeId(), 0.0f, 0.0f, (__underlying_type(EPhase))EPhase::CacheBones);
+	if(FSuspendCounter::Get().SuspendCount == 0)
+	{
+		OutputAnimNodeStart(InContext, FPlatformTime::Cycles64(), InContext.GetPreviousNodeId(), InContext.GetCurrentNodeId(), 0.0f, 0.0f, (__underlying_type(EPhase))EPhase::CacheBones);
+	}
 }
 
 FAnimTrace::FScopedAnimNodeTrace::FScopedAnimNodeTrace(const FPoseContext& InContext)
 	: Context(InContext)
 {
-	OutputAnimNodeStart(InContext, FPlatformTime::Cycles64(), InContext.GetPreviousNodeId(), InContext.GetCurrentNodeId(), 0.0f, 0.0f, (__underlying_type(EPhase))EPhase::Evaluate);
+	if(FSuspendCounter::Get().SuspendCount == 0)
+	{
+		OutputAnimNodeStart(InContext, FPlatformTime::Cycles64(), InContext.GetPreviousNodeId(), InContext.GetCurrentNodeId(), 0.0f, 0.0f, (__underlying_type(EPhase))EPhase::Evaluate);
+	}
 }
 
 FAnimTrace::FScopedAnimNodeTrace::FScopedAnimNodeTrace(const FComponentSpacePoseContext& InContext)
 	: Context(InContext)
 {
-	OutputAnimNodeStart(InContext, FPlatformTime::Cycles64(), InContext.GetPreviousNodeId(), InContext.GetCurrentNodeId(), 0.0f, 0.0f, (__underlying_type(EPhase))EPhase::Evaluate);
+	if(FSuspendCounter::Get().SuspendCount == 0)
+	{
+		OutputAnimNodeStart(InContext, FPlatformTime::Cycles64(), InContext.GetPreviousNodeId(), InContext.GetCurrentNodeId(), 0.0f, 0.0f, (__underlying_type(EPhase))EPhase::Evaluate);
+	}
 }
 
 FAnimTrace::FScopedAnimNodeTrace::~FScopedAnimNodeTrace()
 {
-	OutputAnimNodeEnd(Context, FPlatformTime::Cycles64());
+	if(FSuspendCounter::Get().SuspendCount == 0)
+	{
+		OutputAnimNodeEnd(Context, FPlatformTime::Cycles64());
+	}
 }
 
 FAnimTrace::FScopedAnimGraphTrace::FScopedAnimGraphTrace(const FAnimationInitializeContext& InContext)
@@ -269,6 +308,8 @@ void FAnimTrace::Init()
 		UE_TRACE_EVENT_IS_ENABLED(Animation, AnimNodeValueFloat);
 		UE_TRACE_EVENT_IS_ENABLED(Animation, AnimNodeValueVector);
 		UE_TRACE_EVENT_IS_ENABLED(Animation, AnimNodeValueString);
+		UE_TRACE_EVENT_IS_ENABLED(Animation, AnimNodeValueObject);
+		UE_TRACE_EVENT_IS_ENABLED(Animation, AnimNodeValueClass);
 		UE_TRACE_EVENT_IS_ENABLED(Animation, AnimSequencePlayer);
 		UE_TRACE_EVENT_IS_ENABLED(Animation, StateMachineState);
 		UE_TRACE_EVENT_IS_ENABLED(Animation, Name);
@@ -492,24 +533,38 @@ void FAnimTrace::OutputAnimNodeStart(const FAnimationBaseContext& InContext, uin
 		return;
 	}
 
+	if(InNodeId == INDEX_NONE)
+	{
+		return;
+	}
+
 	check(InContext.AnimInstanceProxy);
 
 	TRACE_OBJECT(InContext.AnimInstanceProxy->GetAnimInstanceObject());
 
+	FString DisplayNameString;
 	IAnimClassInterface* AnimBlueprintClass = InContext.GetAnimClass();
-	check(AnimBlueprintClass);
-	const TArray<FStructPropertyPath>& AnimNodeProperties = AnimBlueprintClass->GetAnimNodeProperties();
-	check(AnimNodeProperties.IsValidIndex(InNodeId));
-	FStructProperty* LinkedProperty = AnimNodeProperties[InNodeId].Get();
-	check(LinkedProperty->Struct);
+	if(AnimBlueprintClass)
+	{
+		const TArray<FStructPropertyPath>& AnimNodeProperties = AnimBlueprintClass->GetAnimNodeProperties();
+		check(AnimNodeProperties.IsValidIndex(InNodeId));
+		FStructProperty* LinkedProperty = AnimNodeProperties[InNodeId].Get();
+		check(LinkedProperty->Struct);
 
 #if WITH_EDITOR
-	FString DisplayNameString = LinkedProperty->Struct->GetDisplayNameText().ToString();
+		DisplayNameString = LinkedProperty->Struct->GetDisplayNameText().ToString();
 #else
-	FString DisplayNameString = LinkedProperty->Struct->GetName();
+		DisplayNameString = LinkedProperty->Struct->GetName();
 #endif
 
-	DisplayNameString.RemoveFromStart(TEXT("Anim Node "));
+		DisplayNameString.RemoveFromStart(TEXT("Anim Node "));
+	}
+	else
+	{
+		DisplayNameString = TEXT("Anim Node");
+	}
+
+	check(InPreviousNodeId != InNodeId);
 
 	UE_TRACE_LOG(Animation, AnimNodeStart, AnimationChannel, (DisplayNameString.Len() + 1) * sizeof(TCHAR))
 		<< AnimNodeStart.StartCycle(InStartCycle)
diff --git a/Engine/Source/Runtime/Engine/Public/Animation/AnimTrace.h b/Engine/Source/Runtime/Engine/Public/Animation/AnimTrace.h
--- a/Engine/Source/Runtime/Engine/Public/Animation/AnimTrace.h
+++ b/Engine/Source/Runtime/Engine/Public/Animation/AnimTrace.h
@@ -72,6 +72,13 @@ struct FAnimTrace
 		EPhase Phase;
 	};
 
+	/** Helper for suspending anim node tracing */
+	struct FScopedAnimNodeTraceSuspend
+	{
+		FScopedAnimNodeTraceSuspend();
+		~FScopedAnimNodeTraceSuspend();
+	};
+
 	/** Describes a debug line output to the world */
 	struct FDebugLine
 	{
@@ -156,6 +163,9 @@ struct FAnimTrace
 #define TRACE_SCOPED_ANIM_NODE(Context) \
 	FAnimTrace::FScopedAnimNodeTrace _ScopedAnimNodeTrace(Context);
 
+#define TRACE_SCOPED_ANIM_NODE_SUSPEND \
+	FAnimTrace::FScopedAnimNodeTraceSuspend _ScopedAnimNodeTraceSuspend;
+
 #define TRACE_ANIM_NODE_VALUE(Context, Key, Value) \
 	FAnimTrace::OutputAnimNodeValue(Context, Key, Value);
 
@@ -173,6 +183,7 @@ struct FAnimTrace
 #define TRACE_SKELETALMESH_FRAME(Component)
 #define TRACE_SCOPED_ANIM_GRAPH(Context)
 #define TRACE_SCOPED_ANIM_NODE(Context)
+#define TRACE_SCOPED_ANIM_NODE_SUSPEND
 #define TRACE_ANIM_NODE_VALUE(Context, Key, Value)
 #define TRACE_ANIM_SEQUENCE_PLAYER(Context, Node)
 #define TRACE_ANIM_STATE_MACHINE_STATE(Context, StateMachineIndex, StateIndex, StateWeight, ElapsedTime)
