Non unity/non PCH build fixes

diff --git a/Engine/Plugins/Animation/GameplayInsights/Source/GameplayInsights/Private/AnimationSharedData.h b/Engine/Plugins/Animation/GameplayInsights/Source/GameplayInsights/Private/AnimationSharedData.h
--- a/Engine/Plugins/Animation/GameplayInsights/Source/GameplayInsights/Private/AnimationSharedData.h
+++ b/Engine/Plugins/Animation/GameplayInsights/Source/GameplayInsights/Private/AnimationSharedData.h
@@ -16,6 +16,7 @@ class FMenuBuilder;
 class UWorld;
 class IAnimationBlueprintEditor;
 struct FCustomDebugObject;
+class SDockTab;
 
 class FAnimationSharedData
 {
diff --git a/Engine/Plugins/Animation/GameplayInsights/Source/GameplayInsights/Private/GameplayInsightsModule.cpp b/Engine/Plugins/Animation/GameplayInsights/Source/GameplayInsights/Private/GameplayInsightsModule.cpp
--- a/Engine/Plugins/Animation/GameplayInsights/Source/GameplayInsights/Private/GameplayInsightsModule.cpp
+++ b/Engine/Plugins/Animation/GameplayInsights/Source/GameplayInsights/Private/GameplayInsightsModule.cpp
@@ -13,6 +13,7 @@
 #include "WorkspaceMenuStructureModule.h"
 #include "TraceServices/ITraceServicesModule.h"
 #include "TraceServices/SessionService.h"
+#include "Widgets/Docking/SDockTab.h"
 
 #if WITH_EDITOR
 #include "IAnimationBlueprintEditorModule.h"
diff --git a/Engine/Plugins/Animation/GameplayInsights/Source/GameplayInsights/Private/SAnimGraphSchematicView.cpp b/Engine/Plugins/Animation/GameplayInsights/Source/GameplayInsights/Private/SAnimGraphSchematicView.cpp
--- a/Engine/Plugins/Animation/GameplayInsights/Source/GameplayInsights/Private/SAnimGraphSchematicView.cpp
+++ b/Engine/Plugins/Animation/GameplayInsights/Source/GameplayInsights/Private/SAnimGraphSchematicView.cpp
@@ -15,7 +15,14 @@
 #include "Framework/MultiBox/MultiBoxBuilder.h"
 #include "Widgets/Input/SEditableTextBox.h"
 #include "Widgets/Input/SHyperlink.h"
+#include "Widgets/Input/SCheckBox.h
+#include "Widgets/Input/SComboButton.h
+#include "Widgets/Images/SImage.h"
+
+#if WITH_EDITOR
 #include "Subsystems/AssetEditorSubsystem.h"
+#include "Editor.h"
+#endif
 
 #define LOCTEXT_NAMESPACE "SAnimGraphSchematicView"
 
