Fixed includes

diff --git a/Engine/Plugins/Animation/GameplayInsights/Source/GameplayInsights/Private/SAnimGraphSchematicView.cpp b/Engine/Plugins/Animation/GameplayInsights/Source/GameplayInsights/Private/SAnimGraphSchematicView.cpp
--- a/Engine/Plugins/Animation/GameplayInsights/Source/GameplayInsights/Private/SAnimGraphSchematicView.cpp
+++ b/Engine/Plugins/Animation/GameplayInsights/Source/GameplayInsights/Private/SAnimGraphSchematicView.cpp
@@ -15,8 +15,8 @@
 #include "Framework/MultiBox/MultiBoxBuilder.h"
 #include "Widgets/Input/SEditableTextBox.h"
 #include "Widgets/Input/SHyperlink.h"
-#include "Widgets/Input/SCheckBox.h
-#include "Widgets/Input/SComboButton.h
+#include "Widgets/Input/SCheckBox.h"
+#include "Widgets/Input/SComboButton.h"
 #include "Widgets/Images/SImage.h"
 
 #if WITH_EDITOR
