More include fixes

diff --git a/Engine/Plugins/Animation/GameplayInsights/Source/GameplayInsights/Private/AnimationSharedData.cpp b/Engine/Plugins/Animation/GameplayInsights/Source/GameplayInsights/Private/AnimationSharedData.cpp
--- a/Engine/Plugins/Animation/GameplayInsights/Source/GameplayInsights/Private/AnimationSharedData.cpp
+++ b/Engine/Plugins/Animation/GameplayInsights/Source/GameplayInsights/Private/AnimationSharedData.cpp
@@ -14,13 +14,15 @@
 #include "Framework/MultiBox/MultiBoxBuilder.h"
 #include "AnimNodesTrack.h"
 #include "GameplayTimingViewExtender.h"
+#include "SAnimGraphSchematicView.h"
+#include "GameplayInsightsModule.h"
+#include "ModuleManager.h"
+#include "Widgets/Docking/SDockTab.h"
 
 #if WITH_EDITOR
 #include "EditorViewportClient.h"
 #include "Editor/EditorEngine.h"
 #endif
-#include "SAnimGraphSchematicView.h"
-#include "GameplayInsightsModule.h"
 
 #define LOCTEXT_NAMESPACE "AnimationSharedData"
 
