One more include fix

diff --git a/Engine/Plugins/Animation/GameplayInsights/Source/GameplayInsights/Private/AnimationSharedData.cpp b/Engine/Plugins/Animation/GameplayInsights/Source/GameplayInsights/Private/AnimationSharedData.cpp
--- a/Engine/Plugins/Animation/GameplayInsights/Source/GameplayInsights/Private/AnimationSharedData.cpp
+++ b/Engine/Plugins/Animation/GameplayInsights/Source/GameplayInsights/Private/AnimationSharedData.cpp
@@ -16,7 +16,7 @@
 #include "GameplayTimingViewExtender.h"
 #include "SAnimGraphSchematicView.h"
 #include "GameplayInsightsModule.h"
-#include "ModuleManager.h"
+#include "Modules/ModuleManager.h"
 #include "Widgets/Docking/SDockTab.h"
 
 #if WITH_EDITOR
