Fixed static analysis warning.

#rnx

diff --git a/Engine/Plugins/Animation/GameplayInsights/Source/GameplayInsights/Private/SAnimGraphSchematicView.cpp b/Engine/Plugins/Animation/GameplayInsights/Source/GameplayInsights/Private/SAnimGraphSchematicView.cpp
--- a/Engine/Plugins/Animation/GameplayInsights/Source/GameplayInsights/Private/SAnimGraphSchematicView.cpp
+++ b/Engine/Plugins/Animation/GameplayInsights/Source/GameplayInsights/Private/SAnimGraphSchematicView.cpp
@@ -813,14 +813,9 @@ TSharedRef<SWidget> SAnimGraphSchematicView::HandleGetViewMenuContent()
 					FUIAction(
 						FExecuteAction::CreateLambda([this, ColumnId]()
 						{
-							FColumnState* StatePtr = Columns.Find(ColumnId);
-							if(StatePtr != nullptr && StatePtr->bEnabled)
+							if(FColumnState* StatePtr = Columns.Find(ColumnId))
 							{
-								StatePtr->bEnabled = false;
-							}
-							else
-							{
-								StatePtr->bEnabled = true;
+								StatePtr->bEnabled = !StatePtr->bEnabled;
 							}
 
 							RefreshColumns();
