Fix crash closing the editor with Gameplay Insights docked and an anim graph document tab open

diff --git a/Engine/Plugins/Animation/GameplayInsights/Source/GameplayInsights/Private/AnimationSharedData.cpp b/Engine/Plugins/Animation/GameplayInsights/Source/GameplayInsights/Private/AnimationSharedData.cpp
--- a/Engine/Plugins/Animation/GameplayInsights/Source/GameplayInsights/Private/AnimationSharedData.cpp
+++ b/Engine/Plugins/Animation/GameplayInsights/Source/GameplayInsights/Private/AnimationSharedData.cpp
@@ -38,17 +38,6 @@ FAnimationSharedData::FAnimationSharedData(FGameplaySharedData& InGameplayShared
 {
 }
 
-FAnimationSharedData::~FAnimationSharedData()
-{
-	for(TWeakPtr<SDockTab> DocumentTab : OpenDocumentTabs)
-	{
-		if(DocumentTab.IsValid())
-		{
-			DocumentTab.Pin()->RequestCloseTab();
-		}
-	}
-}
-
 void FAnimationSharedData::OnBeginSession(Insights::ITimingViewSession& InTimingViewSession)
 {
 	TimingViewSession = &InTimingViewSession;
@@ -497,8 +486,6 @@ void FAnimationSharedData::OpenAnimGraphTab(uint64 InAnimInstanceId) const
 			const FObjectInfo& ObjectInfo = GameplayProvider->GetObjectInfo(InAnimInstanceId);
 			Tab->SetLabel(FText::FromString(ObjectInfo.Name));
 		}
-
-		OpenDocumentTabs.Add(Tab);
 	}
 }
 
diff --git a/Engine/Plugins/Animation/GameplayInsights/Source/GameplayInsights/Private/AnimationSharedData.h b/Engine/Plugins/Animation/GameplayInsights/Source/GameplayInsights/Private/AnimationSharedData.h
--- a/Engine/Plugins/Animation/GameplayInsights/Source/GameplayInsights/Private/AnimationSharedData.h
+++ b/Engine/Plugins/Animation/GameplayInsights/Source/GameplayInsights/Private/AnimationSharedData.h
@@ -23,8 +23,6 @@ class FAnimationSharedData
 public:
 	FAnimationSharedData(FGameplaySharedData& InGameplaySharedData);
 
-	~FAnimationSharedData();
-
 	void OnBeginSession(Insights::ITimingViewSession& InTimingViewSession);
 	void OnEndSession(Insights::ITimingViewSession& InTimingViewSession);
 	void Tick(Insights::ITimingViewSession& InTimingViewSession, const Trace::IAnalysisSession& InAnalysisSession);
@@ -100,9 +98,6 @@ private:
 	/** Various times and ranges */
 	double MarkerTime;
 
-	/** The open document tabs we reference */
-	mutable TArray<TWeakPtr<SDockTab>> OpenDocumentTabs;
-
 	/** Validity flags for pose times/ranges */
 	bool bTimeMarkerValid;
 
