GameplayInsights creates its own store.

diff --git a/Engine/Plugins/Animation/GameplayInsights/Source/GameplayInsights/GameplayInsights.Build.cs b/Engine/Plugins/Animation/GameplayInsights/Source/GameplayInsights/GameplayInsights.Build.cs
--- a/Engine/Plugins/Animation/GameplayInsights/Source/GameplayInsights/GameplayInsights.Build.cs
+++ b/Engine/Plugins/Animation/GameplayInsights/Source/GameplayInsights/GameplayInsights.Build.cs
@@ -13,6 +13,7 @@ namespace UnrealBuildTool.Rules
 				"InputCore",
 				"SlateCore",
 				"Slate",
+				"TraceLog",
 				"TraceAnalysis",
 				"TraceServices",
 				"TraceInsights",
diff --git a/Engine/Plugins/Animation/GameplayInsights/Source/GameplayInsights/Private/GameplayInsightsModule.cpp b/Engine/Plugins/Animation/GameplayInsights/Source/GameplayInsights/Private/GameplayInsightsModule.cpp
--- a/Engine/Plugins/Animation/GameplayInsights/Source/GameplayInsights/Private/GameplayInsightsModule.cpp
+++ b/Engine/Plugins/Animation/GameplayInsights/Source/GameplayInsights/Private/GameplayInsightsModule.cpp
@@ -14,6 +14,8 @@
 #include "TraceServices/ITraceServicesModule.h"
 #include "TraceServices/SessionService.h"
 #include "Widgets/Docking/SDockTab.h"
+#include "Trace/StoreService.h"
+#include "Trace/StoreClient.h"
 
 #if WITH_EDITOR
 #include "IAnimationBlueprintEditorModule.h"
@@ -48,28 +50,24 @@ void FGameplayInsightsModule::StartupModule()
 		GameplayTimingViewExtender.GetCustomDebugObjects(InAnimationBlueprintEditor, OutDebugList);
 	});
 
-	ITraceServicesModule& TraceServicesModule = FModuleManager::LoadModuleChecked<ITraceServicesModule>("TraceServices");
-
-	// Connect to loopback session
-	TSharedPtr<Trace::ISessionService> SessionService = TraceServicesModule.GetSessionService();
-	if(SessionService.IsValid())
-	{
-		// Connect session
-		SessionService->ConnectSession(TEXT("127.0.0.1"));
-
-		// Wait a second then attempt a connection - it takes a bit of time for sessions to get populated
-		FTicker::GetCoreTicker().AddTicker(FTickerDelegate::CreateLambda([](float InDeltaSeconds)
-		{
-			IUnrealInsightsModule& UnrealInsightsModule = FModuleManager::LoadModuleChecked<IUnrealInsightsModule>("TraceInsights");
-
-			// Start analysis on the latest live session
-			UnrealInsightsModule.StartAnalysisForLastLiveSession();
+	IUnrealInsightsModule& UnrealInsightsModule = FModuleManager::LoadModuleChecked<IUnrealInsightsModule>("TraceInsights");
 
-			return false;
-		}), 1.0f);
-	}
+	// Create the Store Service.
+	FString StoreDir = FPaths::ProjectSavedDir() / TEXT("TraceSessions");
+	Trace::FStoreService::FDesc StoreServiceDesc;
+	StoreServiceDesc.StoreDir = *StoreDir;
+	StoreServiceDesc.RecorderPort = 0; // Let system decide port
+	StoreServiceDesc.ThreadCount = 2;
+	StoreService = TSharedPtr<Trace::FStoreService>(Trace::FStoreService::Create(StoreServiceDesc));
+
+	FCoreDelegates::OnPreExit.AddLambda([this]() {
+		StoreClient.Reset();
+		StoreService.Reset();
+	});
 
-	IUnrealInsightsModule& UnrealInsightsModule = FModuleManager::LoadModuleChecked<IUnrealInsightsModule>("TraceInsights");
+	// Connect to our newly created store and setup the insights module
+	UnrealInsightsModule.ConnectToStore(TEXT("localhost"), StoreService->GetPort());
+	Trace::SendTo(TEXT("localhost"), StoreService->GetRecorderPort());
 
 	const float DPIScaleFactor = FPlatformApplicationMisc::GetDPIScaleFactorAtPoint(10.0f, 10.0f);
 
@@ -137,6 +135,7 @@ void FGameplayInsightsModule::StartupModule()
 	UnrealInsightsModule.RegisterMajorTabConfig(FInsightsManagerTabs::NetworkingProfilerTabId, FInsightsMajorTabConfig::Unavailable());
 
 	UnrealInsightsModule.CreateSessionViewer(false);
+	UnrealInsightsModule.StartAnalysisForLastLiveSession();
 #endif
 }
 
diff --git a/Engine/Plugins/Animation/GameplayInsights/Source/GameplayInsights/Private/GameplayInsightsModule.h b/Engine/Plugins/Animation/GameplayInsights/Source/GameplayInsights/Private/GameplayInsightsModule.h
--- a/Engine/Plugins/Animation/GameplayInsights/Source/GameplayInsights/Private/GameplayInsightsModule.h
+++ b/Engine/Plugins/Animation/GameplayInsights/Source/GameplayInsights/Private/GameplayInsightsModule.h
@@ -8,6 +8,12 @@
 #include "GameplayTimingViewExtender.h"
 #include "Framework/Docking/TabManager.h"
 
+namespace Trace 
+{
+	class FStoreService;
+	class FStoreClient;
+}
+
 class FGameplayInsightsModule : public IModuleInterface
 {
 public:
@@ -29,5 +35,7 @@ private:
 	FDelegateHandle CustomDebugObjectHandle;
 #endif
 
+	TSharedPtr<Trace::FStoreService> StoreService;
+	TSharedPtr<Trace::FStoreClient> StoreClient;
 	TWeakPtr<FTabManager> WeakTimingProfilerTabManager;
 };
diff --git a/Engine/Source/Developer/TraceInsights/Private/Insights/InsightsManager.cpp b/Engine/Source/Developer/TraceInsights/Private/Insights/InsightsManager.cpp
--- a/Engine/Source/Developer/TraceInsights/Private/Insights/InsightsManager.cpp
+++ b/Engine/Source/Developer/TraceInsights/Private/Insights/InsightsManager.cpp
@@ -38,8 +38,6 @@ FInsightsManager::FInsightsManager(TSharedRef<Trace::IAnalysisService> InTraceAn
 	, SessionService(InTraceSessionService)
 	, ModuleService(InTraceModuleService)
 	, StoreDir()
-	, StoreHost()
-	, StorePort(0)
 	, StoreClient()
 	, CommandList(new FUICommandList())
 	, ActionManager(this)
@@ -132,8 +130,6 @@ FInsightsSettings& FInsightsManager::GetSettings()
 bool FInsightsManager::ConnectToStore(const TCHAR* Host, uint32 Port)
 {
 	using namespace Trace;
-	StoreHost = Host;
-	StorePort = Port;
 	FStoreClient* Client = FStoreClient::Connect(Host, Port);
 	StoreClient = TUniquePtr<FStoreClient>(Client);
 	return StoreClient.IsValid();
@@ -274,7 +270,7 @@ void FInsightsManager::LoadLastLiveSession()
 {
 	ResetSession();
 
-	if (StoreClient.IsValid())
+	if (!StoreClient.IsValid())
 	{
 		return;
 	}
diff --git a/Engine/Source/Developer/TraceInsights/Private/Insights/InsightsManager.h b/Engine/Source/Developer/TraceInsights/Private/Insights/InsightsManager.h
--- a/Engine/Source/Developer/TraceInsights/Private/Insights/InsightsManager.h
+++ b/Engine/Source/Developer/TraceInsights/Private/Insights/InsightsManager.h
@@ -81,8 +81,6 @@ public:
 	const FString& GetStoreDir() const { return StoreDir; }
 
 	bool ConnectToStore(const TCHAR* Host, uint32 Port);
-	const FString& GetStoreHost() const { return StoreHost; }
-	uint32 GetStorePort() const { return StorePort; }
 	Trace::FStoreClient* GetStoreClient() const { return StoreClient.Get(); }
 
 	/** @return an instance of the trace analysis session. */
@@ -193,15 +191,9 @@ private:
 	TSharedRef<Trace::ISessionService> SessionService;
 	TSharedRef<Trace::IModuleService> ModuleService;
 
-	/** The location of the trace files managed by the trace store. */
+	/** The location of the trace files managed by trace store. */
 	FString StoreDir;
 
-	/** The host name of the trace store. */
-	FString StoreHost;
-
-	/** The port of the trace store. */
-	uint32 StorePort;
-
 	/** The client used to connect to the trace store. */
 	TUniquePtr<Trace::FStoreClient> StoreClient;
 
@@ -211,7 +203,7 @@ private:
 	/** The id of the trace being analyzed. */
 	uint32 CurrentTraceId;
 
-	/** The filename of the trace being analyzed. */
+	/** The filename of trace being analyzed. */
 	FString CurrentTraceFilename;
 
 	/** List of UI commands for this manager. This will be filled by this and corresponding classes. */
diff --git a/Engine/Source/Developer/TraceInsights/Private/Insights/TraceInsightsModule.cpp b/Engine/Source/Developer/TraceInsights/Private/Insights/TraceInsightsModule.cpp
--- a/Engine/Source/Developer/TraceInsights/Private/Insights/TraceInsightsModule.cpp
+++ b/Engine/Source/Developer/TraceInsights/Private/Insights/TraceInsightsModule.cpp
@@ -201,6 +201,13 @@ void FTraceInsightsModule::UnregisterTabSpawners()
 
 ////////////////////////////////////////////////////////////////////////////////////////////////////
 
+bool FTraceInsightsModule::ConnectToStore(const TCHAR* Host, uint32 Port)
+{
+	return FInsightsManager::Get()->ConnectToStore(Host, Port);
+}
+
+////////////////////////////////////////////////////////////////////////////////////////////////////
+
 void FTraceInsightsModule::CreateSessionBrowser(bool bAllowDebugTools, bool bSingleProcess)
 {
 	FInsightsManager::Get()->SetOpenAnalysisInSeparateProcess(!bSingleProcess);
diff --git a/Engine/Source/Developer/TraceInsights/Private/Insights/TraceInsightsModule.h b/Engine/Source/Developer/TraceInsights/Private/Insights/TraceInsightsModule.h
--- a/Engine/Source/Developer/TraceInsights/Private/Insights/TraceInsightsModule.h
+++ b/Engine/Source/Developer/TraceInsights/Private/Insights/TraceInsightsModule.h
@@ -31,6 +31,7 @@ public:
 		return false;
 	}
 
+	virtual bool ConnectToStore(const TCHAR* Host, uint32 Port) override;
 	virtual void CreateSessionBrowser(bool bAllowDebugTools, bool bSingleProcess) override;
 	virtual void CreateSessionViewer(bool bAllowDebugTools) override;
 	virtual void StartAnalysisForTraceFile(const TCHAR* InTraceFile) override;
diff --git a/Engine/Source/Developer/TraceInsights/Private/Insights/Widgets/SStartPageWindow.cpp b/Engine/Source/Developer/TraceInsights/Private/Insights/Widgets/SStartPageWindow.cpp
--- a/Engine/Source/Developer/TraceInsights/Private/Insights/Widgets/SStartPageWindow.cpp
+++ b/Engine/Source/Developer/TraceInsights/Private/Insights/Widgets/SStartPageWindow.cpp
@@ -1493,7 +1493,7 @@ void SStartPageWindow::LoadTrace(uint32 InTraceId)
 
 		const TCHAR* ExecutablePath = FPlatformProcess::ExecutablePath();
 
-		const uint32 StorePort = FInsightsManager::Get()->GetStorePort();
+		const uint32 StorePort = FInsightsManager::Get()->GetStoreClient()->GetStorePort();
 		FString CmdLine = FString::Printf(TEXT("-TraceId=%d -StorePort=%d"), InTraceId, StorePort);
 
 		constexpr bool bLaunchDetached = false;
diff --git a/Engine/Source/Developer/TraceInsights/Public/Insights/IUnrealInsightsModule.h b/Engine/Source/Developer/TraceInsights/Public/Insights/IUnrealInsightsModule.h
--- a/Engine/Source/Developer/TraceInsights/Public/Insights/IUnrealInsightsModule.h
+++ b/Engine/Source/Developer/TraceInsights/Public/Insights/IUnrealInsightsModule.h
@@ -110,6 +110,12 @@ struct TRACEINSIGHTS_API FInsightsMajorTabConfig
 class IUnrealInsightsModule : public IModuleInterface
 {
 public:
+
+	/**
+	 * Connect to specific store.
+	 */
+	virtual bool ConnectToStore(const TCHAR* Host, uint32 Port) = 0;
+
 	/**
 	 * Called when the application starts in "Browser" mode.
 	 */
diff --git a/Engine/Source/Programs/UnrealInsights/UnrealInsights.Build.cs b/Engine/Source/Programs/UnrealInsights/UnrealInsights.Build.cs
--- a/Engine/Source/Programs/UnrealInsights/UnrealInsights.Build.cs
+++ b/Engine/Source/Programs/UnrealInsights/UnrealInsights.Build.cs
@@ -22,6 +22,7 @@ public class UnrealInsights : ModuleRules
 				"SourceCodeAccess",
 				"StandaloneRenderer",
 				"TraceInsights",
+				"TraceAnalysis",
 				"TraceServices"
 			}
 		);
