New animation tracks and tools

Added tracing of anim notifies and montages
Added variant view to view how properties of tracks change over time
Name tracing for animation data is now thread safe

diff --git a/Engine/Plugins/Animation/GameplayInsights/Source/GameplayInsights/Private/AnimNodesTrack.cpp b/Engine/Plugins/Animation/GameplayInsights/Source/GameplayInsights/Private/AnimNodesTrack.cpp
--- a/Engine/Plugins/Animation/GameplayInsights/Source/GameplayInsights/Private/AnimNodesTrack.cpp
+++ b/Engine/Plugins/Animation/GameplayInsights/Source/GameplayInsights/Private/AnimNodesTrack.cpp
@@ -35,7 +35,7 @@
 INSIGHTS_IMPLEMENT_RTTI(FAnimNodesTrack)
 
 FAnimNodesTrack::FAnimNodesTrack(const FAnimationSharedData& InSharedData, uint64 InObjectID, const TCHAR* InName)
-	: TGameplayTrackMixin<FTimingEventsTrack>(InObjectID, FText::Format(LOCTEXT("TrackNameFormat", "Graph - {0}"), FText::FromString(FString(InName))))
+	: FGameplayTimingEventsTrack(InSharedData.GetGameplaySharedData(), InObjectID, FText::Format(LOCTEXT("TrackNameFormat", "Graph - {0}"), FText::FromString(FString(InName))))
 	, SharedData(InSharedData)
 #if WITH_EDITOR
 	, InstanceClass(nullptr)
diff --git a/Engine/Plugins/Animation/GameplayInsights/Source/GameplayInsights/Private/AnimNodesTrack.h b/Engine/Plugins/Animation/GameplayInsights/Source/GameplayInsights/Private/AnimNodesTrack.h
--- a/Engine/Plugins/Animation/GameplayInsights/Source/GameplayInsights/Private/AnimNodesTrack.h
+++ b/Engine/Plugins/Animation/GameplayInsights/Source/GameplayInsights/Private/AnimNodesTrack.h
@@ -18,15 +18,14 @@ class AActor;
 class USkeletalMeshComponent;
 class UAnimInstance;
 
-class FAnimNodesTrack : public TGameplayTrackMixin<FTimingEventsTrack>
+class FAnimNodesTrack : public FGameplayTimingEventsTrack
 #if WITH_ENGINE
 	, public FGCObject
 #endif
 {
-public:
-	static const FName TypeName;
-	static const FName SubTypeName;
+	INSIGHTS_DECLARE_RTTI(FAnimNodesTrack, FGameplayTimingEventsTrack)
 
+public:
 	FAnimNodesTrack(const FAnimationSharedData& InSharedData, uint64 InObjectID, const TCHAR* InName);
 
 	virtual void BuildDrawState(ITimingEventsTrackDrawStateBuilder& Builder, const ITimingTrackUpdateContext& Context) override;
diff --git a/Engine/Plugins/Animation/GameplayInsights/Source/GameplayInsights/Private/AnimNotifiesTrack.cpp b/Engine/Plugins/Animation/GameplayInsights/Source/GameplayInsights/Private/AnimNotifiesTrack.cpp
new file mode 100644
--- /dev/null
+++ b/Engine/Plugins/Animation/GameplayInsights/Source/GameplayInsights/Private/AnimNotifiesTrack.cpp
@@ -0,0 +1,259 @@
+// Copyright Epic Games, Inc. All Rights Reserved.
+
+#include "AnimNotifiesTrack.h"
+#include "Insights/ViewModels/ITimingViewDrawHelper.h"
+#include "GameplayProvider.h"
+#include "AnimationProvider.h"
+#include "Insights/ViewModels/TimingTrackViewport.h"
+#include "Insights/ViewModels/TimingEvent.h"
+#include "AnimationSharedData.h"
+#include "Framework/MultiBox/MultiBoxBuilder.h"
+#include "Insights/ViewModels/TimingEventSearch.h"
+#include "Insights/ViewModels/TooltipDrawState.h"
+#include "Insights/Common/TimeUtils.h"
+#include "VariantTreeNode.h"
+#include "TraceServices/Model/Frames.h"
+
+#define LOCTEXT_NAMESPACE "AnimNotifiesTrack"
+
+INSIGHTS_IMPLEMENT_RTTI(FAnimNotifiesTrack)
+
+FAnimNotifiesTrack::FAnimNotifiesTrack(const FAnimationSharedData& InSharedData, uint64 InObjectID, const TCHAR* InName)
+	: FGameplayTimingEventsTrack(InSharedData.GetGameplaySharedData(), InObjectID, FText::Format(LOCTEXT("TrackNameFormat", "Notifies - {0}"), FText::FromString(FString(InName))))
+	, SharedData(InSharedData)
+{
+}
+
+void FAnimNotifiesTrack::BuildDrawState(ITimingEventsTrackDrawStateBuilder& Builder, const ITimingTrackUpdateContext& Context)
+{
+	const FAnimationProvider* AnimationProvider = SharedData.GetAnalysisSession().ReadProvider<FAnimationProvider>(FAnimationProvider::ProviderName);
+
+	if(AnimationProvider)
+	{
+		Trace::FAnalysisSessionReadScope SessionReadScope(SharedData.GetAnalysisSession());
+
+		bHasNotifies = false;
+
+		AnimationProvider->ReadNotifyTimeline(GetGameplayTrack().GetObjectId(), [this, &AnimationProvider, &Context, &Builder](const FAnimationProvider::AnimNotifyTimeline& InTimeline)
+		{
+			InTimeline.EnumerateEvents(Context.GetViewport().GetStartTime(), Context.GetViewport().GetEndTime(), [this, &AnimationProvider, &Builder](double InStartTime, double InEndTime, uint32 InDepth, const FAnimNotifyMessage& InMessage)
+			{
+				Builder.AddEvent(InStartTime, InEndTime, 0, AnimationProvider->GetName(InMessage.NameId));
+				bHasNotifies = true;
+			});
+		});
+
+		uint32 Depth = bHasNotifies ? 1 : 0;
+		IdToDepthMap.Reset();
+
+		AnimationProvider->EnumerateNotifyStateTimelines(GetGameplayTrack().GetObjectId(), [this, &AnimationProvider, &Context, &Builder, &Depth](uint64 InNotifyId, const FAnimationProvider::AnimNotifyTimeline& InTimeline)
+		{
+			uint32 DepthToUse = 0;
+			if(uint32* FoundDepthPtr = IdToDepthMap.Find(InNotifyId))
+			{
+				DepthToUse = *FoundDepthPtr;
+			}
+			else
+			{
+				IdToDepthMap.Add(InNotifyId, ++Depth);
+				DepthToUse = Depth;
+			}
+
+			InTimeline.EnumerateEvents(Context.GetViewport().GetStartTime(), Context.GetViewport().GetEndTime(), [&AnimationProvider, &Builder, &DepthToUse](double InStartTime, double InEndTime, uint32 InDepth, const FAnimNotifyMessage& InMessage)
+			{
+				Builder.AddEvent(InStartTime, InEndTime, DepthToUse, AnimationProvider->GetName(InMessage.NameId));
+			});
+		});
+	}
+}
+
+void FAnimNotifiesTrack::Draw(const ITimingTrackDrawContext& Context) const
+{
+	DrawEvents(Context);
+	GetGameplayTrack().DrawHeaderForTimingTrack(Context, *this, false);
+}
+
+void FAnimNotifiesTrack::InitTooltip(FTooltipDrawState& Tooltip, const ITimingEvent& HoveredTimingEvent) const
+{
+	auto EventFilter = [&HoveredTimingEvent](double EventStartTime, double EventEndTime, uint32 EventDepth) 
+	{ 
+		return HoveredTimingEvent.GetDepth() == EventDepth;
+	};
+
+	FTimingEventSearchParameters SearchParameters(HoveredTimingEvent.GetStartTime(), HoveredTimingEvent.GetEndTime(), ETimingEventSearchFlags::StopAtFirstMatch, EventFilter);
+
+	FindAnimNotifyMessage(SearchParameters, [this, &Tooltip](double InFoundStartTime, double InFoundEndTime, uint32 InFoundDepth, const FAnimNotifyMessage& InMessage)
+	{
+		Tooltip.ResetContent();
+
+		Tooltip.AddTitle(GetName());
+
+		{
+			Trace::FAnalysisSessionReadScope SessionReadScope(SharedData.GetAnalysisSession());
+
+			const FAnimationProvider* AnimationProvider = SharedData.GetAnalysisSession().ReadProvider<FAnimationProvider>(FAnimationProvider::ProviderName);
+			if(AnimationProvider)
+			{
+				const TCHAR* Name = AnimationProvider->GetName(InMessage.NameId);
+				Tooltip.AddNameValueTextLine(LOCTEXT("NotifyName", "Name").ToString(), Name);
+			}
+		}
+
+		Tooltip.AddNameValueTextLine(LOCTEXT("EventDuration", "Duration").ToString(), TimeUtils::FormatTimeAuto(InFoundEndTime - InFoundStartTime));
+		Tooltip.AddNameValueTextLine(LOCTEXT("EventTime", "Time").ToString(), TimeUtils::FormatTimeAuto(InFoundStartTime));
+
+		if(InMessage.NotifyEventType != EAnimNotifyMessageType::SyncMarker)
+		{
+			Trace::FAnalysisSessionReadScope SessionReadScope(SharedData.GetAnalysisSession());
+
+			const FGameplayProvider* GameplayProvider = SharedData.GetAnalysisSession().ReadProvider<FGameplayProvider>(FGameplayProvider::ProviderName);
+			if(GameplayProvider)
+			{
+				const FObjectInfo& AssetInfo = GameplayProvider->GetObjectInfo(InMessage.AssetId);
+				Tooltip.AddNameValueTextLine(LOCTEXT("Asset", "Asset").ToString(), AssetInfo.PathName);
+
+				const FObjectInfo& NotifyInfo = GameplayProvider->GetObjectInfo(InMessage.NotifyId);
+				const FClassInfo& NotifyClassInfo = GameplayProvider->GetClassInfo(NotifyInfo.ClassId);
+				bool bIsState = InMessage.NotifyEventType == EAnimNotifyMessageType::Begin || InMessage.NotifyEventType == EAnimNotifyMessageType::End || InMessage.NotifyEventType == EAnimNotifyMessageType::Tick;
+				Tooltip.AddNameValueTextLine((bIsState ? LOCTEXT("Notify State Class", "Notify State Class") : LOCTEXT("NotifyClass", "Notify Class")).ToString(), NotifyClassInfo.PathName);
+			}
+		}
+
+		Tooltip.UpdateLayout();
+	});
+}
+
+const TSharedPtr<const ITimingEvent> FAnimNotifiesTrack::SearchEvent(const FTimingEventSearchParameters& InSearchParameters) const
+{
+	TSharedPtr<const ITimingEvent> FoundEvent;
+
+	FindAnimNotifyMessage(InSearchParameters, [this, &FoundEvent](double InFoundStartTime, double InFoundEndTime, uint32 InFoundDepth, const FAnimNotifyMessage& InFoundMessage)
+	{
+		FoundEvent = MakeShared<const FTimingEvent>(SharedThis(this), InFoundStartTime, InFoundEndTime, InFoundDepth);
+	});
+
+	return FoundEvent;
+}
+
+void FAnimNotifiesTrack::FindAnimNotifyMessage(const FTimingEventSearchParameters& InParameters, TFunctionRef<void(double, double, uint32, const FAnimNotifyMessage&)> InFoundPredicate) const
+{
+	TTimingEventSearch<FAnimNotifyMessage>::Search(
+		InParameters,
+
+		[this](TTimingEventSearch<FAnimNotifyMessage>::FContext& InContext)
+		{
+			const FAnimationProvider* AnimationProvider = SharedData.GetAnalysisSession().ReadProvider<FAnimationProvider>(FAnimationProvider::ProviderName);
+
+			if(AnimationProvider)
+			{
+				Trace::FAnalysisSessionReadScope SessionReadScope(SharedData.GetAnalysisSession());
+
+				if(bHasNotifies)
+				{
+					AnimationProvider->ReadNotifyTimeline(GetGameplayTrack().GetObjectId(), [this, &InContext](const FAnimationProvider::AnimNotifyTimeline& InTimeline)
+					{
+						InTimeline.EnumerateEvents(InContext.GetParameters().StartTime, InContext.GetParameters().EndTime, [this, &InContext](double InEventStartTime, double InEventEndTime, uint32 InDepth, const FAnimNotifyMessage& InMessage)
+						{
+							InContext.Check(InEventStartTime, InEventEndTime, 0, InMessage);
+						});
+					});
+				}
+
+				if(IdToDepthMap.Num() > 0)
+				{
+					AnimationProvider->EnumerateNotifyStateTimelines(GetGameplayTrack().GetObjectId(), [this, &InContext](uint32 InNotifyNameId, const FAnimationProvider::AnimNotifyTimeline& InTimeline)
+					{
+						InTimeline.EnumerateEvents(InContext.GetParameters().StartTime, InContext.GetParameters().EndTime, [this, &InContext, InNotifyNameId](double InEventStartTime, double InEventEndTime, uint32 InDepth, const FAnimNotifyMessage& InMessage)
+						{
+							if(const uint32* FoundDepthPtr = IdToDepthMap.Find(InMessage.NotifyId))
+							{
+								InContext.Check(InEventStartTime, InEventEndTime, *FoundDepthPtr, InMessage);
+							}
+						});
+					});
+				}
+			}
+		},
+
+		[&InFoundPredicate](double InFoundStartTime, double InFoundEndTime, uint32 InFoundDepth, const FAnimNotifyMessage& InEvent)
+		{
+			InFoundPredicate(InFoundStartTime, InFoundEndTime, InFoundDepth, InEvent);
+		});
+}
+
+static const TCHAR* GetMessageType(EAnimNotifyMessageType InType)
+{
+	switch(InType)
+	{
+	case EAnimNotifyMessageType::Begin:
+	case EAnimNotifyMessageType::End:
+	case EAnimNotifyMessageType::Tick:
+	{
+		static FText Type = LOCTEXT("AnimNotifyState", "Anim Notify State");
+		return *Type.ToString();
+	}
+	case EAnimNotifyMessageType::Event:
+	{
+		static FText Type = LOCTEXT("AnimNotify", "Anim Notify");
+		return *Type.ToString();
+	}
+	case EAnimNotifyMessageType::SyncMarker:
+	{
+		static FText Type = LOCTEXT("SyncMarker", "Sync Marker");
+		return *Type.ToString();
+	}
+	}
+
+	static FText UnknownType = LOCTEXT("UnknownType", "Unknown");
+	return *UnknownType.ToString();
+}
+
+void FAnimNotifiesTrack::GetVariantsAtTime(double InTime, TArray<TSharedRef<FVariantTreeNode>>& OutVariants) const 
+{
+	TSharedRef<FVariantTreeNode> Header = OutVariants.Add_GetRef(FVariantTreeNode::MakeHeader(LOCTEXT("AnimNotifiesHeader", "Notifies and Sync Markers")));
+
+	const Trace::IFrameProvider& FramesProvider = Trace::ReadFrameProvider(SharedData.GetAnalysisSession());
+	const FAnimationProvider* AnimationProvider = SharedData.GetAnalysisSession().ReadProvider<FAnimationProvider>(FAnimationProvider::ProviderName);
+	const FGameplayProvider* GameplayProvider = SharedData.GetAnalysisSession().ReadProvider<FGameplayProvider>(FGameplayProvider::ProviderName);
+
+	if(AnimationProvider && GameplayProvider)
+	{
+		Trace::FAnalysisSessionReadScope SessionReadScope(SharedData.GetAnalysisSession());
+
+		// round to nearest frame boundary
+		Trace::FFrame Frame;
+		if(FramesProvider.GetFrameFromTime(ETraceFrameType::TraceFrameType_Game, InTime, Frame))
+		{
+			auto ProcessEvent = [this, &AnimationProvider, &GameplayProvider, &Header](double InStartTime, double InEndTime, uint32 InDepth, const FAnimNotifyMessage& InMessage)
+			{
+				const TCHAR* Name = AnimationProvider->GetName(InMessage.NameId);
+				TSharedRef<FVariantTreeNode> NotifyHeader = Header->AddChild(FVariantTreeNode::MakeHeader(FText::FromString(Name)));
+
+				NotifyHeader->AddChild(FVariantTreeNode::MakeString(LOCTEXT("EventType", "Type"), GetMessageType(InMessage.NotifyEventType)));
+				NotifyHeader->AddChild(FVariantTreeNode::MakeFloat(LOCTEXT("EventDuration", "Duration"), InEndTime - InStartTime));
+				NotifyHeader->AddChild(FVariantTreeNode::MakeFloat(LOCTEXT("EventTime", "Time"), InStartTime));
+
+				if(InMessage.NotifyEventType != EAnimNotifyMessageType::SyncMarker)
+				{
+					NotifyHeader->AddChild(FVariantTreeNode::MakeObject(LOCTEXT("Asset", "Asset"), InMessage.AssetId));
+
+					const FObjectInfo& NotifyInfo = GameplayProvider->GetObjectInfo(InMessage.NotifyId);
+					bool bIsState = InMessage.NotifyEventType == EAnimNotifyMessageType::Begin || InMessage.NotifyEventType == EAnimNotifyMessageType::End || InMessage.NotifyEventType == EAnimNotifyMessageType::Tick;
+					NotifyHeader->AddChild(FVariantTreeNode::MakeClass((bIsState ? LOCTEXT("Notify State Class", "Notify State Class") : LOCTEXT("NotifyClass", "Notify Class")), NotifyInfo.ClassId));
+				}
+			};
+
+			AnimationProvider->ReadNotifyTimeline(GetGameplayTrack().GetObjectId(), [&Frame, &ProcessEvent](const FAnimationProvider::AnimNotifyTimeline& InTimeline)
+			{
+				InTimeline.EnumerateEvents(Frame.StartTime, Frame.EndTime, ProcessEvent);
+			});
+
+			AnimationProvider->EnumerateNotifyStateTimelines(GetGameplayTrack().GetObjectId(), [&Frame, &ProcessEvent](uint32 InNotifyNameId, const FAnimationProvider::AnimNotifyTimeline& InTimeline)
+			{
+				InTimeline.EnumerateEvents(Frame.StartTime, Frame.EndTime, ProcessEvent);
+			});
+		}
+	}
+}
+
+#undef LOCTEXT_NAMESPACE
\ No newline at end of file
diff --git a/Engine/Plugins/Animation/GameplayInsights/Source/GameplayInsights/Private/AnimNotifiesTrack.h b/Engine/Plugins/Animation/GameplayInsights/Source/GameplayInsights/Private/AnimNotifiesTrack.h
new file mode 100644
--- /dev/null
+++ b/Engine/Plugins/Animation/GameplayInsights/Source/GameplayInsights/Private/AnimNotifiesTrack.h
@@ -0,0 +1,38 @@
+// Copyright Epic Games, Inc. All Rights Reserved.
+
+#pragma once
+
+#include "GameplayTrack.h"
+#include "Insights/ViewModels/TimingEventsTrack.h"
+
+class FAnimationSharedData;
+class FTimingEventSearchParameters;
+struct FAnimNotifyMessage;
+
+class FAnimNotifiesTrack : public FGameplayTimingEventsTrack
+{
+	INSIGHTS_DECLARE_RTTI(FAnimNotifiesTrack, FGameplayTimingEventsTrack)
+
+public:
+	FAnimNotifiesTrack(const FAnimationSharedData& InSharedData, uint64 InObjectID, const TCHAR* InName);
+
+	virtual void BuildDrawState(ITimingEventsTrackDrawStateBuilder& Builder, const ITimingTrackUpdateContext& Context) override;
+	virtual void Draw(const ITimingTrackDrawContext& Context) const override;
+	virtual void InitTooltip(FTooltipDrawState& Tooltip, const ITimingEvent& HoveredTimingEvent) const override;
+	virtual const TSharedPtr<const ITimingEvent> SearchEvent(const FTimingEventSearchParameters& InSearchParameters) const override;
+	virtual void GetVariantsAtTime(double InTime, TArray<TSharedRef<FVariantTreeNode>>& OutVariants) const override;
+
+private:
+	// Helper function used to find an anim notify message
+	void FindAnimNotifyMessage(const FTimingEventSearchParameters& InParameters, TFunctionRef<void(double, double, uint32, const FAnimNotifyMessage&)> InFoundPredicate) const;
+
+private:
+	// The shared data
+	const FAnimationSharedData& SharedData;
+
+	// Map of notify ID->depth of notify states that are currently drawn
+	TMap<uint64, uint32> IdToDepthMap;
+
+	// Whether we are currently drawing any notifies
+	bool bHasNotifies;
+};
\ No newline at end of file
diff --git a/Engine/Plugins/Animation/GameplayInsights/Source/GameplayInsights/Private/AnimationAnalyzer.cpp b/Engine/Plugins/Animation/GameplayInsights/Source/GameplayInsights/Private/AnimationAnalyzer.cpp
--- a/Engine/Plugins/Animation/GameplayInsights/Source/GameplayInsights/Private/AnimationAnalyzer.cpp
+++ b/Engine/Plugins/Animation/GameplayInsights/Source/GameplayInsights/Private/AnimationAnalyzer.cpp
@@ -26,6 +26,7 @@ void FAnimationAnalyzer::OnAnalysisBegin(const FOnAnalysisContext& Context)
 	Builder.RouteEvent(RouteId_AnimNodeValueBool, "Animation", "AnimNodeValueBool");
 	Builder.RouteEvent(RouteId_AnimNodeValueInt, "Animation", "AnimNodeValueInt");
 	Builder.RouteEvent(RouteId_AnimNodeValueFloat, "Animation", "AnimNodeValueFloat");
+	Builder.RouteEvent(RouteId_AnimNodeValueVector2D, "Animation", "AnimNodeValueVector2D");
 	Builder.RouteEvent(RouteId_AnimNodeValueVector, "Animation", "AnimNodeValueVector");
 	Builder.RouteEvent(RouteId_AnimNodeValueString, "Animation", "AnimNodeValueString");
 	Builder.RouteEvent(RouteId_AnimNodeValueObject, "Animation", "AnimNodeValueObject");
@@ -34,6 +35,9 @@ void FAnimationAnalyzer::OnAnalysisBegin(const FOnAnalysisContext& Context)
 	Builder.RouteEvent(RouteId_BlendSpacePlayer, "Animation", "BlendSpacePlayer");
 	Builder.RouteEvent(RouteId_StateMachineState, "Animation", "StateMachineState");
 	Builder.RouteEvent(RouteId_Name, "Animation", "Name");
+	Builder.RouteEvent(RouteId_Notify, "Animation", "Notify");
+	Builder.RouteEvent(RouteId_SyncMarker, "Animation", "SyncMarker");
+	Builder.RouteEvent(RouteId_Montage, "Animation", "Montage");
 }
 
 bool FAnimationAnalyzer::OnEvent(uint16 RouteId, const FOnEventContext& Context)
@@ -185,6 +189,12 @@ bool FAnimationAnalyzer::OnEvent(uint16 RouteId, const FOnEventContext& Context)
 			AnimationProvider.AppendAnimNodeValue(AnimInstanceId, Context.SessionContext.TimestampFromCycle(Cycle), FrameCounter, NodeId, Key, Value);
 			break;
 		}
+		case RouteId_AnimNodeValueVector2D:
+		{
+			FVector2D Value(EventData.GetValue<float>("ValueX"), EventData.GetValue<float>("ValueY"));
+			AnimationProvider.AppendAnimNodeValue(AnimInstanceId, Context.SessionContext.TimestampFromCycle(Cycle), FrameCounter, NodeId, Key, Value);
+			break;
+		}
 		case RouteId_AnimNodeValueVector:
 		{
 			FVector Value(EventData.GetValue<float>("ValueX"), EventData.GetValue<float>("ValueY"), EventData.GetValue<float>("ValueZ"));
@@ -247,6 +257,39 @@ bool FAnimationAnalyzer::OnEvent(uint16 RouteId, const FOnEventContext& Context)
 		AnimationProvider.AppendStateMachineState(AnimInstanceId, Context.SessionContext.TimestampFromCycle(Cycle), NodeId, StateMachineIndex, StateIndex, StateWeight, ElapsedTime);
 		break;
 	}
+	case RouteId_Notify:
+	{
+		uint64 Cycle = EventData.GetValue<uint64>("Cycle");
+		uint64 AnimInstanceId = EventData.GetValue<uint64>("AnimInstanceId");
+		uint64 AssetId = EventData.GetValue<uint64>("AssetId");
+		uint64 NotifyId = EventData.GetValue<uint64>("NotifyId");
+		uint32 NameId = EventData.GetValue<uint32>("NameId");
+		float Time = EventData.GetValue<float>("Time");
+		float Duration = EventData.GetValue<float>("Duration");
+		uint8 NotifyEventType = EventData.GetValue<uint8>("NotifyEventType");
+		AnimationProvider.AppendNotify(AnimInstanceId, Context.SessionContext.TimestampFromCycle(Cycle), AssetId, NotifyId, NameId, Time, Duration, (EAnimNotifyMessageType)NotifyEventType);
+		break;
+	}
+	case RouteId_SyncMarker:
+	{
+		uint64 Cycle = EventData.GetValue<uint64>("Cycle");
+		uint64 AnimInstanceId = EventData.GetValue<uint64>("AnimInstanceId");
+		uint32 NameId = EventData.GetValue<uint32>("NameId");
+		AnimationProvider.AppendNotify(AnimInstanceId, Context.SessionContext.TimestampFromCycle(Cycle), 0, 0, NameId, 0.0f, 0.0f, EAnimNotifyMessageType::SyncMarker);
+		break;
+	}
+	case RouteId_Montage:
+	{
+		uint64 Cycle = EventData.GetValue<uint64>("Cycle");
+		uint64 AnimInstanceId = EventData.GetValue<uint64>("AnimInstanceId");
+		uint64 MontageId = EventData.GetValue<uint64>("MontageId");
+		uint32 CurrentSectionNameId = EventData.GetValue<uint32>("CurrentSectionNameId");
+		uint32 NextSectionNameId = EventData.GetValue<uint32>("NextSectionNameId");
+		float Weight = EventData.GetValue<float>("Weight");
+		float DesiredWeight = EventData.GetValue<float>("DesiredWeight");
+		AnimationProvider.AppendMontage(AnimInstanceId, Context.SessionContext.TimestampFromCycle(Cycle), MontageId, CurrentSectionNameId, NextSectionNameId, Weight, DesiredWeight);
+		break;
+	}
 	}
 
 	return true;
diff --git a/Engine/Plugins/Animation/GameplayInsights/Source/GameplayInsights/Private/AnimationAnalyzer.h b/Engine/Plugins/Animation/GameplayInsights/Source/GameplayInsights/Private/AnimationAnalyzer.h
--- a/Engine/Plugins/Animation/GameplayInsights/Source/GameplayInsights/Private/AnimationAnalyzer.h
+++ b/Engine/Plugins/Animation/GameplayInsights/Source/GameplayInsights/Private/AnimationAnalyzer.h
@@ -30,6 +30,7 @@ private:
 		RouteId_AnimNodeValueBool,
 		RouteId_AnimNodeValueInt,
 		RouteId_AnimNodeValueFloat,
+		RouteId_AnimNodeValueVector2D,
 		RouteId_AnimNodeValueVector,
 		RouteId_AnimNodeValueString,
 		RouteId_AnimNodeValueObject,
@@ -38,6 +39,9 @@ private:
 		RouteId_BlendSpacePlayer,
 		RouteId_StateMachineState,
 		RouteId_Name,
+		RouteId_Notify,
+		RouteId_SyncMarker,
+		RouteId_Montage,
 	};
 
 	Trace::IAnalysisSession& Session;
diff --git a/Engine/Plugins/Animation/GameplayInsights/Source/GameplayInsights/Private/AnimationProvider.cpp b/Engine/Plugins/Animation/GameplayInsights/Source/GameplayInsights/Private/AnimationProvider.cpp
--- a/Engine/Plugins/Animation/GameplayInsights/Source/GameplayInsights/Private/AnimationProvider.cpp
+++ b/Engine/Plugins/Animation/GameplayInsights/Source/GameplayInsights/Private/AnimationProvider.cpp
@@ -231,6 +231,74 @@ bool FAnimationProvider::ReadStateMachinesTimeline(uint64 InObjectId, TFunctionR
 	return false;
 }
 
+bool FAnimationProvider::ReadNotifyTimeline(uint64 InObjectId, TFunctionRef<void(const AnimNotifyTimeline&)> Callback) const
+{
+	Session.ReadAccessCheck();
+
+	const uint32* IndexPtr = ObjectIdToAnimNotifyTimelines.Find(InObjectId);
+	if(IndexPtr != nullptr)
+	{
+		if (*IndexPtr < uint32(AnimNotifyTimelines.Num()))
+		{
+			Callback(*AnimNotifyTimelines[*IndexPtr]);
+			return true;
+		}
+	}
+
+	return false;
+}
+
+void FAnimationProvider::EnumerateNotifyStateTimelines(uint64 InObjectId, TFunctionRef<void(uint64, const AnimNotifyTimeline&)> Callback) const
+{
+	const uint32* IndexPtr = ObjectIdToAnimNotifyStateTimelines.Find(InObjectId);
+	if(IndexPtr != nullptr)
+	{
+		if (*IndexPtr < uint32(AnimNotifyStateTimelineStorage.Num()))
+		{
+			const TSharedRef<FAnimNotifyStateTimelineStorage>& TimelineStorage = AnimNotifyStateTimelineStorage[*IndexPtr];
+			for(auto NotifyIdPair : TimelineStorage->NotifyIdToAnimNotifyStateTimeline)
+			{
+				Callback(NotifyIdPair.Key, TimelineStorage->Timelines[NotifyIdPair.Value].Get());
+			}
+		}
+	}
+}
+
+bool FAnimationProvider::ReadMontageTimeline(uint64 InObjectId, TFunctionRef<void(const AnimMontageTimeline&)> Callback) const
+{
+	Session.ReadAccessCheck();
+
+	const uint32* IndexPtr = ObjectIdToAnimMontageTimelines.Find(InObjectId);
+	if(IndexPtr != nullptr)
+	{
+		if (*IndexPtr < uint32(AnimMontageTimelineStorage.Num()))
+		{
+			Callback(*AnimMontageTimelineStorage[*IndexPtr]->Timeline.Get());
+			return true;
+		}
+	}
+
+	return false;
+}
+
+void FAnimationProvider::EnumerateMontageIds(uint64 InObjectId, TFunctionRef<void(uint64)> Callback) const
+{
+	Session.ReadAccessCheck();
+
+	const uint32* IndexPtr = ObjectIdToAnimMontageTimelines.Find(InObjectId);
+	if(IndexPtr != nullptr)
+	{
+		if (*IndexPtr < uint32(AnimMontageTimelineStorage.Num()))
+		{
+			const TSharedRef<FMontageTimelineStorage>& TimelineStorage = AnimMontageTimelineStorage[*IndexPtr];
+			for(uint64 Id : TimelineStorage->AllMontageIds)
+			{
+				Callback(Id);
+			}
+		}
+	}
+}
+
 const FSkeletalMeshInfo* FAnimationProvider::FindSkeletalMeshInfo(uint64 InObjectId) const
 {
 	Session.ReadAccessCheck();
@@ -255,7 +323,7 @@ const TCHAR* FAnimationProvider::GetName(uint32 InId) const
 		return *FoundName;
 	}
 
-	static FText UnknownText(LOCTEXT("UnknownCurve", "Unknown"));
+	static FText UnknownText(LOCTEXT("UnknownName", "Unknown"));
 	return *UnknownText.ToString();
 }
 
@@ -268,32 +336,35 @@ FText FAnimationProvider::FormatNodeValue(const FAnimNodeValueMessage& InMessage
 {
 	FText Text;
 
-	switch(InMessage.Type)
+	switch(InMessage.Value.Type)
 	{
 	case EAnimNodeValueType::Bool:
-		Text = InMessage.Bool.bValue ? LOCTEXT("True", "true") : LOCTEXT("False", "false");
+		Text = InMessage.Value.Bool.bValue ? LOCTEXT("True", "true") : LOCTEXT("False", "false");
 		break;
 	case EAnimNodeValueType::Int32:
-		Text =  FText::AsNumber(InMessage.Int32.Value);
+		Text =  FText::AsNumber(InMessage.Value.Int32.Value);
 		break;
 	case EAnimNodeValueType::Float:
-		Text = FText::AsNumber(InMessage.Float.Value);
+		Text = FText::AsNumber(InMessage.Value.Float.Value);
+		break;
+	case EAnimNodeValueType::Vector2D:
+		Text = FText::Format(LOCTEXT("Vector2Format", "({0}, {1})"), FText::AsNumber(InMessage.Value.Vector.Value.X), FText::AsNumber(InMessage.Value.Vector.Value.Y));
 		break;
 	case EAnimNodeValueType::Vector:
-		Text = FText::Format(LOCTEXT("VectorFormat", "({0}, {1}, {2})"), FText::AsNumber(InMessage.Vector.Value.X), FText::AsNumber(InMessage.Vector.Value.Y), FText::AsNumber(InMessage.Vector.Value.Z));
+		Text = FText::Format(LOCTEXT("VectorFormat", "({0}, {1}, {2})"), FText::AsNumber(InMessage.Value.Vector.Value.X), FText::AsNumber(InMessage.Value.Vector.Value.Y), FText::AsNumber(InMessage.Value.Vector.Value.Z));
 		break;
 	case EAnimNodeValueType::String:
-		Text = FText::FromString(InMessage.String.Value);
+		Text = FText::FromString(InMessage.Value.String.Value);
 		break;
 	case EAnimNodeValueType::Object:
 	{
-		const FObjectInfo& ObjectInfo = GameplayProvider.GetObjectInfo(InMessage.Object.Value);
+		const FObjectInfo& ObjectInfo = GameplayProvider.GetObjectInfo(InMessage.Value.Object.Value);
 		Text = FText::FromString(ObjectInfo.PathName);
 		break;
 	}
 	case EAnimNodeValueType::Class:
 	{
-		const FClassInfo& ClassInfo = GameplayProvider.GetClassInfo(InMessage.Class.Value);
+		const FClassInfo& ClassInfo = GameplayProvider.GetClassInfo(InMessage.Value.Class.Value);
 		Text = FText::FromString(ClassInfo.PathName);
 		break;
 	}
@@ -686,8 +757,8 @@ void FAnimationProvider::AppendBlendSpacePlayer(uint64 InAnimInstanceId, double
 void FAnimationProvider::AppendAnimNodeValue(uint64 InAnimInstanceId, double InTime, uint16 InFrameCounter, int32 InNodeId, const TCHAR* InKey, bool bInValue)
 {
 	FAnimNodeValueMessage Message;
-	Message.Bool.bValue = bInValue;
-	Message.Type = EAnimNodeValueType::Bool;
+	Message.Value.Bool.bValue = bInValue;
+	Message.Value.Type = EAnimNodeValueType::Bool;
 
 	AppendAnimNodeValue(InAnimInstanceId, InTime, InFrameCounter, InNodeId, InKey, Message);
 }
@@ -695,8 +766,8 @@ void FAnimationProvider::AppendAnimNodeValue(uint64 InAnimInstanceId, double InT
 void FAnimationProvider::AppendAnimNodeValue(uint64 InAnimInstanceId, double InTime, uint16 InFrameCounter, int32 InNodeId, const TCHAR* InKey, int32 InValue)
 {
 	FAnimNodeValueMessage Message;
-	Message.Int32.Value = InValue;
-	Message.Type = EAnimNodeValueType::Int32;
+	Message.Value.Int32.Value = InValue;
+	Message.Value.Type = EAnimNodeValueType::Int32;
 
 	AppendAnimNodeValue(InAnimInstanceId, InTime, InFrameCounter, InNodeId, InKey, Message);
 }
@@ -704,8 +775,17 @@ void FAnimationProvider::AppendAnimNodeValue(uint64 InAnimInstanceId, double InT
 void FAnimationProvider::AppendAnimNodeValue(uint64 InAnimInstanceId, double InTime, uint16 InFrameCounter, int32 InNodeId, const TCHAR* InKey, float InValue)
 {
 	FAnimNodeValueMessage Message;
-	Message.Float.Value = InValue;
-	Message.Type = EAnimNodeValueType::Float;
+	Message.Value.Float.Value = InValue;
+	Message.Value.Type = EAnimNodeValueType::Float;
+
+	AppendAnimNodeValue(InAnimInstanceId, InTime, InFrameCounter, InNodeId, InKey, Message);
+}
+
+void FAnimationProvider::AppendAnimNodeValue(uint64 InAnimInstanceId, double InTime, uint16 InFrameCounter, int32 InNodeId, const TCHAR* InKey, const FVector2D& InValue)
+{
+	FAnimNodeValueMessage Message;
+	Message.Value.Vector2D.Value = InValue;
+	Message.Value.Type = EAnimNodeValueType::Vector2D;
 
 	AppendAnimNodeValue(InAnimInstanceId, InTime, InFrameCounter, InNodeId, InKey, Message);
 }
@@ -713,8 +793,8 @@ void FAnimationProvider::AppendAnimNodeValue(uint64 InAnimInstanceId, double InT
 void FAnimationProvider::AppendAnimNodeValue(uint64 InAnimInstanceId, double InTime, uint16 InFrameCounter, int32 InNodeId, const TCHAR* InKey, const FVector& InValue)
 {
 	FAnimNodeValueMessage Message;
-	Message.Vector.Value = InValue;
-	Message.Type = EAnimNodeValueType::Vector;
+	Message.Value.Vector.Value = InValue;
+	Message.Value.Type = EAnimNodeValueType::Vector;
 
 	AppendAnimNodeValue(InAnimInstanceId, InTime, InFrameCounter, InNodeId, InKey, Message);
 }
@@ -722,8 +802,8 @@ void FAnimationProvider::AppendAnimNodeValue(uint64 InAnimInstanceId, double InT
 void FAnimationProvider::AppendAnimNodeValue(uint64 InAnimInstanceId, double InTime, uint16 InFrameCounter, int32 InNodeId, const TCHAR* InKey, const TCHAR* InValue)
 {
 	FAnimNodeValueMessage Message;
-	Message.String.Value = Session.StoreString(InValue);
-	Message.Type = EAnimNodeValueType::String;
+	Message.Value.String.Value = Session.StoreString(InValue);
+	Message.Value.Type = EAnimNodeValueType::String;
 
 	AppendAnimNodeValue(InAnimInstanceId, InTime, InFrameCounter, InNodeId, InKey, Message);
 }
@@ -731,8 +811,8 @@ void FAnimationProvider::AppendAnimNodeValue(uint64 InAnimInstanceId, double InT
 void FAnimationProvider::AppendAnimNodeValueObject(uint64 InAnimInstanceId, double InTime, uint16 InFrameCounter, int32 InNodeId, const TCHAR* InKey, uint64 InValue)
 {
 	FAnimNodeValueMessage Message;
-	Message.Object.Value = InValue;
-	Message.Type = EAnimNodeValueType::Object;
+	Message.Value.Object.Value = InValue;
+	Message.Value.Type = EAnimNodeValueType::Object;
 
 	AppendAnimNodeValue(InAnimInstanceId, InTime, InFrameCounter, InNodeId, InKey, Message);
 }
@@ -740,8 +820,8 @@ void FAnimationProvider::AppendAnimNodeValueObject(uint64 InAnimInstanceId, doub
 void FAnimationProvider::AppendAnimNodeValueClass(uint64 InAnimInstanceId, double InTime, uint16 InFrameCounter, int32 InNodeId, const TCHAR* InKey, uint64 InValue)
 {
 	FAnimNodeValueMessage Message;
-	Message.Class.Value = InValue;
-	Message.Type = EAnimNodeValueType::Class;
+	Message.Value.Class.Value = InValue;
+	Message.Value.Type = EAnimNodeValueType::Class;
 
 	AppendAnimNodeValue(InAnimInstanceId, InTime, InFrameCounter, InNodeId, InKey, Message);
 }
@@ -803,4 +883,120 @@ void FAnimationProvider::AppendStateMachineState(uint64 InAnimInstanceId, double
 	Session.UpdateDurationSeconds(InTime);
 }
 
+void FAnimationProvider::AppendNotify(uint64 InAnimInstanceId, double InTime, uint64 InAssetId, uint64 InNotifyId, uint32 InNameId, float InNotifyTime, float InNotifyDuration, EAnimNotifyMessageType InNotifyEventType)
+{
+	Session.WriteAccessCheck();
+
+	// Check if stateful or event-based
+	if(InNotifyEventType == EAnimNotifyMessageType::Begin || InNotifyEventType == EAnimNotifyMessageType::End)
+	{
+		TSharedPtr<Trace::TIntervalTimeline<FAnimNotifyMessage>> Timeline;
+		TSharedPtr<FAnimNotifyStateTimelineStorage> TimelineStorage;
+		uint32* TimelineStorageIndexPtr = ObjectIdToAnimNotifyStateTimelines.Find(InAnimInstanceId);
+		if(TimelineStorageIndexPtr != nullptr)
+		{
+			TimelineStorage = AnimNotifyStateTimelineStorage[*TimelineStorageIndexPtr];
+		}
+		else
+		{
+			ObjectIdToAnimNotifyStateTimelines.Add(InAnimInstanceId, AnimNotifyStateTimelineStorage.Num());
+			TimelineStorage = AnimNotifyStateTimelineStorage.Add_GetRef(MakeShared<FAnimNotifyStateTimelineStorage>());
+		}
+
+		check(TimelineStorage.IsValid());
+
+		uint32* TimelineIndexPtr = TimelineStorage->NotifyIdToAnimNotifyStateTimeline.Find(InNotifyId);
+		if(TimelineIndexPtr != nullptr)
+		{
+			Timeline = TimelineStorage->Timelines[*TimelineIndexPtr];
+		}
+		else
+		{
+			Timeline = MakeShared<Trace::TIntervalTimeline<FAnimNotifyMessage>>(Session.GetLinearAllocator());
+			TimelineStorage->NotifyIdToAnimNotifyStateTimeline.Add(InNotifyId, TimelineStorage->Timelines.Num());
+			TimelineStorage->Timelines.Add(Timeline.ToSharedRef());
+		}
+
+		if(InNotifyEventType == EAnimNotifyMessageType::Begin)
+		{
+			FAnimNotifyMessage Message;
+			Message.AnimInstanceId = InAnimInstanceId;
+			Message.AssetId = InAssetId;
+			Message.NotifyId = InNotifyId;
+			Message.NameId = InNameId;
+			Message.Time = InNotifyTime; 
+			Message.Duration = InNotifyDuration;
+			Message.NotifyEventType = InNotifyEventType;
+
+			Timeline->AppendBeginEvent(InTime, Message);
+		}
+		else if(Timeline->GetEventCount() > 0)
+		{
+			Timeline->EndEvent(Timeline->GetEventCount() - 1, InTime);
+		}
+	}
+	else
+	{
+		TSharedPtr<Trace::TPointTimeline<FAnimNotifyMessage>> Timeline;
+		uint32* IndexPtr = ObjectIdToAnimNotifyTimelines.Find(InAnimInstanceId);
+		if(IndexPtr != nullptr)
+		{
+			Timeline = AnimNotifyTimelines[*IndexPtr];
+		}
+		else
+		{
+			Timeline = MakeShared<Trace::TPointTimeline<FAnimNotifyMessage>>(Session.GetLinearAllocator());
+			ObjectIdToAnimNotifyTimelines.Add(InAnimInstanceId, AnimNotifyTimelines.Num());
+			AnimNotifyTimelines.Add(Timeline.ToSharedRef());
+		}
+
+		FAnimNotifyMessage Message;
+		Message.AnimInstanceId = InAnimInstanceId;
+		Message.AssetId = InAssetId;
+		Message.NotifyId = InNotifyId;
+		Message.NameId = InNameId;
+		Message.Time = InNotifyTime; 
+		Message.Duration = InNotifyDuration;
+		Message.NotifyEventType = InNotifyEventType;
+
+		Timeline->AppendEvent(InTime, Message);
+	}
+
+	Session.UpdateDurationSeconds(InTime);
+}
+
+void FAnimationProvider::AppendMontage(uint64 InAnimInstanceId, double InTime, uint64 InMontageId, uint32 InCurrentSectionNameId, uint32 InNextSectionNameId, float InWeight, float InDesiredWeight)
+{
+	Session.WriteAccessCheck();
+
+	TSharedPtr<FMontageTimelineStorage> TimelineStorage;
+	uint32* IndexPtr = ObjectIdToAnimMontageTimelines.Find(InAnimInstanceId);
+	if(IndexPtr != nullptr)
+	{
+		TimelineStorage = AnimMontageTimelineStorage[*IndexPtr];
+	}
+	else
+	{
+		TimelineStorage = MakeShared<FMontageTimelineStorage>();
+		TimelineStorage->Timeline = MakeShared<Trace::TPointTimeline<FAnimMontageMessage>>(Session.GetLinearAllocator());
+		TimelineStorage->Timeline->SetEnumerateOutsideRange(true);
+		ObjectIdToAnimMontageTimelines.Add(InAnimInstanceId, AnimMontageTimelineStorage.Num());
+		AnimMontageTimelineStorage.Add(TimelineStorage.ToSharedRef());
+	}
+
+	TimelineStorage->AllMontageIds.Add(InMontageId);
+
+	FAnimMontageMessage Message;
+	Message.AnimInstanceId = InAnimInstanceId;
+	Message.MontageId = InMontageId;
+	Message.CurrentSectionNameId = InCurrentSectionNameId;
+	Message.NextSectionNameId = InNextSectionNameId;
+	Message.Weight = InWeight;
+	Message.DesiredWeight = InDesiredWeight;
+
+	TimelineStorage->Timeline->AppendEvent(InTime, Message);
+
+	Session.UpdateDurationSeconds(InTime);
+}
+
 #undef LOCTEXT_NAMESPACE
\ No newline at end of file
diff --git a/Engine/Plugins/Animation/GameplayInsights/Source/GameplayInsights/Private/AnimationProvider.h b/Engine/Plugins/Animation/GameplayInsights/Source/GameplayInsights/Private/AnimationProvider.h
--- a/Engine/Plugins/Animation/GameplayInsights/Source/GameplayInsights/Private/AnimationProvider.h
+++ b/Engine/Plugins/Animation/GameplayInsights/Source/GameplayInsights/Private/AnimationProvider.h
@@ -34,6 +34,10 @@ public:
 	virtual bool ReadAnimSequencePlayersTimeline(uint64 InObjectId, TFunctionRef<void(const AnimSequencePlayersTimeline&)> Callback) const override;
 	virtual bool ReadAnimBlendSpacePlayersTimeline(uint64 InObjectId, TFunctionRef<void(const BlendSpacePlayersTimeline&)> Callback) const override;
 	virtual bool ReadStateMachinesTimeline(uint64 InObjectId, TFunctionRef<void(const StateMachinesTimeline&)> Callback) const override;
+	virtual bool ReadNotifyTimeline(uint64 InObjectId, TFunctionRef<void(const AnimNotifyTimeline&)> Callback) const override;
+	virtual void EnumerateNotifyStateTimelines(uint64 InObjectId, TFunctionRef<void(uint64, const AnimNotifyTimeline&)> Callback) const override;
+	virtual bool ReadMontageTimeline(uint64 InObjectId, TFunctionRef<void(const AnimMontageTimeline&)> Callback) const override;
+	virtual void EnumerateMontageIds(uint64 InObjectId, TFunctionRef<void(uint64)> Callback) const override;
 	virtual const FSkeletalMeshInfo* FindSkeletalMeshInfo(uint64 InObjectId) const override;
 	virtual const TCHAR* GetName(uint32 InId) const override;
 	virtual FText FormatNodeKeyValue(const FAnimNodeValueMessage& InMessage) const override;
@@ -68,6 +72,7 @@ public:
 	void AppendAnimNodeValue(uint64 InAnimInstanceId, double InTime, uint16 InFrameCounter, int32 InNodeId, const TCHAR* InKey, bool bInValue);
 	void AppendAnimNodeValue(uint64 InAnimInstanceId, double InTime, uint16 InFrameCounter, int32 InNodeId, const TCHAR* InKey, int32 InValue);
 	void AppendAnimNodeValue(uint64 InAnimInstanceId, double InTime, uint16 InFrameCounter, int32 InNodeId, const TCHAR* InKey, float InValue);
+	void AppendAnimNodeValue(uint64 InAnimInstanceId, double InTime, uint16 InFrameCounter, int32 InNodeId, const TCHAR* InKey, const FVector2D& InValue);
 	void AppendAnimNodeValue(uint64 InAnimInstanceId, double InTime, uint16 InFrameCounter, int32 InNodeId, const TCHAR* InKey, const FVector& InValue);
 	void AppendAnimNodeValue(uint64 InAnimInstanceId, double InTime, uint16 InFrameCounter, int32 InNodeId, const TCHAR* InKey, const TCHAR* InValue);
 	void AppendAnimNodeValueObject(uint64 InAnimInstanceId, double InTime, uint16 InFrameCounter, int32 InNodeId, const TCHAR* InKey, uint64 InValue);
@@ -82,6 +87,12 @@ public:
 	/** Add a state machine state */
 	void AppendStateMachineState(uint64 InAnimInstanceId, double InTime, int32 InNodeId, int32 InStateMachineIndex, int32 InStateIndex, float InStateWeight, float InElapsedTime);
 
+	/** Add an anim notify */
+	void AppendNotify(uint64 InAnimInstanceId, double InTime, uint64 InAssetId, uint64 InNotifyId, uint32 InNameId, float InNotifyTime, float InNotifyDuration, EAnimNotifyMessageType InNotifyEventType);
+
+	/** Append montage data */
+	void AppendMontage(uint64 InAnimInstanceId, double InTime, uint64 InMontageId, uint32 InCurrentSectionNameId, uint32 InNextSectionNameId, float InWeight, float InDesiredWeight);
+
 private:
 	/** Add anim node values helper */
 	void AppendAnimNodeValue(uint64 InAnimInstanceId, double InTime, uint16 InFrameCounter, int32 InNodeId, const TCHAR* InKey, FAnimNodeValueMessage& InMessage);
@@ -103,6 +114,9 @@ private:
 	TMap<uint64, uint32> ObjectIdToAnimNodeValueTimelines;
 	TMap<uint64, uint32> ObjectIdToAnimSequencePlayerTimelines;
 	TMap<uint64, uint32> ObjectIdToBlendSpacePlayerTimelines;
+	TMap<uint64, uint32> ObjectIdToAnimNotifyStateTimelines;
+	TMap<uint64, uint32> ObjectIdToAnimNotifyTimelines;
+	TMap<uint64, uint32> ObjectIdToAnimMontageTimelines;
 
 	/** All the skeletal mesh info we have seen, grow only for stable indices */
 	TArray<FSkeletalMeshInfo> SkeletalMeshInfos;
@@ -125,6 +139,18 @@ private:
 		TSet<uint32> AllCurveIds;
 	};
 
+	struct FAnimNotifyStateTimelineStorage
+	{
+		TMap<uint64, uint32> NotifyIdToAnimNotifyStateTimeline;
+		TArray<TSharedRef<Trace::TIntervalTimeline<FAnimNotifyMessage>>> Timelines;
+	};
+
+	struct FMontageTimelineStorage
+	{
+		TSharedPtr<Trace::TPointTimeline<FAnimMontageMessage>> Timeline;
+		TSet<uint64> AllMontageIds;
+	};
+
 	/** Message storage */
 	TArray<TSharedRef<FTickRecordTimelineStorage>> TickRecordTimelineStorage;
 	TArray<TSharedRef<FSkeletalMeshTimelineStorage>> SkeletalMeshPoseTimelineStorage;
@@ -135,6 +161,9 @@ private:
 	TArray<TSharedRef<Trace::TPointTimeline<FAnimSequencePlayerMessage>>> AnimSequencePlayerTimelines;
 	TArray<TSharedRef<Trace::TPointTimeline<FBlendSpacePlayerMessage>>> BlendSpacePlayerTimelines;
 	TArray<TSharedRef<Trace::TPointTimeline<FAnimStateMachineMessage>>> StateMachineTimelines;
+	TArray<TSharedRef<FAnimNotifyStateTimelineStorage>> AnimNotifyStateTimelineStorage;
+	TArray<TSharedRef<Trace::TPointTimeline<FAnimNotifyMessage>>> AnimNotifyTimelines;
+	TArray<TSharedRef<FMontageTimelineStorage>> AnimMontageTimelineStorage;
 	TPagedArray<FTransform> SkeletalMeshPoseTransforms;
 	TPagedArray<FSkeletalMeshNamedCurve> SkeletalMeshCurves;
 	TPagedArray<int32> SkeletalMeshParentIndices;
diff --git a/Engine/Plugins/Animation/GameplayInsights/Source/GameplayInsights/Private/AnimationSharedData.cpp b/Engine/Plugins/Animation/GameplayInsights/Source/GameplayInsights/Private/AnimationSharedData.cpp
--- a/Engine/Plugins/Animation/GameplayInsights/Source/GameplayInsights/Private/AnimationSharedData.cpp
+++ b/Engine/Plugins/Animation/GameplayInsights/Source/GameplayInsights/Private/AnimationSharedData.cpp
@@ -18,6 +18,10 @@
 #include "GameplayInsightsModule.h"
 #include "Modules/ModuleManager.h"
 #include "Widgets/Docking/SDockTab.h"
+#include "Animation/AnimTypes.h"
+#include "AnimNotifiesTrack.h"
+#include "MontageTrack.h"
+#include "STrackVariantValueView.h"
 
 #if WITH_EDITOR
 #include "EditorViewportClient.h"
@@ -35,6 +39,8 @@ FAnimationSharedData::FAnimationSharedData(FGameplaySharedData& InGameplayShared
 	, bSkeletalMeshCurveTracksEnabled(true)
 	, bTickRecordTracksEnabled(true)
 	, bAnimNodeTracksEnabled(true)
+	, bAnimNotifyTracksEnabled(true)
+	, bMontageTracksEnabled(true)
 {
 }
 
@@ -169,6 +175,58 @@ void FAnimationSharedData::Tick(Insights::ITimingViewSession& InTimingViewSessio
 							ObjectEventsTrack->GetGameplayTrack().AddChildTrack(AnimNodesTrack->GetGameplayTrack());
 						}
 					});
+
+					auto ProcessNotifyTimeline = [this, &InObjectInfo, &ObjectEventsTrack, &InTimingViewSession]()
+					{
+						auto FindAnimNotifyTrack = [](const FBaseTimingTrack& InTrack)
+						{
+							return InTrack.Is<FAnimNotifiesTrack>();
+						};
+
+						TSharedPtr<FAnimNotifiesTrack> ExistingAnimNotifyTrack = StaticCastSharedPtr<FAnimNotifiesTrack>(ObjectEventsTrack->GetGameplayTrack().FindChildTrack(InObjectInfo.Id, FindAnimNotifyTrack));
+						if(!ExistingAnimNotifyTrack.IsValid())
+						{
+							TSharedPtr<FAnimNotifiesTrack> AnimNotifyTrack = MakeShared<FAnimNotifiesTrack>(*this, InObjectInfo.Id, InObjectInfo.Name);
+							AnimNotifyTrack->SetVisibilityFlag(bAnimNotifyTracksEnabled);
+							AnimNotifyTracks.Add(AnimNotifyTrack.ToSharedRef());
+
+							InTimingViewSession.AddScrollableTrack(AnimNotifyTrack);
+							GameplaySharedData.InvalidateObjectTracksOrder();
+
+							ObjectEventsTrack->GetGameplayTrack().AddChildTrack(AnimNotifyTrack->GetGameplayTrack());
+						}
+					};
+
+					AnimationProvider->ReadNotifyTimeline(InObjectInfo.Id, [&ProcessNotifyTimeline](const IAnimationProvider::AnimNotifyTimeline& InTimeline)
+					{
+						ProcessNotifyTimeline();
+					});
+
+					AnimationProvider->EnumerateNotifyStateTimelines(InObjectInfo.Id, [&ProcessNotifyTimeline](uint32 InNotifyNameId, const IAnimationProvider::AnimNotifyTimeline& InTimeline)
+					{
+						ProcessNotifyTimeline();
+					});
+
+					AnimationProvider->ReadMontageTimeline(InObjectInfo.Id, [this, &InObjectInfo, &ObjectEventsTrack, &InTimingViewSession](const IAnimationProvider::AnimMontageTimeline& InTimeline)
+					{
+						auto FindMontageTrack = [](const FBaseTimingTrack& InTrack)
+						{
+							return InTrack.Is<FMontageTrack>();
+						};
+
+						TSharedPtr<FMontageTrack> ExistingMontageTrack = StaticCastSharedPtr<FMontageTrack>(ObjectEventsTrack->GetGameplayTrack().FindChildTrack(InObjectInfo.Id, FindMontageTrack));
+						if(!ExistingMontageTrack.IsValid())
+						{
+							TSharedPtr<FMontageTrack> MontageTrack = MakeShared<FMontageTrack>(*this, InObjectInfo.Id, InObjectInfo.Name);
+							MontageTrack->SetVisibilityFlag(bMontageTracksEnabled);
+							MontageTracks.Add(MontageTrack.ToSharedRef());
+
+							InTimingViewSession.AddScrollableTrack(MontageTrack);
+							GameplaySharedData.InvalidateObjectTracksOrder();
+
+							ObjectEventsTrack->GetGameplayTrack().AddChildTrack(MontageTrack->GetGameplayTrack());
+						}
+					});
 				}
 			});
 		});
@@ -249,6 +307,30 @@ void FAnimationSharedData::ExtendFilterMenu(FMenuBuilder& InMenuBuilder)
 		NAME_None,
 		EUserInterfaceActionType::ToggleButton
 	);
+
+	InMenuBuilder.AddMenuEntry(
+		LOCTEXT("ToggleAnimNotifyTracks", "Notify Tracks"),
+		LOCTEXT("ToggleAnimNotifyTracks_Tooltip", "Show/hide the animation notify/sync marker tracks"),
+		FSlateIcon(),
+		FUIAction(
+			FExecuteAction::CreateRaw(this, &FAnimationSharedData::ToggleAnimNotifyTracks),
+			FCanExecuteAction(),
+			FIsActionChecked::CreateLambda([this](){ return bAnimNotifyTracksEnabled; })),
+		NAME_None,
+		EUserInterfaceActionType::ToggleButton
+	);
+
+	InMenuBuilder.AddMenuEntry(
+		LOCTEXT("ToggleMontageTracks", "Montage Tracks"),
+		LOCTEXT("ToggleMontageTracks_Tooltip", "Show/hide the montage tracks"),
+		FSlateIcon(),
+		FUIAction(
+			FExecuteAction::CreateRaw(this, &FAnimationSharedData::ToggleMontageTracks),
+			FCanExecuteAction(),
+			FIsActionChecked::CreateLambda([this](){ return bMontageTracksEnabled; })),
+		NAME_None,
+		EUserInterfaceActionType::ToggleButton
+	);
 }
 
 void FAnimationSharedData::ToggleAnimationTracks()
@@ -259,6 +341,8 @@ void FAnimationSharedData::ToggleAnimationTracks()
 	bSkeletalMeshCurveTracksEnabled = bAnimationTracksEnabled;
 	bTickRecordTracksEnabled = bAnimationTracksEnabled;
 	bAnimNodeTracksEnabled = bAnimationTracksEnabled;
+	bAnimNotifyTracksEnabled = bAnimationTracksEnabled;
+	bMontageTracksEnabled = bAnimNotifyTracksEnabled;
 
 	for(TSharedRef<FSkeletalMeshPoseTrack> PoseTrack : SkeletalMeshPoseTracks)
 	{
@@ -279,11 +363,21 @@ void FAnimationSharedData::ToggleAnimationTracks()
 	{
 		AnimNodesTrack->SetVisibilityFlag(bAnimationTracksEnabled);
 	}
+
+	for(TSharedRef<FAnimNotifiesTrack> AnimNotifyTrack : AnimNotifyTracks)
+	{
+		AnimNotifyTrack->SetVisibilityFlag(bAnimationTracksEnabled);
+	}
+
+	for(TSharedRef<FMontageTrack> MontageTrack : MontageTracks)
+	{
+		MontageTrack->SetVisibilityFlag(bAnimationTracksEnabled);
+	}
 }
 
 bool FAnimationSharedData::AreAnimationTracksEnabled() const
 {
-	return bSkeletalMeshPoseTracksEnabled && bSkeletalMeshCurveTracksEnabled && bTickRecordTracksEnabled && bAnimNodeTracksEnabled;
+	return bSkeletalMeshPoseTracksEnabled && bSkeletalMeshCurveTracksEnabled && bTickRecordTracksEnabled && bAnimNodeTracksEnabled && bAnimNotifyTracksEnabled && bMontageTracksEnabled;
 }
 
 void FAnimationSharedData::ToggleSkeletalMeshPoseTracks()
@@ -326,6 +420,26 @@ void FAnimationSharedData::ToggleAnimNodeTracks()
 	}
 }
 
+void FAnimationSharedData::ToggleAnimNotifyTracks()
+{
+	bAnimNotifyTracksEnabled = !bAnimNotifyTracksEnabled;
+
+	for(TSharedRef<FAnimNotifiesTrack> AnimNotifyTrack : AnimNotifyTracks)
+	{
+		AnimNotifyTrack->SetVisibilityFlag(bAnimNotifyTracksEnabled);
+	}
+}
+
+void FAnimationSharedData::ToggleMontageTracks()
+{
+	bMontageTracksEnabled = !bMontageTracksEnabled;
+	
+	for(TSharedRef<FMontageTrack> MontageTrack : MontageTracks)
+	{
+		MontageTrack->SetVisibilityFlag(bMontageTracksEnabled);
+	}
+}
+
 void FAnimationSharedData::OnTimeMarkerChanged(Insights::ETimeChangedFlags InFlags, double InTimeMarker)
 {
 	bTimeMarkerValid = InTimeMarker != std::numeric_limits<double>::infinity();
@@ -447,36 +561,23 @@ void FAnimationSharedData::GetCustomDebugObjects(const IAnimationBlueprintEditor
 
 #endif
 
-/** Search predicate for an open anim graph document */
-struct FSearchForAnimInstanceTab : public FTabManager::FSearchPreference
-{
-	FSearchForAnimInstanceTab(uint64 InAnimInstanceId)
-		: AnimInstanceId(InAnimInstanceId)
-	{
-	}
-
-	virtual TSharedPtr<SDockTab> Search(const FTabManager& Manager, FName PlaceholderId, const TSharedRef<SDockTab>& UnmanagedTab) const override
-	{
-		TSharedRef<SAnimGraphSchematicView> Content = StaticCastSharedRef<SAnimGraphSchematicView>(UnmanagedTab->GetContent());
-		if(Content->GetAnimInstanceId() == AnimInstanceId)
-		{
-			return UnmanagedTab;
-		}
-
-		return TSharedPtr<SDockTab>();
-	}
-
-	uint64 AnimInstanceId;
-};
-
 void FAnimationSharedData::OpenAnimGraphTab(uint64 InAnimInstanceId) const
 {
 	if(TimingViewSession && AnalysisSession)
 	{
 		FGameplayInsightsModule& GameplayInsightsModule = FModuleManager::GetModuleChecked<FGameplayInsightsModule>("GameplayInsights");
-		TSharedRef<SDockTab> Tab = GameplayInsightsModule.SpawnTimingProfilerDocumentTab(FSearchForAnimInstanceTab(InAnimInstanceId));
+		TSharedRef<SDockTab> Tab = GameplayInsightsModule.SpawnTimingProfilerDocumentTab(
+			FGameplaySharedData::FSearchForTab([this, InAnimInstanceId]()
+			{
+				return FGameplaySharedData::FindDocumentTab(WeakAnimGraphDocumentTabs, [InAnimInstanceId](const TSharedRef<SDockTab>& InDockTab)
+				{
+					return StaticCastSharedRef<SAnimGraphSchematicView>(InDockTab->GetContent())->GetAnimInstanceId() == InAnimInstanceId;
+				});
+			})
+		);
 
 		Tab->SetContent(SNew(SAnimGraphSchematicView, InAnimInstanceId, *TimingViewSession, *AnalysisSession));
+		WeakAnimGraphDocumentTabs.Add(Tab);
 
 		const FGameplayProvider* GameplayProvider = AnalysisSession->ReadProvider<FGameplayProvider>(FGameplayProvider::ProviderName);
 		if(GameplayProvider)
diff --git a/Engine/Plugins/Animation/GameplayInsights/Source/GameplayInsights/Private/AnimationSharedData.h b/Engine/Plugins/Animation/GameplayInsights/Source/GameplayInsights/Private/AnimationSharedData.h
--- a/Engine/Plugins/Animation/GameplayInsights/Source/GameplayInsights/Private/AnimationSharedData.h
+++ b/Engine/Plugins/Animation/GameplayInsights/Source/GameplayInsights/Private/AnimationSharedData.h
@@ -12,11 +12,14 @@ class FSkeletalMeshPoseTrack;
 class FSkeletalMeshCurvesTrack;
 class FAnimationTickRecordsTrack;
 class FAnimNodesTrack;
+class FAnimNotifiesTrack;
+class FMontageTrack;
 class FMenuBuilder;
 class UWorld;
 class IAnimationBlueprintEditor;
 struct FCustomDebugObject;
 class SDockTab;
+class FGameplayTrack;
 
 class FAnimationSharedData
 {
@@ -75,6 +78,8 @@ private:
 	void ToggleSkeletalMeshCurveTracks();
 	void ToggleTickRecordTracks();
 	void ToggleAnimNodeTracks();
+	void ToggleAnimNotifyTracks();
+	void ToggleMontageTracks();
 
 private:
 	// The gameplay shared data we are linked to
@@ -91,6 +96,11 @@ private:
 	TArray<TSharedRef<FSkeletalMeshCurvesTrack>> SkeletalMeshCurvesTracks;
 	TArray<TSharedRef<FAnimationTickRecordsTrack>> AnimationTickRecordsTracks;
 	TArray<TSharedRef<FAnimNodesTrack>> AnimNodesTracks;
+	TArray<TSharedRef<FAnimNotifiesTrack>> AnimNotifyTracks;
+	TArray<TSharedRef<FMontageTrack>> MontageTracks;
+
+	// All the documents we have spawned
+	mutable TArray<TWeakPtr<SDockTab>> WeakAnimGraphDocumentTabs;
 
 	// Delegate handles for hooks into the timing view
 	FDelegateHandle TimeMarkerChangedHandle;
@@ -106,4 +116,6 @@ private:
 	bool bSkeletalMeshCurveTracksEnabled;
 	bool bTickRecordTracksEnabled;
 	bool bAnimNodeTracksEnabled;
+	bool bAnimNotifyTracksEnabled;
+	bool bMontageTracksEnabled;
 };
\ No newline at end of file
diff --git a/Engine/Plugins/Animation/GameplayInsights/Source/GameplayInsights/Private/AnimationTickRecordsTrack.cpp b/Engine/Plugins/Animation/GameplayInsights/Source/GameplayInsights/Private/AnimationTickRecordsTrack.cpp
--- a/Engine/Plugins/Animation/GameplayInsights/Source/GameplayInsights/Private/AnimationTickRecordsTrack.cpp
+++ b/Engine/Plugins/Animation/GameplayInsights/Source/GameplayInsights/Private/AnimationTickRecordsTrack.cpp
@@ -22,6 +22,8 @@
 #include "Subsystems/AssetEditorSubsystem.h"
 #include "EdGraph/EdGraphNode.h"
 #endif
+#include "VariantTreeNode.h"
+#include "TraceServices/Model/Frames.h"
 
 #define LOCTEXT_NAMESPACE "AnimationTickRecordsTrack"
 
@@ -58,7 +60,7 @@ static FLinearColor MakeSeriesColor(FTickRecordSeries::ESeriesType InSeed, bool
 }
 
 FAnimationTickRecordsTrack::FAnimationTickRecordsTrack(const FAnimationSharedData& InSharedData, uint64 InObjectID, uint64 InAssetId, int32 InNodeId, const TCHAR* InName)
-	: FGameplayGraphTrack(InObjectID, MakeTrackName(InSharedData.GetGameplaySharedData(), InAssetId, InName))
+	: FGameplayGraphTrack(InSharedData.GetGameplaySharedData(), InObjectID, MakeTrackName(InSharedData.GetGameplaySharedData(), InAssetId, InName))
 	, SharedData(InSharedData)
 	, AssetId(InAssetId)
 	, NodeId(InNodeId)
@@ -87,6 +89,11 @@ FAnimationTickRecordsTrack::FAnimationTickRecordsTrack(const FAnimationSharedDat
 
 void FAnimationTickRecordsTrack::AddAllSeries()
 {
+	if(AllSeries.Num() > 0)
+	{
+		return;
+	}
+
 	struct FSeriesDescription
 	{
 		FText Name;
@@ -410,4 +417,43 @@ FText FAnimationTickRecordsTrack::MakeTrackName(const FGameplaySharedData& InSha
 	return FText::Format(LOCTEXT("AnimationTickRecordsTrackName", "{0} - {1}"), AssetTypeName, FText::FromString(InName));
 }
 
+void FAnimationTickRecordsTrack::GetVariantsAtTime(double InTime, TArray<TSharedRef<FVariantTreeNode>>& OutVariants) const
+{
+	const FGameplayProvider* GameplayProvider = SharedData.GetAnalysisSession().ReadProvider<FGameplayProvider>(FGameplayProvider::ProviderName);
+	const FAnimationProvider* AnimationProvider = SharedData.GetAnalysisSession().ReadProvider<FAnimationProvider>(FAnimationProvider::ProviderName);
+	if(GameplayProvider && AnimationProvider)
+	{
+		Trace::FAnalysisSessionReadScope SessionReadScope(SharedData.GetAnalysisSession());
+
+		const FClassInfo& ClassInfo = GameplayProvider->GetClassInfoFromObject(AssetId);
+		TSharedRef<FVariantTreeNode> Header = OutVariants.Add_GetRef(FVariantTreeNode::MakeObject(FText::FromString(ClassInfo.Name), AssetId));
+
+		const Trace::IFrameProvider& FramesProvider = Trace::ReadFrameProvider(SharedData.GetAnalysisSession());
+
+		AnimationProvider->ReadTickRecordTimeline(GetGameplayTrack().GetObjectId(), AssetId, NodeId, [&Header, &FramesProvider, &InTime](const FAnimationProvider::TickRecordTimeline& InTimeline)
+		{
+			// round to nearest frame boundary
+			Trace::FFrame Frame;
+			if(FramesProvider.GetFrameFromTime(ETraceFrameType::TraceFrameType_Game, InTime, Frame))
+			{
+				InTimeline.EnumerateEvents(Frame.StartTime, Frame.EndTime, [&Header](double InEventStartTime, double InEventEndTime, uint32 InDepth, const FTickRecordMessage& InMessage)
+				{
+					if(Header->GetChildren().Num() == 0)
+					{
+						Header->AddChild(FVariantTreeNode::MakeFloat(LOCTEXT("BlendWeight", "Blend Weight"), InMessage.BlendWeight));
+						Header->AddChild(FVariantTreeNode::MakeFloat(LOCTEXT("PlaybackTime", "Playback Time"), InMessage.PlaybackTime));
+						Header->AddChild(FVariantTreeNode::MakeFloat(LOCTEXT("RootMotionWeight", "Root Motion Weight"), InMessage.RootMotionWeight));
+						Header->AddChild(FVariantTreeNode::MakeFloat(LOCTEXT("PlayRate", "Play Rate"), InMessage.PlayRate));
+						if(InMessage.bIsBlendSpace)
+						{
+							Header->AddChild(FVariantTreeNode::MakeFloat(LOCTEXT("BlendSpacePositionX", "Blend Space Position X"), InMessage.BlendSpacePositionX));
+							Header->AddChild(FVariantTreeNode::MakeFloat(LOCTEXT("BlendSpacePositionY", "Blend Space Position Y"), InMessage.BlendSpacePositionY));
+						}
+					}
+				});
+			}
+		});
+	}
+}
+
 #undef LOCTEXT_NAMESPACE
\ No newline at end of file
diff --git a/Engine/Plugins/Animation/GameplayInsights/Source/GameplayInsights/Private/AnimationTickRecordsTrack.h b/Engine/Plugins/Animation/GameplayInsights/Source/GameplayInsights/Private/AnimationTickRecordsTrack.h
--- a/Engine/Plugins/Animation/GameplayInsights/Source/GameplayInsights/Private/AnimationTickRecordsTrack.h
+++ b/Engine/Plugins/Animation/GameplayInsights/Source/GameplayInsights/Private/AnimationTickRecordsTrack.h
@@ -46,6 +46,7 @@ public:
 	virtual bool UpdateSeriesBounds(FGameplayGraphSeries& InSeries, const FTimingTrackViewport& InViewport) override;
 	virtual void UpdateSeries(FGameplayGraphSeries& InSeries, const FTimingTrackViewport& InViewport) override;
 	virtual void AddAllSeries() override;
+	virtual void GetVariantsAtTime(double InTime, TArray<TSharedRef<FVariantTreeNode>>& OutVariants) const override;
 
 	/** Get the asset ID that this track uses */
 	uint64 GetAssetId() const { return AssetId; }
diff --git a/Engine/Plugins/Animation/GameplayInsights/Source/GameplayInsights/Private/GameplayGraphTrack.cpp b/Engine/Plugins/Animation/GameplayInsights/Source/GameplayInsights/Private/GameplayGraphTrack.cpp
--- a/Engine/Plugins/Animation/GameplayInsights/Source/GameplayInsights/Private/GameplayGraphTrack.cpp
+++ b/Engine/Plugins/Animation/GameplayInsights/Source/GameplayInsights/Private/GameplayGraphTrack.cpp
@@ -5,7 +5,7 @@
 #include "AnimationProvider.h"
 #include "Insights/ViewModels/TimingTrackViewport.h"
 #include "Insights/ViewModels/TimingEvent.h"
-#include "AnimationSharedData.h"
+#include "GameplaySharedData.h"
 #include "Framework/MultiBox/MultiBoxBuilder.h"
 #include "Insights/ViewModels/TimingEventSearch.h"
 #include "Insights/ViewModels/TooltipDrawState.h"
@@ -19,8 +19,6 @@
 
 INSIGHTS_IMPLEMENT_RTTI(FGameplayGraphTrack)
 
-PRAGMA_DISABLE_OPTIMIZATION
-
 static float GetSeriesHeight(const FTimingViewLayout& InLayout)
 {
 	// 1.0f is for the top horizontal line of each track
@@ -48,8 +46,9 @@ void FGameplayGraphSeries::UpdateAutoZoom(const FTimingTrackViewport& InViewport
 	FGraphSeries::UpdateAutoZoom(TopY, BottomY, CurrentMin, CurrentMax, false);
 }
 
-FGameplayGraphTrack::FGameplayGraphTrack(uint64 InObjectID, const FText& InName)
+FGameplayGraphTrack::FGameplayGraphTrack(const FGameplaySharedData& InGameplaySharedData, uint64 InObjectID, const FText& InName)
 	: TGameplayTrackMixin<FGraphTrack>(InObjectID, InName)
+	, GameplaySharedData(InGameplaySharedData)
 	, RequestedTrackSizeScale(1.0f)
 	, BorderY(0.0f)
 	, NumActiveSeries(0)
@@ -123,10 +122,7 @@ void FGameplayGraphTrack::UpdateSeriesInternal(FGameplayGraphSeries& InSeries, c
 
 void FGameplayGraphTrack::PreUpdate(const ITimingTrackUpdateContext& Context)
 {
-	if(AllSeries.Num() == 0)
-	{
-		AddAllSeries();
-	}
+	AddAllSeries();
 
 	FGraphTrack::PreUpdate(Context);
 
@@ -219,6 +215,22 @@ void FGameplayGraphTrack::Draw(const ITimingTrackDrawContext& Context) const
 
 void FGameplayGraphTrack::BuildContextMenu(FMenuBuilder& MenuBuilder)
 {
+	MenuBuilder.BeginSection("View", LOCTEXT("ViewHeader", "View"));
+	{
+		MenuBuilder.AddMenuEntry
+		(
+			LOCTEXT("ViewProperties", "View Properties"),
+			LOCTEXT("ViewProperties_Tooltip", "Open a window to view the properties of this track. You can scrub the timeline to see properties change in real-time."),
+			FSlateIcon(),
+			FUIAction(
+				FExecuteAction::CreateLambda([this](){ GameplaySharedData.OpenTrackVariantsTab(GetGameplayTrack()); })
+			),
+			NAME_None,
+			EUserInterfaceActionType::Button
+		);
+	}
+	MenuBuilder.EndSection();
+
 	FGraphTrack::BuildContextMenu(MenuBuilder);
 
 	MenuBuilder.BeginSection("Layout", LOCTEXT("TrackLayoutMenuHeader", "Track Layout"));
@@ -314,6 +326,4 @@ void FGameplayGraphTrack::BuildContextMenu(FMenuBuilder& MenuBuilder)
 	MenuBuilder.EndSection();
 }
 
-PRAGMA_ENABLE_OPTIMIZATION
-
 #undef LOCTEXT_NAMESPACE
diff --git a/Engine/Plugins/Animation/GameplayInsights/Source/GameplayInsights/Private/GameplayGraphTrack.h b/Engine/Plugins/Animation/GameplayInsights/Source/GameplayInsights/Private/GameplayGraphTrack.h
--- a/Engine/Plugins/Animation/GameplayInsights/Source/GameplayInsights/Private/GameplayGraphTrack.h
+++ b/Engine/Plugins/Animation/GameplayInsights/Source/GameplayInsights/Private/GameplayGraphTrack.h
@@ -8,6 +8,7 @@
 
 class FTimingEventSearchParameters;
 class FGameplayGraphTrack;
+struct FVariantTreeNode;
 
 /** The various layouts that we display series with */
 enum class EGameplayGraphLayout : int32
@@ -43,12 +44,15 @@ class FGameplayGraphTrack : public TGameplayTrackMixin<FGraphTrack>
 	INSIGHTS_DECLARE_RTTI(FGameplayGraphTrack, TGameplayTrackMixin<FGraphTrack>)
 
 public:
-	FGameplayGraphTrack(uint64 InObjectID, const FText& InName);
+	FGameplayGraphTrack(const FGameplaySharedData& InGameplaySharedData, uint64 InObjectID, const FText& InName);
 
 	virtual void Draw(const ITimingTrackDrawContext& Context) const override;
 	virtual void PreUpdate(const ITimingTrackUpdateContext& Context) override;
 	virtual void BuildContextMenu(FMenuBuilder& MenuBuilder) override;
 
+	// Get all variants at the specified time
+	virtual void GetVariantsAtTime(double InTime, TArray<TSharedRef<FVariantTreeNode>>& OutVariants) const {}
+
 	// Get the requested track size scale
 	float GetRequestedTrackSizeScale() const { return RequestedTrackSizeScale; }
 
@@ -77,6 +81,8 @@ private:
 	void UpdateSeriesInternal(FGameplayGraphSeries& InSeries, const FTimingTrackViewport& InViewport, int32 InActiveSeriesIndex);
 
 protected:
+	const FGameplaySharedData& GameplaySharedData;
+
 	/** The track size we want to be displayed at */
 	float RequestedTrackSizeScale;
 
diff --git a/Engine/Plugins/Animation/GameplayInsights/Source/GameplayInsights/Private/GameplayInsightsModule.cpp b/Engine/Plugins/Animation/GameplayInsights/Source/GameplayInsights/Private/GameplayInsightsModule.cpp
--- a/Engine/Plugins/Animation/GameplayInsights/Source/GameplayInsights/Private/GameplayInsightsModule.cpp
+++ b/Engine/Plugins/Animation/GameplayInsights/Source/GameplayInsights/Private/GameplayInsightsModule.cpp
@@ -29,7 +29,7 @@
 
 namespace GameplayInsightsTabs
 {
-	static const FName AnimGraphSchematicView("AnimGraphSchematicView");
+	static const FName DocumentTab("DocumentTab");
 };
 
 void FGameplayInsightsModule::StartupModule()
@@ -85,7 +85,7 @@ void FGameplayInsightsModule::StartupModule()
 	FInsightsMajorTabConfig TimingProfilerConfig;
 	TimingProfilerConfig.TabLabel = LOCTEXT("GameplayInsightsTabName", "Gameplay Insights");
 	TimingProfilerConfig.TabTooltip = LOCTEXT("GameplayInsightsTabTooltip", "Open the Gameplay Insights tab.");
-	TimingProfilerConfig.Layout = FTabManager::NewLayout("GameplayInsightsTimingLayout_v1.0")
+	TimingProfilerConfig.Layout = FTabManager::NewLayout("GameplayInsightsTimingLayout_v1.1")
 	->AddArea
 	(
 		FTabManager::NewPrimaryArea()
@@ -114,7 +114,7 @@ void FGameplayInsightsModule::StartupModule()
 		(
 			FTabManager::NewStack()
 			->SetSizeCoefficient(0.3f)
-			->AddTab(GameplayInsightsTabs::AnimGraphSchematicView, ETabState::ClosedTab)
+			->AddTab(GameplayInsightsTabs::DocumentTab, ETabState::ClosedTab)
 		)
 	);
 
@@ -161,7 +161,7 @@ TSharedRef<SDockTab> FGameplayInsightsModule::SpawnTimingProfilerDocumentTab(con
 	TSharedPtr<FTabManager> TimingProfilerTabManager = WeakTimingProfilerTabManager.Pin();
 	if(TimingProfilerTabManager.IsValid())
 	{
-		TimingProfilerTabManager->InsertNewDocumentTab(GameplayInsightsTabs::AnimGraphSchematicView, InSearchPreference, NewTab);
+		TimingProfilerTabManager->InsertNewDocumentTab(GameplayInsightsTabs::DocumentTab, InSearchPreference, NewTab);
 	}
 	return NewTab;
 }
diff --git a/Engine/Plugins/Animation/GameplayInsights/Source/GameplayInsights/Private/GameplaySharedData.cpp b/Engine/Plugins/Animation/GameplayInsights/Source/GameplayInsights/Private/GameplaySharedData.cpp
--- a/Engine/Plugins/Animation/GameplayInsights/Source/GameplayInsights/Private/GameplaySharedData.cpp
+++ b/Engine/Plugins/Animation/GameplayInsights/Source/GameplayInsights/Private/GameplaySharedData.cpp
@@ -8,6 +8,10 @@
 #include "Insights/ITimingViewSession.h"
 #include "Framework/MultiBox/MultiBoxBuilder.h"
 #include "SGameplayTrackTree.h"
+#include "GameplayInsightsModule.h"
+#include "STrackVariantValueView.h"
+#include "Modules/ModuleManager.h"
+#include "Widgets/Docking/SDockTab.h"
 
 #define LOCTEXT_NAMESPACE "GameplaySharedData"
 
@@ -20,12 +24,16 @@ FGameplaySharedData::FGameplaySharedData()
 
 void FGameplaySharedData::OnBeginSession(Insights::ITimingViewSession& InTimingViewSession)
 {
+	TimingViewSession = &InTimingViewSession;
+
 	ObjectTracks.Reset();
 }
 
 void FGameplaySharedData::OnEndSession(Insights::ITimingViewSession& InTimingViewSession)
 {
 	ObjectTracks.Reset();
+
+	TimingViewSession = nullptr;
 }
 
 TSharedRef<FObjectEventsTrack> FGameplaySharedData::GetObjectEventsTrackForId(Insights::ITimingViewSession& InTimingViewSession, const Trace::IAnalysisSession& InAnalysisSession, const FObjectInfo& InObjectInfo)
@@ -222,4 +230,42 @@ void FGameplaySharedData::EnumerateObjectTracks(TFunctionRef<void(const TSharedR
 	}
 }
 
+TSharedPtr<SDockTab> FGameplaySharedData::FindDocumentTab(const TArray<TWeakPtr<SDockTab>>& InWeakDocumentTabs, TFunction<bool(const TSharedRef<SDockTab>&)> InSearchFunction)
+{
+	for(TWeakPtr<SDockTab> WeakTab : InWeakDocumentTabs)
+	{
+		TSharedPtr<SDockTab> PinnedTab = WeakTab.Pin();
+		if(PinnedTab.IsValid())
+		{
+			if(InSearchFunction(PinnedTab.ToSharedRef()))
+			{
+				return PinnedTab;
+			}
+		}
+	}
+
+	return nullptr;
+}
+
+void FGameplaySharedData::OpenTrackVariantsTab(const FGameplayTrack& InGameplayTrack) const
+{
+	if(TimingViewSession && AnalysisSession)
+	{
+		FGameplayInsightsModule& GameplayInsightsModule = FModuleManager::GetModuleChecked<FGameplayInsightsModule>("GameplayInsights");
+		TSharedRef<SDockTab> Tab = GameplayInsightsModule.SpawnTimingProfilerDocumentTab(
+			FSearchForTab([this, &InGameplayTrack]()
+			{
+				return FindDocumentTab(WeakTrackVariantsDocumentTabs, [&InGameplayTrack](const TSharedRef<SDockTab>& InDockTab)
+				{
+					return StaticCastSharedRef<STrackVariantValueView>(InDockTab->GetContent())->GetTimingTrack() == InGameplayTrack.GetTimingTrack();
+				});
+			})
+		);
+
+		Tab->SetContent(SNew(STrackVariantValueView, InGameplayTrack.GetTimingTrack(), *TimingViewSession, *AnalysisSession));
+		WeakTrackVariantsDocumentTabs.Add(Tab);
+		Tab->SetLabel(FText::FromString(InGameplayTrack.GetTimingTrack()->GetName()));
+	}
+}
+
 #undef LOCTEXT_NAMESPACE
\ No newline at end of file
diff --git a/Engine/Plugins/Animation/GameplayInsights/Source/GameplayInsights/Private/GameplaySharedData.h b/Engine/Plugins/Animation/GameplayInsights/Source/GameplayInsights/Private/GameplaySharedData.h
--- a/Engine/Plugins/Animation/GameplayInsights/Source/GameplayInsights/Private/GameplaySharedData.h
+++ b/Engine/Plugins/Animation/GameplayInsights/Source/GameplayInsights/Private/GameplaySharedData.h
@@ -3,6 +3,7 @@
 #pragma once
 
 #include "Insights/ViewModels/TimingEventsTrack.h"
+#include "Framework/Docking/TabManager.h"
 
 namespace Trace { class IAnalysisSession; }
 namespace Insights { class ITimingViewSession; }
@@ -10,6 +11,7 @@ class FObjectEventsTrack;
 class FSkeletalMeshPoseTrack;
 class FAnimationTickRecordsTrack;
 struct FObjectInfo;
+class FGameplayTrack;
 
 class FGameplaySharedData
 {
@@ -42,6 +44,28 @@ public:
 	// Delegate fired when tracks change
 	FSimpleMulticastDelegate& OnTracksChanged() { return OnTracksChangedDelegate; }
 
+	/** Search predicate wrapper for an open document */
+	struct FSearchForTab : public FTabManager::FSearchPreference
+	{
+		FSearchForTab(TFunction<TSharedPtr<SDockTab>(void)> InSearchFunction)
+			: SearchFunction(InSearchFunction)
+		{
+		}
+
+		virtual TSharedPtr<SDockTab> Search(const FTabManager& Manager, FName PlaceholderId, const TSharedRef<SDockTab>& UnmanagedTab) const override
+		{
+			return SearchFunction();
+		}
+
+		TFunction<TSharedPtr<SDockTab>(void)> SearchFunction;
+	};
+
+	// Helper function to find an existing document tab
+	static TSharedPtr<SDockTab> FindDocumentTab(const TArray<TWeakPtr<SDockTab>>& InWeakDocumentTabs, TFunction<bool(const TSharedRef<SDockTab>&)> InSearchFunction);
+
+	/** Open a tab to see a track's variants */
+	void OpenTrackVariantsTab(const FGameplayTrack& InGameplayTrack) const;
+
 private:
 	// Re-sort tracks if track ordering has changed
 	void SortTracks();
@@ -59,9 +83,15 @@ private:
 	// Cached analysis session, set in Tick()
 	const Trace::IAnalysisSession* AnalysisSession;
 
+	// Cached timing view session, set in OnBeginSession/OnEndSession
+	Insights::ITimingViewSession* TimingViewSession;
+
 	// Delegate fired when tracks change
 	FSimpleMulticastDelegate OnTracksChangedDelegate;
 
+	// All the documents we have spawned
+	mutable TArray<TWeakPtr<SDockTab>> WeakTrackVariantsDocumentTabs;
+
 	// Dirty flag for adding object tracks, used to trigger re-sorting
 	bool bObjectTracksDirty;
 
diff --git a/Engine/Plugins/Animation/GameplayInsights/Source/GameplayInsights/Private/GameplayTrack.cpp b/Engine/Plugins/Animation/GameplayInsights/Source/GameplayInsights/Private/GameplayTrack.cpp
--- a/Engine/Plugins/Animation/GameplayInsights/Source/GameplayInsights/Private/GameplayTrack.cpp
+++ b/Engine/Plugins/Animation/GameplayInsights/Source/GameplayInsights/Private/GameplayTrack.cpp
@@ -8,6 +8,12 @@
 #include "Insights/ViewModels/TimingTrackViewport.h"
 #include "Insights/Common/PaintUtils.h"
 #include "Insights/ViewModels/ITimingViewDrawHelper.h"
+#include "GameplaySharedData.h"
+#include "Framework/MultiBox/MultiBoxBuilder.h"
+
+INSIGHTS_IMPLEMENT_RTTI(FGameplayTimingEventsTrack);
+
+#define LOCTEXT_NAMESPACE "GameplayTrack"
 
 void FGameplayTrack::AddChildTrack(FGameplayTrack& InChildTrack)
 {
@@ -84,3 +90,24 @@ void FGameplayTrack::DrawHeaderForTimingTrack(const ITimingTrackDrawContext& InC
 		}
 	}
 }
+
+void FGameplayTimingEventsTrack::BuildContextMenu(FMenuBuilder& MenuBuilder)
+{
+	MenuBuilder.BeginSection("View", LOCTEXT("ViewHeader", "View"));
+	{
+		MenuBuilder.AddMenuEntry
+		(
+			LOCTEXT("ViewProperties", "View Properties"),
+			LOCTEXT("ViewProperties_Tooltip", "Open a window to view the properties of this track. You can scrub the timeline to see properties change in real-time."),
+			FSlateIcon(),
+			FUIAction(
+				FExecuteAction::CreateLambda([this](){ GameplaySharedData.OpenTrackVariantsTab(GetGameplayTrack()); })
+			),
+			NAME_None,
+			EUserInterfaceActionType::Button
+		);
+	}
+	MenuBuilder.EndSection();
+}
+
+#undef LOCTEXT_NAMESPACE
\ No newline at end of file
diff --git a/Engine/Plugins/Animation/GameplayInsights/Source/GameplayInsights/Private/GameplayTrack.h b/Engine/Plugins/Animation/GameplayInsights/Source/GameplayInsights/Private/GameplayTrack.h
--- a/Engine/Plugins/Animation/GameplayInsights/Source/GameplayInsights/Private/GameplayTrack.h
+++ b/Engine/Plugins/Animation/GameplayInsights/Source/GameplayInsights/Private/GameplayTrack.h
@@ -9,6 +9,9 @@ namespace GameplayTrackConstants
 	constexpr float IndentSize = 12.0f;
 }
 
+struct FVariantTreeNode;
+class FGameplaySharedData;
+
 // Provides parent/child hierarchy structure and the owning object Id
 // Designed as a compositional member of outer timing tracks
 class FGameplayTrack
@@ -88,5 +91,25 @@ private:
 	FGameplayTrack GameplayTrack;
 };
 
+// Common base class for timing-event tracks
+class FGameplayTimingEventsTrack : public TGameplayTrackMixin<FTimingEventsTrack>
+{
+	INSIGHTS_DECLARE_RTTI(FGameplayTimingEventsTrack, TGameplayTrackMixin<FTimingEventsTrack>)
+
+public:
+	FGameplayTimingEventsTrack(const FGameplaySharedData& InGameplaySharedData, uint64 InObjectId, const FText& InName)
+		: TGameplayTrackMixin<FTimingEventsTrack>(InObjectId, InName)
+		, GameplaySharedData(InGameplaySharedData)
+	{}
+
+	virtual void BuildContextMenu(FMenuBuilder& MenuBuilder) override;
+
+	/** Get all variants at the specified time */
+	virtual void GetVariantsAtTime(double InTime, TArray<TSharedRef<FVariantTreeNode>>& OutVariants) const {}
+
+protected:
+	const FGameplaySharedData& GameplaySharedData;
+};
+
 //template <class Base>
-//INSIGHTS_IMPLEMENT_RTTI(TGameplayTrackMixin<Base>) // Note: All templeted classes will return same type name as FName(TEXT"TGameplayTrackMixin<Base>") !!!
+//INSIGHTS_IMPLEMENT_RTTI(TGameplayTrackMixin<Base>) // Note: All templated classes will return same type name as FName(TEXT"TGameplayTrackMixin<Base>") !!!
diff --git a/Engine/Plugins/Animation/GameplayInsights/Source/GameplayInsights/Private/IAnimationProvider.h b/Engine/Plugins/Animation/GameplayInsights/Source/GameplayInsights/Private/IAnimationProvider.h
--- a/Engine/Plugins/Animation/GameplayInsights/Source/GameplayInsights/Private/IAnimationProvider.h
+++ b/Engine/Plugins/Animation/GameplayInsights/Source/GameplayInsights/Private/IAnimationProvider.h
@@ -92,19 +92,29 @@ enum class EAnimNodeValueType : uint8
 	Bool,
 	Int32,
 	Float,
+	Vector2D,
 	Vector,
 	String,
 	Object,
 	Class,
 };
 
-struct FAnimNodeValueMessage
+struct FVariantValue
 {
-	FAnimNodeValueMessage()
+	FVariantValue()
 		: Vector(FVector::ZeroVector)
 	{
 	}
 
+	struct FVector2DEntry
+	{
+		FVector2DEntry(const FVector2D& InVector)
+			: Value(InVector)
+		{}
+
+		FVector2D Value;
+	};
+
 	struct FVectorEntry
 	{
 		FVectorEntry(const FVector& InVector)
@@ -114,8 +124,6 @@ struct FAnimNodeValueMessage
 		FVector Value;
 	};
 
-	const TCHAR* Key = nullptr;
-	uint64 AnimInstanceId = 0;
 	union
 	{
 		struct
@@ -130,6 +138,7 @@ struct FAnimNodeValueMessage
 		{
 			float Value;
 		} Float;
+		FVector2DEntry Vector2D;
 		FVectorEntry Vector;
 		struct
 		{
@@ -144,9 +153,17 @@ struct FAnimNodeValueMessage
 			uint64 Value;
 		} Class;
 	};
+
+	EAnimNodeValueType Type = EAnimNodeValueType::Bool;
+};
+
+struct FAnimNodeValueMessage
+{
+	const TCHAR* Key = nullptr;
+	uint64 AnimInstanceId = 0;
+	FVariantValue Value;
 	int32 NodeId = -1;
 	uint16 FrameCounter = 0;
-	EAnimNodeValueType Type = EAnimNodeValueType::Bool;
 };
 
 struct FAnimSequencePlayerMessage
@@ -178,6 +195,36 @@ struct FAnimStateMachineMessage
 	float ElapsedTime = 0.0f;
 };
 
+enum class EAnimNotifyMessageType : uint8
+{
+	Event = 0,
+	Begin = 1,
+	End = 2,
+	Tick = 3,
+	SyncMarker = 4	// We 'fake' sync markers with a notify type for convenience
+};
+
+struct FAnimNotifyMessage
+{
+	uint64 AnimInstanceId = 0;
+	uint64 AssetId = 0;
+	uint64 NotifyId = 0;
+	uint32 NameId = 0;
+	float Time = 0.0f; 
+	float Duration = 0.0f;
+	EAnimNotifyMessageType NotifyEventType = EAnimNotifyMessageType::Event;
+};
+
+struct FAnimMontageMessage
+{
+	uint64 AnimInstanceId = 0;
+	uint64 MontageId = 0;
+	uint32 CurrentSectionNameId = 0;
+	uint32 NextSectionNameId = 0;
+	float Weight = 0.0f;
+	float DesiredWeight = 0.0f;
+};
+
 class IAnimationProvider : public Trace::IProvider
 {
 public:
@@ -189,6 +236,8 @@ public:
 	typedef Trace::ITimeline<FAnimSequencePlayerMessage> AnimSequencePlayersTimeline;
 	typedef Trace::ITimeline<FAnimStateMachineMessage> StateMachinesTimeline;
 	typedef Trace::ITimeline<FBlendSpacePlayerMessage> BlendSpacePlayersTimeline;
+	typedef Trace::ITimeline<FAnimNotifyMessage> AnimNotifyTimeline;
+	typedef Trace::ITimeline<FAnimMontageMessage> AnimMontageTimeline;
 
 	virtual bool ReadSkeletalMeshPoseTimeline(uint64 InObjectId, TFunctionRef<void(const SkeletalMeshPoseTimeline&, bool)> Callback) const = 0;
 	virtual void GetSkeletalMeshComponentSpacePose(const FSkeletalMeshPoseMessage& InMessage, const FSkeletalMeshInfo& InMeshInfo, FTransform& OutComponentToWorld, TArray<FTransform>& OutTransforms) const = 0;
@@ -202,6 +251,10 @@ public:
 	virtual bool ReadAnimSequencePlayersTimeline(uint64 InObjectId, TFunctionRef<void(const AnimSequencePlayersTimeline&)> Callback) const = 0;
 	virtual bool ReadAnimBlendSpacePlayersTimeline(uint64 InObjectId, TFunctionRef<void(const BlendSpacePlayersTimeline&)> Callback) const = 0;
 	virtual bool ReadStateMachinesTimeline(uint64 InObjectId, TFunctionRef<void(const StateMachinesTimeline&)> Callback) const = 0;
+	virtual bool ReadNotifyTimeline(uint64 InObjectId, TFunctionRef<void(const AnimNotifyTimeline&)> Callback) const = 0;
+	virtual void EnumerateNotifyStateTimelines(uint64 InObjectId, TFunctionRef<void(uint64, const AnimNotifyTimeline&)> Callback) const = 0;
+	virtual bool ReadMontageTimeline(uint64 InObjectId, TFunctionRef<void(const AnimMontageTimeline&)> Callback) const = 0;
+	virtual void EnumerateMontageIds(uint64 InObjectId, TFunctionRef<void(uint64)> Callback) const = 0;
 	virtual const FSkeletalMeshInfo* FindSkeletalMeshInfo(uint64 InObjectId) const = 0;
 	virtual const TCHAR* GetName(uint32 InId) const = 0;
 	virtual FText FormatNodeKeyValue(const FAnimNodeValueMessage& InMessage) const = 0;
diff --git a/Engine/Plugins/Animation/GameplayInsights/Source/GameplayInsights/Private/MontageTrack.cpp b/Engine/Plugins/Animation/GameplayInsights/Source/GameplayInsights/Private/MontageTrack.cpp
new file mode 100644
--- /dev/null
+++ b/Engine/Plugins/Animation/GameplayInsights/Source/GameplayInsights/Private/MontageTrack.cpp
@@ -0,0 +1,256 @@
+// Copyright Epic Games, Inc. All Rights Reserved.
+
+#include "MontageTrack.h"
+#include "GameplayProvider.h"
+#include "AnimationProvider.h"
+#include "Insights/ViewModels/TimingTrackViewport.h"
+#include "Insights/ViewModels/TimingEvent.h"
+#include "AnimationSharedData.h"
+#include "Framework/MultiBox/MultiBoxBuilder.h"
+#include "Insights/ViewModels/TimingEventSearch.h"
+#include "Insights/ViewModels/TooltipDrawState.h"
+#include "Modules/ModuleManager.h"
+#include "Insights/ViewModels/GraphTrackBuilder.h"
+#include "VariantTreeNode.h"
+#include "TraceServices/Model/Frames.h"
+
+#define LOCTEXT_NAMESPACE "MontageTrack"
+
+INSIGHTS_IMPLEMENT_RTTI(FMontageTrack)
+
+FMontageTrack::FMontageTrack(const FAnimationSharedData& InSharedData, uint64 InObjectID, const TCHAR* InName)
+	: FGameplayGraphTrack(InSharedData.GetGameplaySharedData(), InObjectID, FText::Format(LOCTEXT("TrackNameFormat", "Montage - {0}"), FText::FromString(FString(InName))))
+	, SharedData(InSharedData)
+{
+	bDrawLabels = true;
+	Layout = EGameplayGraphLayout::Stack;
+}
+
+void FMontageTrack::AddAllSeries()
+{
+	const FAnimationProvider* AnimationProvider = SharedData.GetAnalysisSession().ReadProvider<FAnimationProvider>(FAnimationProvider::ProviderName);
+	const FGameplayProvider* GameplayProvider = SharedData.GetAnalysisSession().ReadProvider<FGameplayProvider>(FGameplayProvider::ProviderName);
+
+	if(AnimationProvider && GameplayProvider)
+	{
+		Trace::FAnalysisSessionReadScope SessionReadScope(SharedData.GetAnalysisSession());
+
+		AnimationProvider->EnumerateMontageIds(GetGameplayTrack().GetObjectId(), [this, &AnimationProvider, &GameplayProvider](uint64 InMontageId)
+		{
+			if(!AllSeries.ContainsByPredicate([InMontageId](const TSharedPtr<FGraphSeries>& InSeries){ return StaticCastSharedPtr<FMontageSeries>(InSeries)->MontageId == InMontageId; }))
+			{
+				auto MakeCurveSeriesColor = [](uint64 InSeed, bool bInLine)
+				{
+					FRandomStream Stream(GetTypeHash(InSeed));
+					const uint8 Hue = (uint8)(Stream.FRand() * 255.0f);
+					const uint8 SatVal = bInLine ? 196 : 128;
+					return FLinearColor::MakeFromHSV8(Hue, SatVal, SatVal);
+				};
+
+				TSharedRef<FMontageSeries> Series = MakeShared<FMontageSeries>();
+
+				const FObjectInfo& MontageInfo = GameplayProvider->GetObjectInfo(InMontageId);
+
+				Series->SetName(MontageInfo.Name);
+				Series->SetDescription(FText::Format(LOCTEXT("MontageTooltipFormat", "Weight for montage '{0}'"), FText::FromString(MontageInfo.Name)));
+
+				const FLinearColor LineColor = MakeCurveSeriesColor(InMontageId, true);
+				const FLinearColor FillColor = MakeCurveSeriesColor(InMontageId, false);
+				Series->SetColor(LineColor, LineColor, FillColor);
+
+				Series->MontageId = InMontageId;
+				Series->SetVisibility(true);
+				Series->SetBaselineY(25.0f);
+				Series->SetScaleY(20.0f);
+				Series->EnableAutoZoom();
+				AllSeries.Add(Series);
+			}
+		});
+	}
+}
+
+bool FMontageTrack::UpdateSeriesBounds(FGameplayGraphSeries& InSeries, const FTimingTrackViewport& InViewport)
+{
+	bool bFoundEvents = false;
+
+	const FAnimationProvider* AnimationProvider = SharedData.GetAnalysisSession().ReadProvider<FAnimationProvider>(FAnimationProvider::ProviderName);
+
+	if(AnimationProvider)
+	{
+		Trace::FAnalysisSessionReadScope SessionReadScope(SharedData.GetAnalysisSession());
+
+		FMontageSeries& MontageSeries = *static_cast<FMontageSeries*>(&InSeries);
+
+		MontageSeries.CurrentMin = 0.0f;
+		MontageSeries.CurrentMax = 0.0f;
+
+		AnimationProvider->ReadMontageTimeline(GetGameplayTrack().GetObjectId(), [&bFoundEvents, &InViewport, &MontageSeries, &AnimationProvider](const FAnimationProvider::AnimMontageTimeline& InTimeline)
+		{
+			InTimeline.EnumerateEvents(InViewport.GetStartTime(), InViewport.GetEndTime(), [&bFoundEvents, &MontageSeries, &AnimationProvider](double InStartTime, double InEndTime, uint32 InDepth, const FAnimMontageMessage& InMessage)
+			{
+				if(InMessage.MontageId == MontageSeries.MontageId)
+				{
+					MontageSeries.CurrentMin = FMath::Min(MontageSeries.CurrentMin, InMessage.Weight);
+					MontageSeries.CurrentMax = FMath::Max(MontageSeries.CurrentMax, InMessage.Weight);
+					bFoundEvents = true;
+				}
+			});
+		});
+	}
+
+	return bFoundEvents;
+}
+
+void FMontageTrack::UpdateSeries(FGameplayGraphSeries& InSeries, const FTimingTrackViewport& InViewport)
+{
+	const FAnimationProvider* AnimationProvider = SharedData.GetAnalysisSession().ReadProvider<FAnimationProvider>(FAnimationProvider::ProviderName);
+
+	if(AnimationProvider)
+	{
+		Trace::FAnalysisSessionReadScope SessionReadScope(SharedData.GetAnalysisSession());
+
+		FMontageSeries& MontageSeries = *static_cast<FMontageSeries*>(&InSeries);
+
+		FGraphTrackBuilder Builder(*this, MontageSeries, InViewport);
+
+		AnimationProvider->ReadMontageTimeline(GetGameplayTrack().GetObjectId(), [this, &Builder, &InViewport, &MontageSeries, &AnimationProvider](const FAnimationProvider::AnimMontageTimeline& InTimeline)
+		{
+			int32 FrameCounter = 0;
+			int32 LastFrameWithMontage = 0;
+
+			InTimeline.EnumerateEvents(InViewport.GetStartTime(), InViewport.GetEndTime(), [this, &FrameCounter, &LastFrameWithMontage, &Builder, &MontageSeries, &AnimationProvider](double InStartTime, double InEndTime, uint32 InDepth, const FAnimMontageMessage& InMessage)
+			{
+				FrameCounter++;
+
+				if(InMessage.MontageId == MontageSeries.MontageId)
+				{
+					Builder.AddEvent(InStartTime, InEndTime - InStartTime, InMessage.Weight, LastFrameWithMontage == FrameCounter - 1);
+
+					LastFrameWithMontage = FrameCounter;
+				}
+			});
+		});
+	}
+}
+
+void FMontageTrack::InitTooltip(FTooltipDrawState& Tooltip, const ITimingEvent& HoveredTimingEvent) const
+{
+	const FGraphTrackEvent& GraphTrackEvent = *static_cast<const FGraphTrackEvent*>(&HoveredTimingEvent);
+
+	FTimingEventSearchParameters SearchParameters(HoveredTimingEvent.GetStartTime(), HoveredTimingEvent.GetEndTime(), ETimingEventSearchFlags::StopAtFirstMatch);
+
+	FindMontageMessage(SearchParameters, [this, &Tooltip, &GraphTrackEvent](double InFoundStartTime, double InFoundEndTime, uint32 InFoundDepth, const FAnimMontageMessage& InMessage)
+	{
+		Tooltip.ResetContent();
+
+		Tooltip.AddTitle(GetName());
+
+		Tooltip.AddNameValueTextLine(LOCTEXT("EventTime", "Time").ToString(), FText::AsNumber(InFoundStartTime).ToString());
+		Tooltip.AddNameValueTextLine(LOCTEXT("EventWeight", "Weight").ToString(), FText::AsNumber(GraphTrackEvent.GetValue()).ToString());
+		Tooltip.AddNameValueTextLine(LOCTEXT("EventDesiredWeight", "Desired Weight").ToString(), FText::AsNumber(InMessage.DesiredWeight).ToString());
+
+		{
+			Trace::FAnalysisSessionReadScope SessionReadScope(SharedData.GetAnalysisSession());
+
+			const FGameplayProvider* GameplayProvider = SharedData.GetAnalysisSession().ReadProvider<FGameplayProvider>(FGameplayProvider::ProviderName);
+			if(GameplayProvider)
+			{
+				const FObjectInfo& MontageInfo = GameplayProvider->GetObjectInfo(InMessage.MontageId);
+				Tooltip.AddNameValueTextLine(LOCTEXT("MontageName", "Montage").ToString(), MontageInfo.PathName);
+			}
+
+			const FAnimationProvider* AnimationProvider = SharedData.GetAnalysisSession().ReadProvider<FAnimationProvider>(FAnimationProvider::ProviderName);
+			if(AnimationProvider)
+			{
+				const TCHAR* CurrentSectionName = AnimationProvider->GetName(InMessage.CurrentSectionNameId);
+				Tooltip.AddNameValueTextLine(LOCTEXT("CurrentSectionName", "Current Section").ToString(), FText::FromString(CurrentSectionName).ToString());
+				const TCHAR* NextSectionName = AnimationProvider->GetName(InMessage.NextSectionNameId);
+				Tooltip.AddNameValueTextLine(LOCTEXT("NextSectionName", "Next Section").ToString(), FText::FromString(NextSectionName).ToString());
+			}
+		}
+
+		Tooltip.UpdateLayout();
+	});
+}
+
+const TSharedPtr<const ITimingEvent> FMontageTrack::SearchEvent(const FTimingEventSearchParameters& InSearchParameters) const
+{
+	TSharedPtr<const ITimingEvent> FoundEvent;
+
+	FindMontageMessage(InSearchParameters, [this, &FoundEvent](double InFoundStartTime, double InFoundEndTime, uint32 InFoundDepth, const FAnimMontageMessage& InFoundMessage)
+	{
+		FoundEvent = MakeShared<const FTimingEvent>(SharedThis(this), InFoundStartTime, InFoundEndTime, InFoundDepth);
+	});
+
+	return FoundEvent;
+}
+
+void FMontageTrack::FindMontageMessage(const FTimingEventSearchParameters& InParameters, TFunctionRef<void(double, double, uint32, const FAnimMontageMessage&)> InFoundPredicate) const
+{
+	TTimingEventSearch<FAnimMontageMessage>::Search(
+		InParameters,
+
+		[this](TTimingEventSearch<FAnimMontageMessage>::FContext& InContext)
+		{
+			const FAnimationProvider* AnimationProvider = SharedData.GetAnalysisSession().ReadProvider<FAnimationProvider>(FAnimationProvider::ProviderName);
+
+			if(AnimationProvider)
+			{
+				Trace::FAnalysisSessionReadScope SessionReadScope(SharedData.GetAnalysisSession());
+
+				AnimationProvider->ReadMontageTimeline(GetGameplayTrack().GetObjectId(), [&InContext](const FAnimationProvider::AnimMontageTimeline& InTimeline)
+				{
+					InTimeline.EnumerateEvents(InContext.GetParameters().StartTime, InContext.GetParameters().EndTime, [&InContext](double InEventStartTime, double InEventEndTime, uint32 InDepth, const FAnimMontageMessage& InMessage)
+					{
+						InContext.Check(InEventStartTime, InEventEndTime, 0, InMessage);
+					});
+				});
+			}
+		},
+
+		[&InFoundPredicate](double InFoundStartTime, double InFoundEndTime, uint32 InFoundDepth, const FAnimMontageMessage& InEvent)
+		{
+			InFoundPredicate(InFoundStartTime, InFoundEndTime, InFoundDepth, InEvent);
+		});
+}
+
+void FMontageTrack::GetVariantsAtTime(double InTime, TArray<TSharedRef<FVariantTreeNode>>& OutVariants) const 
+{
+	TSharedRef<FVariantTreeNode> Header = OutVariants.Add_GetRef(FVariantTreeNode::MakeHeader(LOCTEXT("MontagesHeader", "Montages")));
+
+	const Trace::IFrameProvider& FramesProvider = Trace::ReadFrameProvider(SharedData.GetAnalysisSession());
+	const FAnimationProvider* AnimationProvider = SharedData.GetAnalysisSession().ReadProvider<FAnimationProvider>(FAnimationProvider::ProviderName);
+	const FGameplayProvider* GameplayProvider = SharedData.GetAnalysisSession().ReadProvider<FGameplayProvider>(FGameplayProvider::ProviderName);
+
+	if(AnimationProvider && GameplayProvider)
+	{
+		Trace::FAnalysisSessionReadScope SessionReadScope(SharedData.GetAnalysisSession());
+
+		AnimationProvider->ReadMontageTimeline(GetGameplayTrack().GetObjectId(), [this, &InTime, &AnimationProvider, &FramesProvider, &GameplayProvider, &Header](const FAnimationProvider::AnimMontageTimeline& InTimeline)
+		{
+			// round to nearest frame boundary
+			Trace::FFrame Frame;
+			if(FramesProvider.GetFrameFromTime(ETraceFrameType::TraceFrameType_Game, InTime, Frame))
+			{
+				InTimeline.EnumerateEvents(Frame.StartTime, Frame.EndTime, [this, &AnimationProvider, &GameplayProvider, &Header, &Frame](double InStartTime, double InEndTime, uint32 InDepth, const FAnimMontageMessage& InMessage)
+				{
+					if(InStartTime >= Frame.StartTime && InEndTime <= Frame.EndTime)
+					{
+						const FObjectInfo& MontageInfo = GameplayProvider->GetObjectInfo(InMessage.MontageId);
+						TSharedRef<FVariantTreeNode> MontageHeader = Header->AddChild(FVariantTreeNode::MakeObject(FText::FromString(MontageInfo.Name), InMessage.MontageId));
+
+						MontageHeader->AddChild(FVariantTreeNode::MakeFloat(LOCTEXT("EventWeight", "Weight"), InMessage.Weight));
+						MontageHeader->AddChild(FVariantTreeNode::MakeFloat(LOCTEXT("EventDesiredWeight", "Desired Weight"), InMessage.DesiredWeight));
+
+						const TCHAR* CurrentSectionName = AnimationProvider->GetName(InMessage.CurrentSectionNameId);
+						MontageHeader->AddChild(FVariantTreeNode::MakeString(LOCTEXT("CurrentSectionName", "Current Section"), CurrentSectionName));
+						const TCHAR* NextSectionName = AnimationProvider->GetName(InMessage.NextSectionNameId);
+						MontageHeader->AddChild(FVariantTreeNode::MakeString(LOCTEXT("NextSectionName", "Next Section"), NextSectionName));
+					}
+				});
+			}
+		});
+	}
+}
+
+#undef LOCTEXT_NAMESPACE
diff --git a/Engine/Plugins/Animation/GameplayInsights/Source/GameplayInsights/Private/MontageTrack.h b/Engine/Plugins/Animation/GameplayInsights/Source/GameplayInsights/Private/MontageTrack.h
new file mode 100644
--- /dev/null
+++ b/Engine/Plugins/Animation/GameplayInsights/Source/GameplayInsights/Private/MontageTrack.h
@@ -0,0 +1,39 @@
+// Copyright Epic Games, Inc. All Rights Reserved.
+
+#pragma once
+
+#include "GameplayGraphTrack.h"
+
+class FAnimationSharedData;
+class FTimingEventSearchParameters;
+struct FAnimMontageMessage;
+class FMontageTrack;
+
+class FMontageSeries : public FGameplayGraphSeries
+{
+public:
+	uint64 MontageId;
+};
+
+class FMontageTrack : public FGameplayGraphTrack
+{
+	INSIGHTS_DECLARE_RTTI(FMontageTrack, FGameplayGraphTrack)
+
+public:
+	FMontageTrack(const FAnimationSharedData& InSharedData, uint64 InObjectID, const TCHAR* InName);
+
+	virtual void InitTooltip(FTooltipDrawState& Tooltip, const ITimingEvent& HoveredTimingEvent) const override;
+	virtual const TSharedPtr<const ITimingEvent> SearchEvent(const FTimingEventSearchParameters& InSearchParameters) const override;
+	virtual void AddAllSeries() override;
+	virtual bool UpdateSeriesBounds(FGameplayGraphSeries& InSeries, const FTimingTrackViewport& InViewport) override;
+	virtual void UpdateSeries(FGameplayGraphSeries& InSeries, const FTimingTrackViewport& InViewport) override;
+	virtual void GetVariantsAtTime(double InTime, TArray<TSharedRef<FVariantTreeNode>>& OutVariants) const override;
+
+private:
+	// Helper function used to find a montage message
+	void FindMontageMessage(const FTimingEventSearchParameters& InParameters, TFunctionRef<void(double, double, uint32, const FAnimMontageMessage&)> InFoundPredicate) const;
+
+private:
+	/** The shared data */
+	const FAnimationSharedData& SharedData;
+};
diff --git a/Engine/Plugins/Animation/GameplayInsights/Source/GameplayInsights/Private/ObjectEventsTrack.cpp b/Engine/Plugins/Animation/GameplayInsights/Source/GameplayInsights/Private/ObjectEventsTrack.cpp
--- a/Engine/Plugins/Animation/GameplayInsights/Source/GameplayInsights/Private/ObjectEventsTrack.cpp
+++ b/Engine/Plugins/Animation/GameplayInsights/Source/GameplayInsights/Private/ObjectEventsTrack.cpp
@@ -8,13 +8,15 @@
 #include "Insights/ViewModels/TooltipDrawState.h"
 #include "GameplaySharedData.h"
 #include "Insights/ViewModels/TimingEventSearch.h"
+#include "TraceServices/Model/Frames.h"
+#include "VariantTreeNode.h"
 
 #define LOCTEXT_NAMESPACE "ObjectEventsTrack"
 
 INSIGHTS_IMPLEMENT_RTTI(FObjectEventsTrack)
 
 FObjectEventsTrack::FObjectEventsTrack(const FGameplaySharedData& InSharedData, uint64 InObjectID, const TCHAR* InName)
-	: TGameplayTrackMixin<FTimingEventsTrack>(InObjectID, MakeTrackName(InSharedData, InObjectID, InName))
+	: FGameplayTimingEventsTrack(InSharedData, InObjectID, MakeTrackName(InSharedData, InObjectID, InName))
 	, SharedData(InSharedData)
 {
 }
@@ -121,4 +123,31 @@ FText FObjectEventsTrack::MakeTrackName(const FGameplaySharedData& InSharedData,
 	return FText::Format(LOCTEXT("ObjectEventsTrackName", "{0} - {1}"), ClassName, FText::FromString(InName));
 }
 
+void FObjectEventsTrack::GetVariantsAtTime(double InTime, TArray<TSharedRef<FVariantTreeNode>>& OutVariants) const
+{
+	const FGameplayProvider* GameplayProvider = SharedData.GetAnalysisSession().ReadProvider<FGameplayProvider>(FGameplayProvider::ProviderName);
+	if(GameplayProvider)
+	{
+		Trace::FAnalysisSessionReadScope SessionReadScope(SharedData.GetAnalysisSession());
+
+		const Trace::IFrameProvider& FramesProvider = Trace::ReadFrameProvider(SharedData.GetAnalysisSession());
+
+		TSharedRef<FVariantTreeNode> Header = OutVariants.Add_GetRef(FVariantTreeNode::MakeHeader(FText::FromString(GetName())));
+
+		// object events
+		GameplayProvider->ReadObjectEventsTimeline(GetGameplayTrack().GetObjectId(), [&InTime, &FramesProvider, &Header](const FGameplayProvider::ObjectEventsTimeline& InTimeline)
+		{
+			// round to nearest frame boundary
+			Trace::FFrame Frame;
+			if(FramesProvider.GetFrameFromTime(ETraceFrameType::TraceFrameType_Game, InTime, Frame))
+			{
+				InTimeline.EnumerateEvents(Frame.StartTime, Frame.EndTime, [&Header](double InStartTime, double InEndTime, uint32 InDepth, const FObjectEventMessage& InMessage)
+				{
+					Header->AddChild(FVariantTreeNode::MakeFloat(FText::FromString(InMessage.Name), InStartTime));
+				});
+			}
+		});
+	}
+}
+
 #undef LOCTEXT_NAMESPACE
\ No newline at end of file
diff --git a/Engine/Plugins/Animation/GameplayInsights/Source/GameplayInsights/Private/ObjectEventsTrack.h b/Engine/Plugins/Animation/GameplayInsights/Source/GameplayInsights/Private/ObjectEventsTrack.h
--- a/Engine/Plugins/Animation/GameplayInsights/Source/GameplayInsights/Private/ObjectEventsTrack.h
+++ b/Engine/Plugins/Animation/GameplayInsights/Source/GameplayInsights/Private/ObjectEventsTrack.h
@@ -9,9 +9,9 @@ class FGameplaySharedData;
 struct FObjectEventMessage;
 class FTimingEventSearchParameters;
 
-class FObjectEventsTrack : public TGameplayTrackMixin<FTimingEventsTrack>
+class FObjectEventsTrack : public FGameplayTimingEventsTrack
 {
-	INSIGHTS_DECLARE_RTTI(FObjectEventsTrack, TGameplayTrackMixin<FTimingEventsTrack>)
+	INSIGHTS_DECLARE_RTTI(FObjectEventsTrack, FGameplayTimingEventsTrack)
 
 public:
 	FObjectEventsTrack(const FGameplaySharedData& InSharedData, uint64 InObjectID, const TCHAR* InName);
@@ -20,6 +20,7 @@ public:
 	virtual void Draw(const ITimingTrackDrawContext& Context) const override;
 	virtual void InitTooltip(FTooltipDrawState& Tooltip, const ITimingEvent& HoveredTimingEvent) const override;
 	virtual const TSharedPtr<const ITimingEvent> SearchEvent(const FTimingEventSearchParameters& InSearchParameters) const override;
+	virtual void GetVariantsAtTime(double InTime, TArray<TSharedRef<FVariantTreeNode>>& OutVariants) const override;
 
 private:
 	// Helper function used to find an object event
diff --git a/Engine/Plugins/Animation/GameplayInsights/Source/GameplayInsights/Private/SAnimGraphSchematicView.cpp b/Engine/Plugins/Animation/GameplayInsights/Source/GameplayInsights/Private/SAnimGraphSchematicView.cpp
--- a/Engine/Plugins/Animation/GameplayInsights/Source/GameplayInsights/Private/SAnimGraphSchematicView.cpp
+++ b/Engine/Plugins/Animation/GameplayInsights/Source/GameplayInsights/Private/SAnimGraphSchematicView.cpp
@@ -47,7 +47,7 @@ enum class EAnimGraphSchematicFilterState
 	Highlighted,
 };
 
-// A in the tree of 'properties' for an animation node's debug info
+// A node in the tree of 'properties' for an animation node's debug info
 class FAnimGraphSchematicPropertyNode
 {
 public:
@@ -72,13 +72,13 @@ public:
 
 	static TSharedRef<SWidget> StaticMakeValueWidget(const Trace::IAnalysisSession& InAnalysisSession, const TSharedRef<FAnimNodeValueMessage>& InValue) 
 	{ 
-		switch(InValue->Type)
+		switch(InValue->Value.Type)
 		{
 		case EAnimNodeValueType::Bool:
 			return 
 				SNew(SCheckBox)
 				.IsEnabled(false)
-				.IsChecked(InValue->Bool.bValue ? ECheckBoxState::Checked : ECheckBoxState::Unchecked);
+				.IsChecked(InValue->Value.Bool.bValue ? ECheckBoxState::Checked : ECheckBoxState::Unchecked);
 
 		case EAnimNodeValueType::Int32:
 			return
@@ -88,7 +88,7 @@ public:
 					SNew(SEditableTextBox)
 					.IsEnabled(false)
 					.Font(FCoreStyle::Get().GetFontStyle("SmallFont"))
-					.Text(FText::AsNumber(InValue->Int32.Value))
+					.Text(FText::AsNumber(InValue->Value.Int32.Value))
 				];
 
 		case EAnimNodeValueType::Float:
@@ -99,7 +99,34 @@ public:
 					SNew(SEditableTextBox)
 					.IsEnabled(false)
 					.Font(FCoreStyle::Get().GetFontStyle("SmallFont"))
-					.Text(FText::AsNumber(InValue->Float.Value))
+					.Text(FText::AsNumber(InValue->Value.Float.Value))
+				];
+
+		case EAnimNodeValueType::Vector2D:
+			return SNew(SHorizontalBox)
+				+SHorizontalBox::Slot()
+				.AutoWidth()
+				[
+					SNew(SBox)
+					.WidthOverride(125.0f)
+					[
+						SNew(SEditableTextBox)
+						.IsEnabled(false)
+						.Font(FCoreStyle::Get().GetFontStyle("SmallFont"))
+						.Text(FText::AsNumber(InValue->Value.Vector2D.Value.X))
+					]
+				]
+				+SHorizontalBox::Slot()
+				.AutoWidth()
+				[
+					SNew(SBox)
+					.WidthOverride(125.0f)
+					[
+						SNew(SEditableTextBox)
+						.IsEnabled(false)
+						.Font(FCoreStyle::Get().GetFontStyle("SmallFont"))
+						.Text(FText::AsNumber(InValue->Value.Vector2D.Value.Y))
+					]
 				];
 
 		case EAnimNodeValueType::Vector:
@@ -113,7 +140,7 @@ public:
 						SNew(SEditableTextBox)
 						.IsEnabled(false)
 						.Font(FCoreStyle::Get().GetFontStyle("SmallFont"))
-						.Text(FText::AsNumber(InValue->Vector.Value.X))
+						.Text(FText::AsNumber(InValue->Value.Vector.Value.X))
 					]
 				]
 				+SHorizontalBox::Slot()
@@ -125,7 +152,7 @@ public:
 						SNew(SEditableTextBox)
 						.IsEnabled(false)
 						.Font(FCoreStyle::Get().GetFontStyle("SmallFont"))
-						.Text(FText::AsNumber(InValue->Vector.Value.Y))
+						.Text(FText::AsNumber(InValue->Value.Vector.Value.Y))
 					]
 				]
 				+SHorizontalBox::Slot()
@@ -137,7 +164,7 @@ public:
 						SNew(SEditableTextBox)
 						.IsEnabled(false)
 						.Font(FCoreStyle::Get().GetFontStyle("SmallFont"))
-						.Text(FText::AsNumber(InValue->Vector.Value.Z))
+						.Text(FText::AsNumber(InValue->Value.Vector.Value.Z))
 					]
 				];
 
@@ -145,7 +172,7 @@ public:
 			return 
 				SNew(STextBlock)
 				.Font(FCoreStyle::Get().GetFontStyle("SmallFont"))
-				.Text(FText::FromString(InValue->String.Value));
+				.Text(FText::FromString(InValue->Value.String.Value));
 
 		case EAnimNodeValueType::Object:
 		{
@@ -154,7 +181,7 @@ public:
 			{
 				Trace::FAnalysisSessionReadScope SessionReadScope(InAnalysisSession);
 
-				const FObjectInfo& ObjectInfo = GameplayProvider->GetObjectInfo(InValue->Object.Value);
+				const FObjectInfo& ObjectInfo = GameplayProvider->GetObjectInfo(InValue->Value.Object.Value);
 #if WITH_EDITOR
 				return 
 					SNew(SHyperlink)
@@ -183,7 +210,7 @@ public:
 			{
 				Trace::FAnalysisSessionReadScope SessionReadScope(InAnalysisSession);
 
-				const FClassInfo& ClassInfo = GameplayProvider->GetClassInfo(InValue->Class.Value);
+				const FClassInfo& ClassInfo = GameplayProvider->GetClassInfo(InValue->Value.Class.Value);
 #if WITH_EDITOR
 				return 
 					SNew(SHyperlink)
@@ -600,9 +627,9 @@ void SAnimGraphSchematicView::RefreshNodes()
 										TSharedRef<FAnimNodeValueMessage> WeightMessage = MakeShared<FAnimNodeValueMessage>();
 										WeightMessage->NodeId = InMessage.NodeId;
 										WeightMessage->FrameCounter = InMessage.FrameCounter;
-										WeightMessage->Type = EAnimNodeValueType::Float;
+										WeightMessage->Value.Type = EAnimNodeValueType::Float;
 										WeightMessage->Key = TEXT("Weight");
-										WeightMessage->Float.Value = InMessage.Weight;
+										WeightMessage->Value.Float.Value = InMessage.Weight;
 
 										Node->KeysAndValues.Add(AnimGraphSchematicColumns::Weight, WeightMessage);
 										Node->Values.Add(WeightMessage);
@@ -610,9 +637,9 @@ void SAnimGraphSchematicView::RefreshNodes()
 										TSharedRef<FAnimNodeValueMessage> RootMotionWeightMessage = MakeShared<FAnimNodeValueMessage>();
 										RootMotionWeightMessage->NodeId = InMessage.NodeId;
 										RootMotionWeightMessage->FrameCounter = InMessage.FrameCounter;
-										RootMotionWeightMessage->Type = EAnimNodeValueType::Float;
+										RootMotionWeightMessage->Value.Type = EAnimNodeValueType::Float;
 										RootMotionWeightMessage->Key = TEXT("Root Motion Weight");
-										RootMotionWeightMessage->Float.Value = InMessage.RootMotionWeight;
+										RootMotionWeightMessage->Value.Float.Value = InMessage.RootMotionWeight;
 
 										Node->KeysAndValues.Add(AnimGraphSchematicColumns::RootMotionWeight, RootMotionWeightMessage);
 										Node->Values.Add(RootMotionWeightMessage);
diff --git a/Engine/Plugins/Animation/GameplayInsights/Source/GameplayInsights/Private/STrackVariantValueView.cpp b/Engine/Plugins/Animation/GameplayInsights/Source/GameplayInsights/Private/STrackVariantValueView.cpp
new file mode 100644
--- /dev/null
+++ b/Engine/Plugins/Animation/GameplayInsights/Source/GameplayInsights/Private/STrackVariantValueView.cpp
@@ -0,0 +1,42 @@
+// Copyright Epic Games, Inc. All Rights Reserved.
+
+#include "STrackVariantValueView.h"
+#include "SVariantValueView.h"
+#include "Insights/ViewModels/BaseTimingTrack.h"
+#include "GameplayTrack.h"
+#include "GameplayGraphTrack.h"
+#include "Insights/ITimingViewSession.h"
+
+void STrackVariantValueView::Construct(const FArguments& InArgs, const TSharedRef<FBaseTimingTrack>& InTimingTrack, Insights::ITimingViewSession& InTimingViewSession, const Trace::IAnalysisSession& InAnalysisSession)
+{
+	TimingTrack = InTimingTrack;
+
+	InTimingViewSession.OnTimeMarkerChanged().AddSP(this, &STrackVariantValueView::HandleTimeMarkerChanged);
+
+	ChildSlot
+	[
+		SAssignNew(VariantValueView, SVariantValueView, InAnalysisSession)
+		.OnGetVariantValues_Lambda([this](double InTime, TArray<TSharedRef<FVariantTreeNode>>& OutValues)
+		{
+			TSharedPtr<FBaseTimingTrack> PinnedTrack = TimingTrack.Pin();
+			if(PinnedTrack.IsValid())
+			{
+				if(PinnedTrack->Is<FGameplayTimingEventsTrack>())
+				{
+					StaticCastSharedPtr<FGameplayTimingEventsTrack>(PinnedTrack)->GetVariantsAtTime(InTime, OutValues);
+				}
+				else if(PinnedTrack->Is<FGameplayGraphTrack>())
+				{
+					StaticCastSharedPtr<FGameplayGraphTrack>(PinnedTrack)->GetVariantsAtTime(InTime, OutValues);
+				}
+			}
+		})
+	];
+
+	VariantValueView->RequestRefresh(InTimingViewSession.GetTimeMarker());
+}
+
+void STrackVariantValueView::HandleTimeMarkerChanged(Insights::ETimeChangedFlags InFlags, double InTimeMarker)
+{
+	VariantValueView->RequestRefresh(InTimeMarker);
+}
diff --git a/Engine/Plugins/Animation/GameplayInsights/Source/GameplayInsights/Private/STrackVariantValueView.h b/Engine/Plugins/Animation/GameplayInsights/Source/GameplayInsights/Private/STrackVariantValueView.h
new file mode 100644
--- /dev/null
+++ b/Engine/Plugins/Animation/GameplayInsights/Source/GameplayInsights/Private/STrackVariantValueView.h
@@ -0,0 +1,35 @@
+// Copyright Epic Games, Inc. All Rights Reserved.
+
+#pragma once
+
+#include "CoreMinimal.h"
+#include "Widgets/SCompoundWidget.h"
+#include "Widgets/DeclarativeSyntaxSupport.h"
+
+class FBaseTimingTrack;
+namespace Insights { class ITimingViewSession; }
+namespace Trace { class IAnalysisSession; }
+namespace Insights { enum class ETimeChangedFlags; }
+class SVariantValueView;
+
+// Wrapper for a variant value view for a track
+class STrackVariantValueView : public SCompoundWidget
+{
+public:
+	SLATE_BEGIN_ARGS(STrackVariantValueView) {}
+
+	SLATE_END_ARGS()
+
+	void Construct(const FArguments& InArgs, const TSharedRef<FBaseTimingTrack>& InTimingTrack, Insights::ITimingViewSession& InTimingViewSession, const Trace::IAnalysisSession& InAnalysisSession);
+
+	TSharedPtr<FBaseTimingTrack> GetTimingTrack() const { return TimingTrack.Pin(); }
+
+private:
+	// Handle the time marker being scrubbed
+	void HandleTimeMarkerChanged(Insights::ETimeChangedFlags InFlags, double InTimeMarker);
+
+private:
+	TWeakPtr<FBaseTimingTrack> TimingTrack;
+
+	TSharedPtr<SVariantValueView> VariantValueView;
+};
\ No newline at end of file
diff --git a/Engine/Plugins/Animation/GameplayInsights/Source/GameplayInsights/Private/SVariantValueView.cpp b/Engine/Plugins/Animation/GameplayInsights/Source/GameplayInsights/Private/SVariantValueView.cpp
new file mode 100644
--- /dev/null
+++ b/Engine/Plugins/Animation/GameplayInsights/Source/GameplayInsights/Private/SVariantValueView.cpp
@@ -0,0 +1,346 @@
+// Copyright Epic Games, Inc. All Rights Reserved.
+
+#include "SVariantValueView.h"
+#include "AnimationProvider.h"
+#include "Widgets/Input/SSearchBox.h"
+#include "TraceServices/Model/AnalysisSession.h"
+#include "Insights/ITimingViewSession.h"
+#include "GameplayProvider.h"
+#include "TraceServices/Model/Frames.h"
+#include "Widgets/Layout/SBorder.h"
+#include "Widgets/SBoxPanel.h"
+#include "Styling/CoreStyle.h"
+#include "GameplayInsightsStyle.h"
+#include "Widgets/Layout/SScrollBorder.h"
+#include "Widgets/Input/SEditableTextBox.h"
+#include "Widgets/Input/SHyperlink.h"
+#include "Widgets/Input/SCheckBox.h"
+#include "VariantTreeNode.h"
+
+#if WITH_EDITOR
+#include "Subsystems/AssetEditorSubsystem.h"
+#include "Editor.h"
+#endif
+
+#define LOCTEXT_NAMESPACE "SVariantValueView"
+
+namespace VariantColumns
+{
+	static const FName Name("Name");
+	static const FName Value("Value");
+};
+
+static TSharedRef<SWidget> MakeVariantValueWidget(const Trace::IAnalysisSession& InAnalysisSession, const FVariantValue& InValue) 
+{ 
+	switch(InValue.Type)
+	{
+	case EAnimNodeValueType::Bool:
+		return 
+			SNew(SCheckBox)
+			.IsEnabled(false)
+			.IsChecked(InValue.Bool.bValue ? ECheckBoxState::Checked : ECheckBoxState::Unchecked);
+
+	case EAnimNodeValueType::Int32:
+		return
+			SNew(SBox)
+			.WidthOverride(125.0f)
+			[
+				SNew(SEditableTextBox)
+				.IsEnabled(false)
+				.Font(FCoreStyle::Get().GetFontStyle("SmallFont"))
+				.Text(FText::AsNumber(InValue.Int32.Value))
+			];
+
+	case EAnimNodeValueType::Float:
+		return 
+			SNew(SBox)
+			.WidthOverride(125.0f)
+			[
+				SNew(SEditableTextBox)
+				.IsEnabled(false)
+				.Font(FCoreStyle::Get().GetFontStyle("SmallFont"))
+				.Text(FText::AsNumber(InValue.Float.Value))
+			];
+
+	case EAnimNodeValueType::Vector2D:
+		return SNew(SHorizontalBox)
+			+SHorizontalBox::Slot()
+			.AutoWidth()
+			[
+				SNew(SBox)
+				.WidthOverride(125.0f)
+				[
+					SNew(SEditableTextBox)
+					.IsEnabled(false)
+					.Font(FCoreStyle::Get().GetFontStyle("SmallFont"))
+					.Text(FText::AsNumber(InValue.Vector2D.Value.X))
+				]
+			]
+			+SHorizontalBox::Slot()
+			.AutoWidth()
+			[
+				SNew(SBox)
+				.WidthOverride(125.0f)
+				[
+					SNew(SEditableTextBox)
+					.IsEnabled(false)
+					.Font(FCoreStyle::Get().GetFontStyle("SmallFont"))
+					.Text(FText::AsNumber(InValue.Vector2D.Value.Y))
+				]
+			];
+
+	case EAnimNodeValueType::Vector:
+		return SNew(SHorizontalBox)
+			+SHorizontalBox::Slot()
+			.AutoWidth()
+			[
+				SNew(SBox)
+				.WidthOverride(125.0f)
+				[
+					SNew(SEditableTextBox)
+					.IsEnabled(false)
+					.Font(FCoreStyle::Get().GetFontStyle("SmallFont"))
+					.Text(FText::AsNumber(InValue.Vector.Value.X))
+				]
+			]
+			+SHorizontalBox::Slot()
+			.AutoWidth()
+			[
+				SNew(SBox)
+				.WidthOverride(125.0f)
+				[
+					SNew(SEditableTextBox)
+					.IsEnabled(false)
+					.Font(FCoreStyle::Get().GetFontStyle("SmallFont"))
+					.Text(FText::AsNumber(InValue.Vector.Value.Y))
+				]
+			]
+			+SHorizontalBox::Slot()
+			.AutoWidth()
+			[
+				SNew(SBox)
+				.WidthOverride(125.0f)
+				[
+					SNew(SEditableTextBox)
+					.IsEnabled(false)
+					.Font(FCoreStyle::Get().GetFontStyle("SmallFont"))
+					.Text(FText::AsNumber(InValue.Vector.Value.Z))
+				]
+			];
+
+	case EAnimNodeValueType::String:
+		return 
+			SNew(STextBlock)
+			.Font(FCoreStyle::Get().GetFontStyle("SmallFont"))
+			.Text(FText::FromString(InValue.String.Value));
+
+	case EAnimNodeValueType::Object:
+	{
+		const FGameplayProvider* GameplayProvider = InAnalysisSession.ReadProvider<FGameplayProvider>(FGameplayProvider::ProviderName);
+		if(GameplayProvider)
+		{
+			Trace::FAnalysisSessionReadScope SessionReadScope(InAnalysisSession);
+
+			const FObjectInfo& ObjectInfo = GameplayProvider->GetObjectInfo(InValue.Object.Value);
+#if WITH_EDITOR
+			return 
+				SNew(SHyperlink)
+				.Text(FText::FromString(ObjectInfo.Name))
+				.TextStyle(&FCoreStyle::Get().GetWidgetStyle<FTextBlockStyle>("SmallText"))
+				.ToolTipText(FText::Format(LOCTEXT("AssetHyperlinkTooltipFormat", "Open asset '{0}'"), FText::FromString(ObjectInfo.PathName)))
+				.OnNavigate_Lambda([ObjectInfo]()
+				{
+					GEditor->GetEditorSubsystem<UAssetEditorSubsystem>()->OpenEditorForAsset(ObjectInfo.PathName);
+				});
+#else
+			return 
+				SNew(STextBlock)
+				.Font(FCoreStyle::Get().GetFontStyle("SmallFont"))
+				.Text(FText::FromString(ObjectInfo.Name))
+				.ToolTipText(FText::FromString(ObjectInfo.PathName));
+#endif
+		}
+	}
+	break;
+
+	case EAnimNodeValueType::Class:
+	{
+		const FGameplayProvider* GameplayProvider = InAnalysisSession.ReadProvider<FGameplayProvider>(FGameplayProvider::ProviderName);
+		if(GameplayProvider)
+		{
+			Trace::FAnalysisSessionReadScope SessionReadScope(InAnalysisSession);
+
+			const FClassInfo& ClassInfo = GameplayProvider->GetClassInfo(InValue.Class.Value);
+#if WITH_EDITOR
+			return 
+				SNew(SHyperlink)
+				.Text(FText::FromString(ClassInfo.Name))
+				.TextStyle(&FCoreStyle::Get().GetWidgetStyle<FTextBlockStyle>("SmallText"))
+				.ToolTipText(FText::Format(LOCTEXT("ClassHyperlinkTooltipFormat", "Open class '{0}'"), FText::FromString(ClassInfo.PathName)))
+				.OnNavigate_Lambda([ClassInfo]()
+				{
+					GEditor->GetEditorSubsystem<UAssetEditorSubsystem>()->OpenEditorForAsset(ClassInfo.PathName);
+				});
+#else
+			return 
+				SNew(STextBlock)
+				.Font(FCoreStyle::Get().GetFontStyle("SmallFont"))
+				.Text(FText::FromString(ClassInfo.Name))
+				.ToolTipText(FText::FromString(ClassInfo.PathName));
+#endif			
+		}
+	}
+	break;
+	}
+
+	return SNullWidget::NullWidget; 
+}
+
+// Container for an entry in the property view
+class SVariantValueNode : public SMultiColumnTableRow<TSharedRef<FVariantTreeNode>>
+{
+public:
+	void Construct(const FArguments& InArgs, const TSharedRef<STableViewBase>& InOwnerTable, TSharedRef<FVariantTreeNode> InNode, const Trace::IAnalysisSession& InAnalysisSession)
+	{
+		Node = InNode;
+		AnalysisSession = &InAnalysisSession;
+
+		SMultiColumnTableRow<TSharedRef<FVariantTreeNode>>::Construct(
+			FSuperRowType::FArguments()
+			.Padding(1.0f),
+			InOwnerTable
+		);
+	}
+
+	// SMultiColumnTableRow interface
+	virtual TSharedRef<SWidget> GenerateWidgetForColumn(const FName& InColumnName) override
+	{
+		const bool bIsRoot = !Node->GetParent().IsValid();
+
+		if (InColumnName == VariantColumns::Name)
+		{
+			return 
+				SNew(SBorder)
+				.BorderImage(bIsRoot ? FGameplayInsightsStyle::Get().GetBrush("SchematicViewRootLeft") : FCoreStyle::Get().GetBrush("NoBorder"))
+				[
+					SNew(SHorizontalBox)
+					+SHorizontalBox::Slot()
+					.AutoWidth()
+					.Padding(6, 0, 0, 0)
+					.VAlign(VAlign_Center)
+					[
+						SNew(SExpanderArrow, SharedThis(this))
+						.IndentAmount(0)
+					]
+					+SHorizontalBox::Slot()
+					.FillWidth(1.0f)
+					.VAlign(VAlign_Center)
+					[
+						SNew(STextBlock)
+						.Font(FCoreStyle::Get().GetFontStyle(bIsRoot ? "ExpandableArea.TitleFont" : "SmallFont"))
+						.Text(Node->GetName())
+					]
+				];
+		}
+		else if(InColumnName == VariantColumns::Value)
+		{
+			return
+				SNew(SBorder)
+				.BorderImage(bIsRoot ? FGameplayInsightsStyle::Get().GetBrush("SchematicViewRootMid") : FCoreStyle::Get().GetBrush("NoBorder"))
+				[
+					SNew(SHorizontalBox)
+					+SHorizontalBox::Slot()
+					.AutoWidth()
+					.VAlign(VAlign_Center)
+					[
+						MakeVariantValueWidget(*AnalysisSession, Node->GetValue())
+					]
+				];
+		}
+
+		return SNullWidget::NullWidget;
+	}
+
+	const Trace::IAnalysisSession* AnalysisSession;
+	TSharedPtr<FVariantTreeNode> Node;
+};
+
+void SVariantValueView::Construct(const FArguments& InArgs, const Trace::IAnalysisSession& InAnalysisSession)
+{
+	OnGetVariantValues = InArgs._OnGetVariantValues;
+
+	AnalysisSession = &InAnalysisSession;
+	bNeedsRefresh = false;
+
+	VariantTreeView = SNew(STreeView<TSharedRef<FVariantTreeNode>>)
+		.SelectionMode(ESelectionMode::None)
+		.OnGenerateRow(this, &SVariantValueView::HandleGeneratePropertyRow)
+		.OnGetChildren(this, &SVariantValueView::HandleGetPropertyChildren)
+		.TreeItemsSource(&VariantTreeNodes)
+		.HeaderRow(
+			SNew(SHeaderRow)
+			.Visibility(EVisibility::Collapsed)
+			+SHeaderRow::Column(VariantColumns::Name)
+			.DefaultLabel(LOCTEXT("ValueNameColumn", "Name"))
+			+SHeaderRow::Column(VariantColumns::Value)
+			.DefaultLabel(LOCTEXT("ValueValueColumn", "Value"))
+		);
+
+	ChildSlot
+	[
+		SNew(SVerticalBox)
+		+SVerticalBox::Slot()
+		.FillHeight(1.0f)
+		[
+			SNew(SScrollBorder, VariantTreeView.ToSharedRef())
+			[
+				VariantTreeView.ToSharedRef()
+			]
+		]
+	];
+
+	RegisterActiveTimer(0.0f, FWidgetActiveTimerDelegate::CreateLambda([this](double InCurrentTime, float InDeltaTime)
+	{
+		if(bNeedsRefresh)
+		{
+			RefreshNodes();
+			bNeedsRefresh = false;
+		}
+		return EActiveTimerReturnType::Continue;
+	}));
+}
+
+TSharedRef<ITableRow> SVariantValueView::HandleGeneratePropertyRow(TSharedRef<FVariantTreeNode> Item, const TSharedRef<STableViewBase>& OwnerTable)
+{
+	return 
+		SNew(SVariantValueNode, OwnerTable, Item, *AnalysisSession);
+}
+
+void SVariantValueView::HandleGetPropertyChildren(TSharedRef<FVariantTreeNode> InItem, TArray<TSharedRef<FVariantTreeNode>>& OutChildren)
+{
+	for(const TSharedRef<FVariantTreeNode>& ChildNode : InItem->GetChildren())
+	{
+		VariantTreeView->SetItemExpansion(ChildNode, true);
+	}
+
+	OutChildren.Append(InItem->GetChildren());
+}
+
+void SVariantValueView::RefreshNodes()
+{
+	VariantTreeNodes.Reset();
+
+	OnGetVariantValues.ExecuteIfBound(TimeMarker, VariantTreeNodes);
+
+	if(VariantTreeNodes.Num() > 0)
+	{
+		for(const TSharedRef<FVariantTreeNode>& VariantTreeNode : VariantTreeNodes)
+		{
+			VariantTreeView->SetItemExpansion(VariantTreeNode, true);
+		}
+	}
+
+	VariantTreeView->RequestTreeRefresh();
+}
+
+#undef LOCTEXT_NAMESPACE
diff --git a/Engine/Plugins/Animation/GameplayInsights/Source/GameplayInsights/Private/SVariantValueView.h b/Engine/Plugins/Animation/GameplayInsights/Source/GameplayInsights/Private/SVariantValueView.h
new file mode 100644
--- /dev/null
+++ b/Engine/Plugins/Animation/GameplayInsights/Source/GameplayInsights/Private/SVariantValueView.h
@@ -0,0 +1,51 @@
+// Copyright Epic Games, Inc. All Rights Reserved.
+
+#pragma once
+
+#include "CoreMinimal.h"
+#include "Widgets/SCompoundWidget.h"
+#include "Widgets/Views/STreeView.h"
+
+namespace Trace { class IAnalysisSession; }
+struct FVariantTreeNode;
+
+// Delegate called to get variant values to display
+DECLARE_DELEGATE_TwoParams(FOnGetVariantValues, double /*InTime*/, TArray<TSharedRef<FVariantTreeNode>>& /*OutValues*/);
+
+class SVariantValueView : public SCompoundWidget
+{
+public:
+	SLATE_BEGIN_ARGS(SVariantValueView) {}
+
+	SLATE_EVENT(FOnGetVariantValues, OnGetVariantValues)
+
+	SLATE_END_ARGS()
+
+	void Construct(const FArguments& InArgs, const Trace::IAnalysisSession& InAnalysisSession);
+
+	/** Refresh the displayed variants. */
+	void RequestRefresh(double InTimeMarker) { TimeMarker = InTimeMarker; bNeedsRefresh = true; }
+
+private:
+	// Generate a row widget for a property item
+	TSharedRef<ITableRow> HandleGeneratePropertyRow(TSharedRef<FVariantTreeNode> Item, const TSharedRef<STableViewBase>& OwnerTable);
+
+	// Get the children of a property item
+	void HandleGetPropertyChildren(TSharedRef<FVariantTreeNode> InItem, TArray<TSharedRef<FVariantTreeNode>>& OutChildren);
+
+	// Refresh the nodes
+	void RefreshNodes();
+
+private:
+	const Trace::IAnalysisSession* AnalysisSession;
+
+	TSharedPtr<STreeView<TSharedRef<FVariantTreeNode>>> VariantTreeView;
+
+	TArray<TSharedRef<FVariantTreeNode>> VariantTreeNodes;
+
+	double TimeMarker;
+
+	FOnGetVariantValues OnGetVariantValues;
+
+	bool bNeedsRefresh;
+};
diff --git a/Engine/Plugins/Animation/GameplayInsights/Source/GameplayInsights/Private/SkeletalMeshCurvesTrack.cpp b/Engine/Plugins/Animation/GameplayInsights/Source/GameplayInsights/Private/SkeletalMeshCurvesTrack.cpp
--- a/Engine/Plugins/Animation/GameplayInsights/Source/GameplayInsights/Private/SkeletalMeshCurvesTrack.cpp
+++ b/Engine/Plugins/Animation/GameplayInsights/Source/GameplayInsights/Private/SkeletalMeshCurvesTrack.cpp
@@ -11,13 +11,15 @@
 #include "Insights/ViewModels/TooltipDrawState.h"
 #include "Modules/ModuleManager.h"
 #include "Insights/ViewModels/GraphTrackBuilder.h"
+#include "VariantTreeNode.h"
+#include "TraceServices/Model/Frames.h"
 
 #define LOCTEXT_NAMESPACE "SkeletalMeshCurvesTrack"
 
 INSIGHTS_IMPLEMENT_RTTI(FSkeletalMeshCurvesTrack)
 
 FSkeletalMeshCurvesTrack::FSkeletalMeshCurvesTrack(const FAnimationSharedData& InSharedData, uint64 InObjectID, const TCHAR* InName)
-	: FGameplayGraphTrack(InObjectID, FText::Format(LOCTEXT("TrackNameFormat", "Curves - {0}"), FText::FromString(FString(InName))))
+	: FGameplayGraphTrack(InSharedData.GetGameplaySharedData(), InObjectID, FText::Format(LOCTEXT("TrackNameFormat", "Curves - {0}"), FText::FromString(FString(InName))))
 	, SharedData(InSharedData)
 {
 }
@@ -26,40 +28,43 @@ void FSkeletalMeshCurvesTrack::AddAllSeries()
 {
 	const FAnimationProvider* AnimationProvider = SharedData.GetAnalysisSession().ReadProvider<FAnimationProvider>(FAnimationProvider::ProviderName);
 
-	bool bFirstSeries = true;
+	bool bFirstSeries = AllSeries.Num() == 0;
 	if(AnimationProvider)
 	{
 		Trace::FAnalysisSessionReadScope SessionReadScope(SharedData.GetAnalysisSession());
 
 		AnimationProvider->EnumerateSkeletalMeshCurveIds(GetGameplayTrack().GetObjectId(), [this, &bFirstSeries, &AnimationProvider](uint32 InCurveId)
 		{
-			auto MakeCurveSeriesColor = [](uint32 InSeed, bool bInLine)
+			if(!AllSeries.ContainsByPredicate([InCurveId](const TSharedPtr<FGraphSeries>& InSeries){ return StaticCastSharedPtr<FSkeletalMeshCurveSeries>(InSeries)->CurveId == InCurveId; }))
 			{
-				FRandomStream Stream(InSeed);
-				const uint8 Hue = (uint8)(Stream.FRand() * 255.0f);
-				const uint8 SatVal = bInLine ? 196 : 128;
-				return FLinearColor::MakeFromHSV8(Hue, SatVal, SatVal);
-			};
+				auto MakeCurveSeriesColor = [](uint32 InSeed, bool bInLine)
+				{
+					FRandomStream Stream(InSeed);
+					const uint8 Hue = (uint8)(Stream.FRand() * 255.0f);
+					const uint8 SatVal = bInLine ? 196 : 128;
+					return FLinearColor::MakeFromHSV8(Hue, SatVal, SatVal);
+				};
 
-			TSharedRef<FSkeletalMeshCurveSeries> Series = MakeShared<FSkeletalMeshCurveSeries>();
+				TSharedRef<FSkeletalMeshCurveSeries> Series = MakeShared<FSkeletalMeshCurveSeries>();
 
-			const TCHAR* CurveName = AnimationProvider->GetName(InCurveId);
+				const TCHAR* CurveName = AnimationProvider->GetName(InCurveId);
 
-			Series->SetName(CurveName);
-			Series->SetDescription(FText::Format(LOCTEXT("CurveNameFormat", "Values for curve '{0}'"), FText::FromString(CurveName)));
+				Series->SetName(CurveName);
+				Series->SetDescription(FText::Format(LOCTEXT("CurveNameFormat", "Values for curve '{0}'"), FText::FromString(CurveName)));
 
-			const FLinearColor LineColor = MakeCurveSeriesColor(InCurveId, true);
-			const FLinearColor FillColor = MakeCurveSeriesColor(InCurveId, false);
-			Series->SetColor(LineColor, LineColor, FillColor);
+				const FLinearColor LineColor = MakeCurveSeriesColor(InCurveId, true);
+				const FLinearColor FillColor = MakeCurveSeriesColor(InCurveId, false);
+				Series->SetColor(LineColor, LineColor, FillColor);
 
-			Series->CurveId = InCurveId;
-			Series->SetVisibility(bFirstSeries);
-			Series->SetBaselineY(25.0f);
-			Series->SetScaleY(20.0f);
-			Series->EnableAutoZoom();
-			AllSeries.Add(Series);	
+				Series->CurveId = InCurveId;
+				Series->SetVisibility(bFirstSeries);
+				Series->SetBaselineY(25.0f);
+				Series->SetScaleY(20.0f);
+				Series->EnableAutoZoom();
+				AllSeries.Add(Series);	
 
-			bFirstSeries = false;
+				bFirstSeries = false;
+			}
 		});
 	}
 }
@@ -206,4 +211,38 @@ void FSkeletalMeshCurvesTrack::FindSkeletalMeshPoseMessage(const FTimingEventSea
 		});
 }
 
+void FSkeletalMeshCurvesTrack::GetVariantsAtTime(double InTime, TArray<TSharedRef<FVariantTreeNode>>& OutVariants) const
+{
+	const FGameplayProvider* GameplayProvider = SharedData.GetAnalysisSession().ReadProvider<FGameplayProvider>(FGameplayProvider::ProviderName);
+	const FAnimationProvider* AnimationProvider = SharedData.GetAnalysisSession().ReadProvider<FAnimationProvider>(FAnimationProvider::ProviderName);
+	if(GameplayProvider && AnimationProvider)
+	{
+		Trace::FAnalysisSessionReadScope SessionReadScope(SharedData.GetAnalysisSession());
+
+		TSharedRef<FVariantTreeNode> Header = OutVariants.Add_GetRef(FVariantTreeNode::MakeHeader(FText::FromString(GetName())));
+
+		const Trace::IFrameProvider& FramesProvider = Trace::ReadFrameProvider(SharedData.GetAnalysisSession());
+
+		AnimationProvider->ReadSkeletalMeshPoseTimeline(GetGameplayTrack().GetObjectId(), [&Header, &FramesProvider, &AnimationProvider, &InTime](const FAnimationProvider::SkeletalMeshPoseTimeline& InTimeline, bool bInHasCurves)
+		{
+			// round to nearest frame boundary
+			Trace::FFrame Frame;
+			if(FramesProvider.GetFrameFromTime(ETraceFrameType::TraceFrameType_Game, InTime, Frame))
+			{
+				InTimeline.EnumerateEvents(Frame.StartTime, Frame.EndTime, [&Header, &AnimationProvider, &Frame](double InStartTime, double InEndTime, uint32 InDepth, const FSkeletalMeshPoseMessage& InMessage)
+				{
+					if(Frame.StartTime <= InStartTime && Frame.EndTime >= InEndTime)
+					{
+						AnimationProvider->EnumerateSkeletalMeshCurves(InMessage, [&Header, &AnimationProvider](const FSkeletalMeshNamedCurve& InCurve)
+						{
+							const TCHAR* CurveName = AnimationProvider->GetName(InCurve.Id);
+							Header->AddChild(FVariantTreeNode::MakeFloat(FText::FromString(CurveName), InCurve.Value));
+						});
+					}
+				});
+			}
+		});
+	}
+}
+
 #undef LOCTEXT_NAMESPACE
diff --git a/Engine/Plugins/Animation/GameplayInsights/Source/GameplayInsights/Private/SkeletalMeshCurvesTrack.h b/Engine/Plugins/Animation/GameplayInsights/Source/GameplayInsights/Private/SkeletalMeshCurvesTrack.h
--- a/Engine/Plugins/Animation/GameplayInsights/Source/GameplayInsights/Private/SkeletalMeshCurvesTrack.h
+++ b/Engine/Plugins/Animation/GameplayInsights/Source/GameplayInsights/Private/SkeletalMeshCurvesTrack.h
@@ -27,6 +27,7 @@ public:
 	virtual void AddAllSeries() override;
 	virtual bool UpdateSeriesBounds(FGameplayGraphSeries& InSeries, const FTimingTrackViewport& InViewport) override;
 	virtual void UpdateSeries(FGameplayGraphSeries& InSeries, const FTimingTrackViewport& InViewport) override;
+	virtual void GetVariantsAtTime(double InTime, TArray<TSharedRef<FVariantTreeNode>>& OutVariants) const override;
 
 private:
 	// Helper function used to find a skeletal mesh pose
diff --git a/Engine/Plugins/Animation/GameplayInsights/Source/GameplayInsights/Private/SkeletalMeshPoseTrack.cpp b/Engine/Plugins/Animation/GameplayInsights/Source/GameplayInsights/Private/SkeletalMeshPoseTrack.cpp
--- a/Engine/Plugins/Animation/GameplayInsights/Source/GameplayInsights/Private/SkeletalMeshPoseTrack.cpp
+++ b/Engine/Plugins/Animation/GameplayInsights/Source/GameplayInsights/Private/SkeletalMeshPoseTrack.cpp
@@ -24,7 +24,7 @@
 INSIGHTS_IMPLEMENT_RTTI(FSkeletalMeshPoseTrack)
 
 FSkeletalMeshPoseTrack::FSkeletalMeshPoseTrack(const FAnimationSharedData& InSharedData, uint64 InObjectID, const TCHAR* InName)
-	: TGameplayTrackMixin<FTimingEventsTrack>(InObjectID, FText::Format(LOCTEXT("TrackNameFormat", "Pose - {0}"), FText::FromString(FString(InName))))
+	: FGameplayTimingEventsTrack(InSharedData.GetGameplaySharedData(), InObjectID, FText::Format(LOCTEXT("TrackNameFormat", "Pose - {0}"), FText::FromString(FString(InName))))
 	, SharedData(InSharedData)
 	, Color(FLinearColor::MakeRandomColor())
 	, bDrawPose(false)
diff --git a/Engine/Plugins/Animation/GameplayInsights/Source/GameplayInsights/Private/SkeletalMeshPoseTrack.h b/Engine/Plugins/Animation/GameplayInsights/Source/GameplayInsights/Private/SkeletalMeshPoseTrack.h
--- a/Engine/Plugins/Animation/GameplayInsights/Source/GameplayInsights/Private/SkeletalMeshPoseTrack.h
+++ b/Engine/Plugins/Animation/GameplayInsights/Source/GameplayInsights/Private/SkeletalMeshPoseTrack.h
@@ -20,15 +20,14 @@ class FAnimationSharedData;
 class FTimingEventSearchParameters;
 struct FSkeletalMeshPoseMessage;
 
-class FSkeletalMeshPoseTrack : public TGameplayTrackMixin<FTimingEventsTrack>
+class FSkeletalMeshPoseTrack : public FGameplayTimingEventsTrack
 #if WITH_ENGINE
 	, public FGCObject
 #endif
 {
-public:
-	static const FName TypeName;
-	static const FName SubTypeName;
+	INSIGHTS_DECLARE_RTTI(FSkeletalMeshPoseTrack, FGameplayTimingEventsTrack)
 
+public:
 	FSkeletalMeshPoseTrack(const FAnimationSharedData& InSharedData, uint64 InObjectID, const TCHAR* InName);
 	~FSkeletalMeshPoseTrack();
 
diff --git a/Engine/Plugins/Animation/GameplayInsights/Source/GameplayInsights/Private/VariantTreeNode.h b/Engine/Plugins/Animation/GameplayInsights/Source/GameplayInsights/Private/VariantTreeNode.h
new file mode 100644
--- /dev/null
+++ b/Engine/Plugins/Animation/GameplayInsights/Source/GameplayInsights/Private/VariantTreeNode.h
@@ -0,0 +1,118 @@
+// Copyright Epic Games, Inc. All Rights Reserved.
+
+#pragma once
+
+#include "CoreMinimal.h"
+#include "IAnimationProvider.h"
+
+/** A wrapper around a variant value, used to display collections of values in a tree */
+struct FVariantTreeNode : TSharedFromThis<FVariantTreeNode>
+{
+	FVariantTreeNode(const FText& InName, const FVariantValue& InValue)
+		: Name(InName)
+		, Value(InValue)
+	{}
+
+	const TSharedRef<FVariantTreeNode>& AddChild(const TSharedRef<FVariantTreeNode>& InChild)
+	{
+		check(!InChild->Parent.IsValid());
+		InChild->Parent = SharedThis(this);
+		return Children.Add_GetRef(InChild);
+	}
+
+	FText GetName() const { return Name; }
+
+	const FVariantValue& GetValue() const { return Value; }
+
+	TSharedPtr<FVariantTreeNode> GetParent() const { return Parent.Pin(); }
+
+	const TArray<TSharedRef<FVariantTreeNode>>& GetChildren() const { return Children; }
+
+	static TSharedRef<FVariantTreeNode> MakeHeader(const FText& InName)
+	{
+		FVariantValue Value;
+		Value.Type = EAnimNodeValueType::String;
+		Value.String.Value = TEXT("");
+
+		return MakeShared<FVariantTreeNode>(InName, Value);
+	}
+
+	static TSharedRef<FVariantTreeNode> MakeBool(const FText& InName, bool InValue)
+	{
+		FVariantValue Value;
+		Value.Type = EAnimNodeValueType::Bool;
+		Value.Bool.bValue = InValue;
+
+		return MakeShared<FVariantTreeNode>(InName, Value);
+	}
+
+	static TSharedRef<FVariantTreeNode> MakeInt32(const FText& InName, int32 InValue)
+	{
+		FVariantValue Value;
+		Value.Type = EAnimNodeValueType::Int32;
+		Value.Int32.Value = InValue;
+
+		return MakeShared<FVariantTreeNode>(InName, Value);
+	}
+
+	static TSharedRef<FVariantTreeNode> MakeFloat(const FText& InName, float InValue)
+	{
+		FVariantValue Value;
+		Value.Type = EAnimNodeValueType::Float;
+		Value.Float.Value = InValue;
+
+		return MakeShared<FVariantTreeNode>(InName, Value);
+	}
+
+	static TSharedRef<FVariantTreeNode> MakeVector2D(const FText& InName, const FVector2D& InValue)
+	{
+		FVariantValue Value;
+		Value.Type = EAnimNodeValueType::Vector2D;
+		Value.Vector2D.Value = InValue;
+
+		return MakeShared<FVariantTreeNode>(InName, Value);
+	}
+
+
+	static TSharedRef<FVariantTreeNode> MakeVector(const FText& InName, const FVector& InValue)
+	{
+		FVariantValue Value;
+		Value.Type = EAnimNodeValueType::Vector;
+		Value.Vector.Value = InValue;
+
+		return MakeShared<FVariantTreeNode>(InName, Value);
+	}
+
+	static TSharedRef<FVariantTreeNode> MakeString(const FText& InName, const TCHAR* InValue)
+	{
+		FVariantValue Value;
+		Value.Type = EAnimNodeValueType::String;
+		Value.String.Value = InValue;
+
+		return MakeShared<FVariantTreeNode>(InName, Value);
+	}
+
+	static TSharedRef<FVariantTreeNode> MakeObject(const FText& InName, uint64 InValue)
+	{
+		FVariantValue Value;
+		Value.Type = EAnimNodeValueType::Object;
+		Value.Object.Value = InValue;
+
+		return MakeShared<FVariantTreeNode>(InName, Value);
+	}
+
+	static TSharedRef<FVariantTreeNode> MakeClass(const FText& InName, uint64 InValue)
+	{
+		FVariantValue Value;
+		Value.Type = EAnimNodeValueType::Class;
+		Value.Class.Value = InValue;
+
+		return MakeShared<FVariantTreeNode>(InName, Value);
+	}
+
+private:
+	FText Name;
+	FVariantValue Value;
+	TWeakPtr<FVariantTreeNode> Parent;
+	TArray<TSharedRef<FVariantTreeNode>> Children;
+};
\ No newline at end of file
diff --git a/Engine/Source/Runtime/Engine/Private/Animation/AnimInstance.cpp b/Engine/Source/Runtime/Engine/Private/Animation/AnimInstance.cpp
--- a/Engine/Source/Runtime/Engine/Private/Animation/AnimInstance.cpp
+++ b/Engine/Source/Runtime/Engine/Private/Animation/AnimInstance.cpp
@@ -262,6 +262,7 @@ void UAnimInstance::UninitializeAnimation()
 			const FAnimNotifyEvent& AnimNotifyEvent = ActiveAnimNotifyState[Index];
 			if (ShouldTriggerAnimNotifyState(AnimNotifyEvent.NotifyStateClass))
 			{
+				TRACE_ANIM_NOTIFY(this, AnimNotifyEvent, End);
 				AnimNotifyEvent.NotifyStateClass->NotifyEnd(SkelMeshComp, Cast<UAnimSequenceBase>(AnimNotifyEvent.NotifyStateClass->GetOuter()));
 			}
 		}
@@ -306,6 +307,13 @@ void UAnimInstance::UpdateMontage(float DeltaSeconds)
 	// if we do multi threading, make sure this stays in game thread. 
 	// This is because branch points need to execute arbitrary code inside this call.
 	Montage_Advance(DeltaSeconds);
+
+#if ANIM_TRACE_ENABLED
+	for (FAnimMontageInstance* MontageInstance : MontageInstances)
+	{
+		TRACE_ANIM_MONTAGE(this, *MontageInstance);
+	}
+#endif
 }
 
 void UAnimInstance::UpdateMontageSyncGroup()
@@ -1321,6 +1329,7 @@ void UAnimInstance::TriggerAnimNotifies(float DeltaSeconds)
 		const FAnimNotifyEvent& AnimNotifyEvent = ActiveAnimNotifyState[Index];
 		if (AnimNotifyEvent.NotifyStateClass && ShouldTriggerAnimNotifyState(AnimNotifyEvent.NotifyStateClass))
 		{
+			TRACE_ANIM_NOTIFY(this, AnimNotifyEvent, End);
 			AnimNotifyEvent.NotifyStateClass->NotifyEnd(SkelMeshComp, Cast<UAnimSequenceBase>(AnimNotifyEvent.NotifyStateClass->GetOuter()));
 		}
 		// The NotifyEnd callback above may have triggered actor destruction and the tear down
@@ -1338,6 +1347,7 @@ void UAnimInstance::TriggerAnimNotifies(float DeltaSeconds)
 	{
 		if (ShouldTriggerAnimNotifyState(AnimNotifyEvent->NotifyStateClass))
 		{
+			TRACE_ANIM_NOTIFY(this, *AnimNotifyEvent, Begin);
 			AnimNotifyEvent->NotifyStateClass->NotifyBegin(SkelMeshComp, Cast<UAnimSequenceBase>(AnimNotifyEvent->NotifyStateClass->GetOuter()), AnimNotifyEvent->GetDuration());
 		}
 	}
@@ -1350,6 +1360,7 @@ void UAnimInstance::TriggerAnimNotifies(float DeltaSeconds)
 	{
 		if (ShouldTriggerAnimNotifyState(AnimNotifyEvent.NotifyStateClass))
 		{
+			TRACE_ANIM_NOTIFY(this, AnimNotifyEvent, Tick);
 			AnimNotifyEvent.NotifyStateClass->NotifyTick(SkelMeshComp, Cast<UAnimSequenceBase>(AnimNotifyEvent.NotifyStateClass->GetOuter()), DeltaSeconds);
 		}
 	}
@@ -1368,6 +1379,7 @@ void UAnimInstance::TriggerSingleAnimNotify(const FAnimNotifyEvent* AnimNotifyEv
 		if (AnimNotifyEvent->Notify != nullptr)
 		{	
 			// Implemented notify: just call Notify. UAnimNotify will forward this to the event which will do the work.
+			TRACE_ANIM_NOTIFY(this, *AnimNotifyEvent, Event);
 			AnimNotifyEvent->Notify->Notify(GetSkelMeshComponent(), Cast<UAnimSequenceBase>(AnimNotifyEvent->Notify->GetOuter()));
 		}
 		else if (AnimNotifyEvent->NotifyName != NAME_None)
@@ -1380,6 +1392,7 @@ void UAnimInstance::TriggerSingleAnimNotify(const FAnimNotifyEvent* AnimNotifyEv
 				// if parameter is none, add event
 				if (Function->NumParms == 0)
 				{
+					TRACE_ANIM_NOTIFY(this, *AnimNotifyEvent, Event);
 					ProcessEvent(Function, NULL);
 				}
 				else if ((Function->NumParms == 1) && (CastField<FObjectProperty>(Function->PropertyLink) != NULL))
@@ -1391,6 +1404,8 @@ void UAnimInstance::TriggerSingleAnimNotify(const FAnimNotifyEvent* AnimNotifyEv
 
 					FAnimNotifierHandler_Parms Parms;
 					Parms.Notify = AnimNotifyEvent->Notify;
+
+					TRACE_ANIM_NOTIFY(this, *AnimNotifyEvent, Event);
 					ProcessEvent(Function, &Parms);
 				}
 				else
@@ -1411,6 +1426,7 @@ void UAnimInstance::EndNotifyStates()
 	{
 		if (UAnimNotifyState* NotifyState = Event.NotifyStateClass)
 		{
+			TRACE_ANIM_NOTIFY(this, Event, End);
 			NotifyState->NotifyEnd(SkelMeshComp, Cast<UAnimSequenceBase>(NotifyState->GetOuter()));
 		}
 	}
@@ -1621,6 +1637,7 @@ void UAnimInstance::TriggerMontageEndedEvent(const FQueuedMontageEndedEvent& Mon
 			{
 				if (ShouldTriggerAnimNotifyState(AnimNotifyEvent.NotifyStateClass))
 				{
+					TRACE_ANIM_NOTIFY(this, AnimNotifyEvent, End);
 					AnimNotifyEvent.NotifyStateClass->NotifyEnd(SkelMeshComp, NotifyMontage);
 				}
 
diff --git a/Engine/Source/Runtime/Engine/Private/Animation/AnimInstanceProxy.cpp b/Engine/Source/Runtime/Engine/Private/Animation/AnimInstanceProxy.cpp
--- a/Engine/Source/Runtime/Engine/Private/Animation/AnimInstanceProxy.cpp
+++ b/Engine/Source/Runtime/Engine/Private/Animation/AnimInstanceProxy.cpp
@@ -826,6 +826,13 @@ void FAnimInstanceProxy::TickAssetPlayerInstances(float DeltaSeconds)
 					}
 				}
 			}
+
+#if ANIM_TRACE_ENABLED
+			for(const FPassedMarker& PassedMarker : TickContext.MarkerTickContext.MarkersPassedThisTick)
+			{
+				TRACE_ANIM_SYNC_MARKER(CastChecked<UAnimInstance>(GetAnimInstanceObject()), PassedMarker);
+			}
+#endif
 		}
 	}
 
@@ -847,6 +854,13 @@ void FAnimInstanceProxy::TickAssetPlayerInstances(float DeltaSeconds)
 		{
 			ExtractedRootMotion.AccumulateWithBlend(TickContext.RootMotionMovementParams, AssetPlayerToTick.GetRootMotionWeight());
 		}
+
+#if ANIM_TRACE_ENABLED
+		for(const FPassedMarker& PassedMarker : TickContext.MarkerTickContext.MarkersPassedThisTick)
+		{
+			TRACE_ANIM_SYNC_MARKER(CastChecked<UAnimInstance>(GetAnimInstanceObject()), PassedMarker);
+		}
+#endif
 	}
 }
 
diff --git a/Engine/Source/Runtime/Engine/Private/Animation/AnimMontage.cpp b/Engine/Source/Runtime/Engine/Private/Animation/AnimMontage.cpp
--- a/Engine/Source/Runtime/Engine/Private/Animation/AnimMontage.cpp
+++ b/Engine/Source/Runtime/Engine/Private/Animation/AnimMontage.cpp
@@ -16,6 +16,7 @@
 #include "Animation/AnimNotifies/AnimNotifyState.h"
 #include "Animation/AnimSingleNodeInstance.h"
 #include "Engine/Engine.h"
+#include "Animation/AnimTrace.h"
 
 DEFINE_LOG_CATEGORY(LogAnimMontage);
 
@@ -1476,6 +1477,7 @@ void FAnimMontageInstance::Terminate()
 		{
 			FAnimNotifyEvent& NotifyEvent = ActiveStateBranchingPoints[Index];
 			FBranchingPointNotifyPayload BranchingPointNotifyPayload(AnimInstance->GetSkelMeshComponent(), Montage, &NotifyEvent, InstanceID);
+			TRACE_ANIM_NOTIFY(Inst, NotifyEvent, End);
 			NotifyEvent.NotifyStateClass->BranchingPointNotifyEnd(BranchingPointNotifyPayload);
 		}
 		ActiveStateBranchingPoints.Empty();
@@ -2329,6 +2331,13 @@ void FAnimMontageInstance::Advance(float DeltaTime, struct FRootMotionMovementPa
 		}
 	}
 
+#if ANIM_TRACE_ENABLED
+	for(const FPassedMarker& PassedMarker : MarkersPassedThisTick)
+	{
+		TRACE_ANIM_SYNC_MARKER(AnimInstance.Get(), PassedMarker);
+	}
+#endif
+
 	// If this Montage has no weight, it should be terminated.
 	if (IsStopped() && (Blend.IsComplete()))
 	{
@@ -2415,6 +2424,7 @@ void FAnimMontageInstance::UpdateActiveStateBranchingPoints(float CurrentTrackPo
 			if (!bNotifyIsActive)
 			{
 				FBranchingPointNotifyPayload BranchingPointNotifyPayload(AnimInstance->GetSkelMeshComponent(), Montage, &NotifyEvent, InstanceID);
+				TRACE_ANIM_NOTIFY(AnimInstance.Get(), NotifyEvent, End);
 				NotifyEvent.NotifyStateClass->BranchingPointNotifyEnd(BranchingPointNotifyPayload);
 				ActiveStateBranchingPoints.RemoveAt(Index, 1);
 			}
@@ -2433,6 +2443,7 @@ void FAnimMontageInstance::UpdateActiveStateBranchingPoints(float CurrentTrackPo
 			if (bNotifyIsActive && !ActiveStateBranchingPoints.Contains(NotifyEvent))
 			{
 				FBranchingPointNotifyPayload BranchingPointNotifyPayload(AnimInstance->GetSkelMeshComponent(), Montage, &NotifyEvent, InstanceID);
+				TRACE_ANIM_NOTIFY(AnimInstance.Get(), NotifyEvent, Begin);
 				NotifyEvent.NotifyStateClass->BranchingPointNotifyBegin(BranchingPointNotifyPayload);
 				ActiveStateBranchingPoints.Add(NotifyEvent);
 			}
@@ -2469,12 +2480,14 @@ void FAnimMontageInstance::BranchingPointEventHandler(const FBranchingPointMarke
 				if (BranchingPointMarker->NotifyEventType == EAnimNotifyEventType::Begin)
 				{
 					FBranchingPointNotifyPayload BranchingPointNotifyPayload(AnimInstance->GetSkelMeshComponent(), Montage, NotifyEvent, InstanceID);
+					TRACE_ANIM_NOTIFY(AnimInstance.Get(), *NotifyEvent, Begin);
 					NotifyEvent->NotifyStateClass->BranchingPointNotifyBegin(BranchingPointNotifyPayload);
 					ActiveStateBranchingPoints.Add(*NotifyEvent);
 				}
 				else
 				{
 					FBranchingPointNotifyPayload BranchingPointNotifyPayload(AnimInstance->GetSkelMeshComponent(), Montage, NotifyEvent, InstanceID);
+					TRACE_ANIM_NOTIFY(AnimInstance.Get(), *NotifyEvent, End);
 					NotifyEvent->NotifyStateClass->BranchingPointNotifyEnd(BranchingPointNotifyPayload);
 					ActiveStateBranchingPoints.RemoveSingleSwap(*NotifyEvent);
 				}
@@ -2484,6 +2497,7 @@ void FAnimMontageInstance::BranchingPointEventHandler(const FBranchingPointMarke
 			{
 				// Implemented notify: just call Notify. UAnimNotify will forward this to the event which will do the work.
 				FBranchingPointNotifyPayload BranchingPointNotifyPayload(AnimInstance->GetSkelMeshComponent(), Montage, NotifyEvent, InstanceID);
+				TRACE_ANIM_NOTIFY(AnimInstance.Get(), *NotifyEvent, Event);
 				NotifyEvent->Notify->BranchingPointNotify(BranchingPointNotifyPayload);
 			}
 			// Try to match a notify function by name.
diff --git a/Engine/Source/Runtime/Engine/Private/Animation/AnimTrace.cpp b/Engine/Source/Runtime/Engine/Private/Animation/AnimTrace.cpp
--- a/Engine/Source/Runtime/Engine/Private/Animation/AnimTrace.cpp
+++ b/Engine/Source/Runtime/Engine/Private/Animation/AnimTrace.cpp
@@ -16,6 +16,9 @@
 #include "Animation/AnimMontage.h"
 #include "Animation/BlendSpaceBase.h"
 #include "Animation/AnimNode_SequencePlayer.h"
+#include "Animation/AnimNotifies/AnimNotify.h"
+#include "Animation/AnimNotifies/AnimNotifyState.h"
+#include "Animation/AnimTypes.h"
 
 UE_TRACE_CHANNEL_DEFINE(AnimationChannel);
 
@@ -109,6 +112,16 @@ UE_TRACE_EVENT_BEGIN(Animation, AnimNodeValueFloat)
 	UE_TRACE_EVENT_FIELD(uint16, FrameCounter)
 UE_TRACE_EVENT_END()
 
+UE_TRACE_EVENT_BEGIN(Animation, AnimNodeValueVector2D)
+	UE_TRACE_EVENT_FIELD(uint64, Cycle)
+	UE_TRACE_EVENT_FIELD(uint64, AnimInstanceId)
+	UE_TRACE_EVENT_FIELD(int32, NodeId)
+	UE_TRACE_EVENT_FIELD(int32, KeyLength)
+	UE_TRACE_EVENT_FIELD(float, ValueX)
+	UE_TRACE_EVENT_FIELD(float, ValueY)
+	UE_TRACE_EVENT_FIELD(uint16, FrameCounter)
+UE_TRACE_EVENT_END()
+
 UE_TRACE_EVENT_BEGIN(Animation, AnimNodeValueVector)
 	UE_TRACE_EVENT_FIELD(uint64, Cycle)
 	UE_TRACE_EVENT_FIELD(uint64, AnimInstanceId)
@@ -179,12 +192,45 @@ UE_TRACE_EVENT_BEGIN(Animation, Name, Important)
 	UE_TRACE_EVENT_FIELD(uint32, Id)
 UE_TRACE_EVENT_END()
 
+UE_TRACE_EVENT_BEGIN(Animation, Notify)
+	UE_TRACE_EVENT_FIELD(uint64, Cycle)
+	UE_TRACE_EVENT_FIELD(uint64, AnimInstanceId)
+	UE_TRACE_EVENT_FIELD(uint64, AssetId)
+	UE_TRACE_EVENT_FIELD(uint64, NotifyId)
+	UE_TRACE_EVENT_FIELD(uint32, NameId)
+	UE_TRACE_EVENT_FIELD(float, Time)
+	UE_TRACE_EVENT_FIELD(float, Duration)
+	UE_TRACE_EVENT_FIELD(uint8, NotifyEventType)
+UE_TRACE_EVENT_END()
+
+UE_TRACE_EVENT_BEGIN(Animation, SyncMarker)
+	UE_TRACE_EVENT_FIELD(uint64, Cycle)
+	UE_TRACE_EVENT_FIELD(uint64, AnimInstanceId)
+	UE_TRACE_EVENT_FIELD(uint32, NameId)
+UE_TRACE_EVENT_END()
+
+UE_TRACE_EVENT_BEGIN(Animation, Montage)
+	UE_TRACE_EVENT_FIELD(uint64, Cycle)
+	UE_TRACE_EVENT_FIELD(uint64, AnimInstanceId)
+	UE_TRACE_EVENT_FIELD(uint64, MontageId)
+	UE_TRACE_EVENT_FIELD(uint32, CurrentSectionNameId)
+	UE_TRACE_EVENT_FIELD(uint32, NextSectionNameId)
+	UE_TRACE_EVENT_FIELD(float, Weight)
+	UE_TRACE_EVENT_FIELD(float, DesiredWeight)
+UE_TRACE_EVENT_END()
+
 // Object annotations used for tracing
 FUObjectAnnotationSparseBool GSkeletalMeshTraceAnnotations;
 
 // Map used for unique name output
 TMap<FName, uint32> GAnimTraceNames;
 
+// Global unique name index
+uint32 GAnimTraceCurrentNameId = 1;
+
+// Critical section used to lock global name map & index
+FCriticalSection GAnimTraceNameCriticalSection;
+
 // Scratch buffers for various traces to avoid allocation churn.
 // These can be removed when lambda support is added for array fields to remove a memcpy.
 struct FAnimTraceScratchBuffers : public TThreadSingleton<FAnimTraceScratchBuffers>
@@ -388,11 +434,26 @@ void FAnimTrace::OutputSkeletalMesh(const USkeletalMesh* InMesh)
 
 uint32 FAnimTrace::OutputName(const FName& InName)
 {
-	static uint32 CurrentId = 1;
-	check(IsInGameThread());
+	uint32 NameId = 0;
+	bool bShouldTrace = false;
+	if(InName != NAME_None)
+	{
+		FScopeLock ScopeLock(&GAnimTraceNameCriticalSection);
+
+		uint32* ExistingIdPtr = GAnimTraceNames.Find(InName);
+		if(ExistingIdPtr == nullptr)
+		{
+			NameId = GAnimTraceCurrentNameId++;
+			GAnimTraceNames.Add(InName, NameId);
+			bShouldTrace = true;
+		}
+		else
+		{
+			NameId = *ExistingIdPtr;
+		}
+	}
 
-	uint32* ExistingIdPtr = GAnimTraceNames.Find(InName);
-	if(ExistingIdPtr == nullptr)
+	if(bShouldTrace)
 	{
 		int32 NameStringLength = InName.GetStringLength() + 1;
 
@@ -401,19 +462,12 @@ uint32 FAnimTrace::OutputName(const FName& InName)
 			InName.ToString(reinterpret_cast<TCHAR*>(Out), NameStringLength);
 		};
 
-		uint32 NewId = CurrentId++;
-
 		UE_TRACE_LOG(Animation, Name, AnimationChannel, NameStringLength * sizeof(TCHAR))
-			<< Name.Id(NewId)
+			<< Name.Id(NameId)
 			<< Name.Attachment(StringCopyFunc);
-
-		GAnimTraceNames.Add(InName, NewId);
-		return NewId;
-	}
-	else
-	{
-		return *ExistingIdPtr;
 	}
+
+	return NameId;
 }
 
 void FAnimTrace::OutputSkeletalMeshComponent(const USkeletalMeshComponent* InComponent)
@@ -649,6 +703,32 @@ void FAnimTrace::OutputAnimNodeValue(const FAnimationBaseContext& InContext, con
 		<< AnimNodeValueFloat.Attachment(InKey, KeyLength * sizeof(TCHAR));
 }
 
+void FAnimTrace::OutputAnimNodeValue(const FAnimationBaseContext& InContext, const TCHAR* InKey, const FVector2D& InValue)
+{
+	bool bChannelEnabled = UE_TRACE_CHANNELEXPR_IS_ENABLED(AnimationChannel);
+	if (!bChannelEnabled)
+	{
+		return;
+	}
+
+	check(InContext.AnimInstanceProxy);
+
+	TRACE_OBJECT(InContext.AnimInstanceProxy->GetAnimInstanceObject());
+
+	int32 KeyLength = FCString::Strlen(InKey) + 1;
+
+	UE_TRACE_LOG(Animation, AnimNodeValueVector2D, AnimationChannel, KeyLength * sizeof(TCHAR))
+		<< AnimNodeValueVector2D.Cycle(FPlatformTime::Cycles64())
+		<< AnimNodeValueVector2D.AnimInstanceId(FObjectTrace::GetObjectId(InContext.AnimInstanceProxy->GetAnimInstanceObject()))
+		<< AnimNodeValueVector2D.NodeId(InContext.GetCurrentNodeId())
+		<< AnimNodeValueVector2D.KeyLength(KeyLength)
+		<< AnimNodeValueVector2D.ValueX(InValue.X)
+		<< AnimNodeValueVector2D.ValueY(InValue.Y)
+		<< AnimNodeValueVector2D.FrameCounter((uint16)(GFrameCounter % 0xffff))
+		<< AnimNodeValueVector2D.Attachment(InKey, KeyLength * sizeof(TCHAR));
+}
+
+
 void FAnimTrace::OutputAnimNodeValue(const FAnimationBaseContext& InContext, const TCHAR* InKey, const FRotator& Value)
 {
 	const FVector VectorValue(Value.Roll, Value.Pitch, Value.Yaw);
@@ -834,4 +914,88 @@ void FAnimTrace::OutputStateMachineState(const FAnimationBaseContext& InContext,
 		<< StateMachineState.ElapsedTime(InElapsedTime);
 }
 
+void FAnimTrace::OutputAnimNotify(UAnimInstance* InAnimInstance, const FAnimNotifyEvent& InNotifyEvent, ENotifyEventType InEventType)
+{
+	bool bChannelEnabled = UE_TRACE_CHANNELEXPR_IS_ENABLED(AnimationChannel);
+	if (!bChannelEnabled)
+	{
+		return;
+	}
+
+	TRACE_OBJECT(InAnimInstance);
+
+	const UObject* NotifyObject = nullptr;
+	const UObject* NotifyAsset = nullptr;
+	if(InNotifyEvent.Notify)
+	{
+		NotifyObject = InNotifyEvent.Notify;
+		NotifyAsset = NotifyObject->GetOuter();
+	}
+	else if(InNotifyEvent.NotifyStateClass)
+	{
+		NotifyObject = InNotifyEvent.NotifyStateClass;
+		NotifyAsset = NotifyObject->GetOuter();
+	}
+
+	TRACE_OBJECT(NotifyAsset);
+	TRACE_OBJECT(NotifyObject);
+
+	const uint32 NameId = OutputName(InNotifyEvent.NotifyName);
+
+	UE_TRACE_LOG(Animation, Notify, AnimationChannel)
+		<< Notify.Cycle(FPlatformTime::Cycles64())
+		<< Notify.AnimInstanceId(FObjectTrace::GetObjectId(InAnimInstance))
+		<< Notify.AssetId(FObjectTrace::GetObjectId(NotifyAsset))
+		<< Notify.NotifyId(FObjectTrace::GetObjectId(NotifyObject))
+		<< Notify.NameId(NameId)
+		<< Notify.Time(InNotifyEvent.GetTime())
+		<< Notify.Duration(InNotifyEvent.GetDuration())
+		<< Notify.NotifyEventType((uint8)InEventType);
+}
+
+void FAnimTrace::OutputAnimSyncMarker(UAnimInstance* InAnimInstance, const FPassedMarker& InPassedSyncMarker)
+{
+	bool bChannelEnabled = UE_TRACE_CHANNELEXPR_IS_ENABLED(AnimationChannel);
+	if (!bChannelEnabled)
+	{
+		return;
+	}
+
+	TRACE_OBJECT(InAnimInstance);
+
+	const uint32 NameId = OutputName(InPassedSyncMarker.PassedMarkerName);
+
+	UE_TRACE_LOG(Animation, SyncMarker, AnimationChannel)
+		<< SyncMarker.Cycle(FPlatformTime::Cycles64())
+		<< SyncMarker.AnimInstanceId(FObjectTrace::GetObjectId(InAnimInstance))
+		<< SyncMarker.NameId(NameId);
+}
+
+void FAnimTrace::OutputMontage(UAnimInstance* InAnimInstance, const FAnimMontageInstance& InMontageInstance)
+{
+	bool bChannelEnabled = UE_TRACE_CHANNELEXPR_IS_ENABLED(AnimationChannel);
+	if (!bChannelEnabled)
+	{
+		return;
+	}
+
+	if(InMontageInstance.Montage != nullptr)
+	{
+		TRACE_OBJECT(InAnimInstance);
+		TRACE_OBJECT(InMontageInstance.Montage);
+
+		const uint32 CurrentSectionNameId = OutputName(InMontageInstance.GetCurrentSection());
+		const uint32 NextSectionNameId = OutputName(InMontageInstance.GetNextSection());
+
+		UE_TRACE_LOG(Animation, Montage, AnimationChannel)
+			<< Montage.Cycle(FPlatformTime::Cycles64())
+			<< Montage.AnimInstanceId(FObjectTrace::GetObjectId(InAnimInstance))
+			<< Montage.MontageId(FObjectTrace::GetObjectId(InMontageInstance.Montage))
+			<< Montage.CurrentSectionNameId(CurrentSectionNameId)
+			<< Montage.NextSectionNameId(NextSectionNameId)
+			<< Montage.Weight(InMontageInstance.GetWeight())
+			<< Montage.DesiredWeight(InMontageInstance.GetDesiredWeight());
+	}
+}
+
 #endif
\ No newline at end of file
diff --git a/Engine/Source/Runtime/Engine/Public/Animation/AnimTrace.h b/Engine/Source/Runtime/Engine/Public/Animation/AnimTrace.h
--- a/Engine/Source/Runtime/Engine/Public/Animation/AnimTrace.h
+++ b/Engine/Source/Runtime/Engine/Public/Animation/AnimTrace.h
@@ -17,6 +17,7 @@ struct FAnimInstanceProxy;
 struct FAnimTickRecord;
 struct FAnimationBaseContext;
 struct FAnimationUpdateContext;
+class UAnimInstance;
 class USkeletalMesh;
 class USkeletalMeshComponent;
 struct FAnimationInitializeContext;
@@ -29,6 +30,10 @@ class FName;
 struct FVector;
 struct FRotator;
 struct FAnimNode_SequencePlayer;
+struct FAnimNotifyEvent;
+struct FPassedMarker;
+struct FAnimSyncMarker;
+struct FAnimMontageInstance;
 
 struct FAnimTrace
 {
@@ -42,6 +47,16 @@ struct FAnimTrace
 		Evaluate = 4,
 	};
 
+	/** The various events called on notifies */
+	enum class ENotifyEventType : uint8
+	{
+		Event = 0,
+		Begin = 1,
+		End = 2,
+		Tick = 3,
+		SyncMarker = 4	// We 'fake' sync markers with a notify type for convenience
+	};
+
 	/** Helper for outputting anim nodes */
 	struct FScopedAnimNodeTrace
 	{
@@ -122,6 +137,7 @@ struct FAnimTrace
 	ENGINE_API static void OutputAnimNodeValue(const FAnimationBaseContext& InContext, const TCHAR* InKey, bool InValue);
 	ENGINE_API static void OutputAnimNodeValue(const FAnimationBaseContext& InContext, const TCHAR* InKey, int32 InValue);
 	ENGINE_API static void OutputAnimNodeValue(const FAnimationBaseContext& InContext, const TCHAR* InKey, float InValue);
+	ENGINE_API static void OutputAnimNodeValue(const FAnimationBaseContext& InContext, const TCHAR* InKey, const FVector2D& InValue);
 	ENGINE_API static void OutputAnimNodeValue(const FAnimationBaseContext& InContext, const TCHAR* InKey, const FVector& InValue);
 	ENGINE_API static void OutputAnimNodeValue(const FAnimationBaseContext& InContext, const TCHAR* InKey, const FRotator& InValue);
 	ENGINE_API static void OutputAnimNodeValue(const FAnimationBaseContext& InContext, const TCHAR* InKey, const FName& InValue);
@@ -140,6 +156,15 @@ struct FAnimTrace
 
 	/** Helper function to output a state machine state's info */
 	ENGINE_API static void OutputStateMachineState(const FAnimationBaseContext& InContext, int32 InStateMachineIndex, int32 InStateIndex, float InStateWeight, float InElapsedTime);
+
+	/** Helper function to output an anim notify event */
+	ENGINE_API static void OutputAnimNotify(UAnimInstance* InAnimInstance, const FAnimNotifyEvent& InNotifyEvent, ENotifyEventType InEventType);
+
+	/** Helper function to output an anim sync marker event */
+	ENGINE_API static void OutputAnimSyncMarker(UAnimInstance* InAnimInstance, const FPassedMarker& InPassedSyncMarker);
+
+	/** Helper function to output a montage instance's info */
+	ENGINE_API static void OutputMontage(UAnimInstance* InAnimInstance, const FAnimMontageInstance& InMontageInstance);
 };
 
 #define TRACE_ANIM_TICK_RECORD(Context, TickRecord) \
@@ -172,6 +197,15 @@ struct FAnimTrace
 #define TRACE_ANIM_STATE_MACHINE_STATE(Context, StateMachineIndex, StateIndex, StateWeight, ElapsedTime) \
 	FAnimTrace::OutputStateMachineState(Context, StateMachineIndex, StateIndex, StateWeight, ElapsedTime);
 
+#define TRACE_ANIM_NOTIFY(AnimInstance, NotifyEvent, EventType) \
+	FAnimTrace::OutputAnimNotify(AnimInstance, NotifyEvent, FAnimTrace::ENotifyEventType::##EventType);
+
+#define TRACE_ANIM_SYNC_MARKER(AnimInstance, SyncMarker) \
+	FAnimTrace::OutputAnimSyncMarker(AnimInstance, SyncMarker);
+
+#define TRACE_ANIM_MONTAGE(AnimInstance, MontageInstance) \
+	FAnimTrace::OutputMontage(AnimInstance, MontageInstance);
+
 #else
 
 #define TRACE_ANIM_TICK_RECORD(Context, TickRecord)
@@ -184,5 +218,8 @@ struct FAnimTrace
 #define TRACE_ANIM_NODE_VALUE(Context, Key, Value)
 #define TRACE_ANIM_SEQUENCE_PLAYER(Context, Node)
 #define TRACE_ANIM_STATE_MACHINE_STATE(Context, StateMachineIndex, StateIndex, StateWeight, ElapsedTime)
+#define TRACE_ANIM_NOTIFY(AnimInstance, NotifyEvent, EventType)
+#define TRACE_ANIM_SYNC_MARKER(AnimInstance, SyncMarker)
+#define TRACE_ANIM_MONTAGE(AnimInstance, MontageInstance)
 
 #endif
\ No newline at end of file
