Removed unused variable.

diff --git a/Engine/Plugins/Animation/GameplayInsights/Source/GameplayInsights/Private/GameplayInsightsModule.cpp b/Engine/Plugins/Animation/GameplayInsights/Source/GameplayInsights/Private/GameplayInsightsModule.cpp
--- a/Engine/Plugins/Animation/GameplayInsights/Source/GameplayInsights/Private/GameplayInsightsModule.cpp
+++ b/Engine/Plugins/Animation/GameplayInsights/Source/GameplayInsights/Private/GameplayInsightsModule.cpp
@@ -61,7 +61,6 @@ void FGameplayInsightsModule::StartupModule()
 	StoreService = TSharedPtr<Trace::FStoreService>(Trace::FStoreService::Create(StoreServiceDesc));
 
 	FCoreDelegates::OnPreExit.AddLambda([this]() {
-		StoreClient.Reset();
 		StoreService.Reset();
 	});
 
diff --git a/Engine/Plugins/Animation/GameplayInsights/Source/GameplayInsights/Private/GameplayInsightsModule.h b/Engine/Plugins/Animation/GameplayInsights/Source/GameplayInsights/Private/GameplayInsightsModule.h
--- a/Engine/Plugins/Animation/GameplayInsights/Source/GameplayInsights/Private/GameplayInsightsModule.h
+++ b/Engine/Plugins/Animation/GameplayInsights/Source/GameplayInsights/Private/GameplayInsightsModule.h
@@ -36,6 +36,5 @@ private:
 #endif
 
 	TSharedPtr<Trace::FStoreService> StoreService;
-	TSharedPtr<Trace::FStoreClient> StoreClient;
 	TWeakPtr<FTabManager> WeakTimingProfilerTabManager;
 };
