Gameplay Insights (Editor) changes:
- Introduce GameplayInsightsEditor module, allowing to compile non-editor functionality for UnrealInsights
- Enable GameplayInsights modules for UnrealInsights standalone project / tool
- Ensure that UnrealInsights load 'Default' phase plugins
- Split extendability for Major Insights tabs, into Configuration and Extenders

diff --git a/Engine/Plugins/Animation/GameplayInsights/GameplayInsights.uplugin b/Engine/Plugins/Animation/GameplayInsights/GameplayInsights.uplugin
--- a/Engine/Plugins/Animation/GameplayInsights/GameplayInsights.uplugin
+++ b/Engine/Plugins/Animation/GameplayInsights/GameplayInsights.uplugin
@@ -14,11 +14,20 @@
 	"CanContainContent": false,
 	"IsBetaVersion": false,
 	"Installed": false,
+	"SupportedPrograms": [ "UnrealInsights" ],
 	"Modules": [
 		{
 			"Name": "GameplayInsights",
+			"Type": "UncookedOnly",
+			"LoadingPhase": "Default",
+			"WhitelistPrograms": [
+				"UnrealInsights"
+			]
+		},
+		{
+			"Name": "GameplayInsightsEditor",
 			"Type": "Editor",
 			"LoadingPhase": "Default"
 		}
 	]
-}
\ No newline at end of file
+}
diff --git a/Engine/Plugins/Animation/GameplayInsights/Source/GameplayInsights/GameplayInsights.Build.cs b/Engine/Plugins/Animation/GameplayInsights/Source/GameplayInsights/GameplayInsights.Build.cs
--- a/Engine/Plugins/Animation/GameplayInsights/Source/GameplayInsights/GameplayInsights.Build.cs
+++ b/Engine/Plugins/Animation/GameplayInsights/Source/GameplayInsights/GameplayInsights.Build.cs
@@ -39,6 +39,7 @@ namespace UnrealBuildTool.Rules
 					"AnimationBlueprintEditor",
 					"Persona",
 					"UnrealEd",
+  					"GameplayInsightsEditor"
 				});
 			}
 		}
diff --git a/Engine/Plugins/Animation/GameplayInsights/Source/GameplayInsights/Private/AnimationSharedData.cpp b/Engine/Plugins/Animation/GameplayInsights/Source/GameplayInsights/Private/AnimationSharedData.cpp
--- a/Engine/Plugins/Animation/GameplayInsights/Source/GameplayInsights/Private/AnimationSharedData.cpp
+++ b/Engine/Plugins/Animation/GameplayInsights/Source/GameplayInsights/Private/AnimationSharedData.cpp
@@ -18,11 +18,14 @@
 #include "GameplayInsightsModule.h"
 #include "Modules/ModuleManager.h"
 #include "Widgets/Docking/SDockTab.h"
-#include "Animation/AnimTypes.h"
 #include "AnimNotifiesTrack.h"
 #include "MontageTrack.h"
 #include "STrackVariantValueView.h"
 
+#if WITH_ENGINE
+#include "Animation/AnimTypes.h"
+#endif 
+
 #if WITH_EDITOR
 #include "EditorViewportClient.h"
 #include "Editor/EditorEngine.h"
diff --git a/Engine/Plugins/Animation/GameplayInsights/Source/GameplayInsights/Private/GameplayInsightsModule.cpp b/Engine/Plugins/Animation/GameplayInsights/Source/GameplayInsights/Private/GameplayInsightsModule.cpp
--- a/Engine/Plugins/Animation/GameplayInsights/Source/GameplayInsights/Private/GameplayInsightsModule.cpp
+++ b/Engine/Plugins/Animation/GameplayInsights/Source/GameplayInsights/Private/GameplayInsightsModule.cpp
@@ -27,10 +27,7 @@
 
 #define LOCTEXT_NAMESPACE "GameplayInsightsModule"
 
-namespace GameplayInsightsTabs
-{
-	static const FName DocumentTab("DocumentTab");
-};
+const FName GameplayInsightsTabs::DocumentTab("DocumentTab");
 
 void FGameplayInsightsModule::StartupModule()
 {
@@ -43,99 +40,105 @@ void FGameplayInsightsModule::StartupModule()
 		return true;
 	});
 
+	IUnrealInsightsModule& UnrealInsightsModule = FModuleManager::LoadModuleChecked<IUnrealInsightsModule>("TraceInsights");
+
 #if WITH_EDITOR
-	IAnimationBlueprintEditorModule& AnimationBlueprintEditorModule = FModuleManager::LoadModuleChecked<IAnimationBlueprintEditorModule>("AnimationBlueprintEditor");
-	CustomDebugObjectHandle = AnimationBlueprintEditorModule.OnGetCustomDebugObjects().AddLambda([this](const IAnimationBlueprintEditor& InAnimationBlueprintEditor, TArray<FCustomDebugObject>& OutDebugList)
+	if (!IsRunningCommandlet())
 	{
-		GameplayTimingViewExtender.GetCustomDebugObjects(InAnimationBlueprintEditor, OutDebugList);
-	});
-
-	IUnrealInsightsModule& UnrealInsightsModule = FModuleManager::LoadModuleChecked<IUnrealInsightsModule>("TraceInsights");
+		IAnimationBlueprintEditorModule& AnimationBlueprintEditorModule = FModuleManager::LoadModuleChecked<IAnimationBlueprintEditorModule>("AnimationBlueprintEditor");
+		CustomDebugObjectHandle = AnimationBlueprintEditorModule.OnGetCustomDebugObjects().AddLambda([this](const IAnimationBlueprintEditor& InAnimationBlueprintEditor, TArray<FCustomDebugObject>& OutDebugList)
+		{
+			GameplayTimingViewExtender.GetCustomDebugObjects(InAnimationBlueprintEditor, OutDebugList);
+		});
 
-	// Create the Store Service.
-	FString StoreDir = FPaths::ProjectSavedDir() / TEXT("TraceSessions");
-	Trace::FStoreService::FDesc StoreServiceDesc;
-	StoreServiceDesc.StoreDir = *StoreDir;
-	StoreServiceDesc.RecorderPort = 0; // Let system decide port
-	StoreServiceDesc.ThreadCount = 2;
-	StoreService = TSharedPtr<Trace::FStoreService>(Trace::FStoreService::Create(StoreServiceDesc));
+		// Create the Store Service.
+		FString StoreDir = FPaths::ProjectSavedDir() / TEXT("TraceSessions");
+		Trace::FStoreService::FDesc StoreServiceDesc;
+		StoreServiceDesc.StoreDir = *StoreDir;
+		StoreServiceDesc.RecorderPort = 0; // Let system decide port
+		StoreServiceDesc.ThreadCount = 2;
+		StoreService = TSharedPtr<Trace::FStoreService>(Trace::FStoreService::Create(StoreServiceDesc));
 
-	FCoreDelegates::OnPreExit.AddLambda([this]() {
-		StoreService.Reset();
-	});
+		FCoreDelegates::OnPreExit.AddLambda([this]() {
+			StoreService.Reset();
+		});
 
-	// Connect to our newly created store and setup the insights module
-	UnrealInsightsModule.ConnectToStore(TEXT("localhost"), StoreService->GetPort());
-	Trace::SendTo(TEXT("localhost"), StoreService->GetRecorderPort());
+		// Connect to our newly created store and setup the insights module
+		UnrealInsightsModule.ConnectToStore(TEXT("localhost"), StoreService->GetPort());
+		Trace::SendTo(TEXT("localhost"), StoreService->GetRecorderPort());
 
-	const float DPIScaleFactor = FPlatformApplicationMisc::GetDPIScaleFactorAtPoint(10.0f, 10.0f);
+		const float DPIScaleFactor = FPlatformApplicationMisc::GetDPIScaleFactorAtPoint(10.0f, 10.0f);
 
-	TSharedRef<FTabManager::FLayout> MajorTabsLayout = FTabManager::NewLayout("GameplayInsightsMajorLayout_v1.0")
-	->AddArea
-	(
-		FTabManager::NewArea(1280.f * DPIScaleFactor, 720.0f * DPIScaleFactor)
-		->Split
-		(
-			FTabManager::NewStack()
-			->AddTab(FInsightsManagerTabs::TimingProfilerTabId, ETabState::ClosedTab)
-		)
-	);
-
-	FInsightsMajorTabConfig TimingProfilerConfig;
-	TimingProfilerConfig.TabLabel = LOCTEXT("GameplayInsightsTabName", "Gameplay Insights");
-	TimingProfilerConfig.TabTooltip = LOCTEXT("GameplayInsightsTabTooltip", "Open the Gameplay Insights tab.");
-	TimingProfilerConfig.Layout = FTabManager::NewLayout("GameplayInsightsTimingLayout_v1.1")
-	->AddArea
-	(
-		FTabManager::NewPrimaryArea()
-		->SetOrientation(Orient_Horizontal)
-		->Split
+		TSharedRef<FTabManager::FLayout> MajorTabsLayout = FTabManager::NewLayout("GameplayInsightsMajorLayout_v1.0")
+		->AddArea
 		(
-			FTabManager::NewSplitter()
-			->SetOrientation(Orient_Vertical)
-			->SetSizeCoefficient(0.7f)
+			FTabManager::NewArea(1280.f * DPIScaleFactor, 720.0f * DPIScaleFactor)
 			->Split
 			(
 				FTabManager::NewStack()
-				->SetSizeCoefficient(0.1f)
-				->SetHideTabWell(true)
-				->AddTab(FTimingProfilerTabs::FramesTrackID, ETabState::OpenedTab)
+				->AddTab(FInsightsManagerTabs::TimingProfilerTabId, ETabState::ClosedTab)
+			)
+		);
+
+		FInsightsMajorTabConfig TimingProfilerConfig;
+		TimingProfilerConfig.TabLabel = LOCTEXT("GameplayInsightsTabName", "Gameplay Insights");
+		TimingProfilerConfig.TabTooltip = LOCTEXT("GameplayInsightsTabTooltip", "Open the Gameplay Insights tab.");
+		TimingProfilerConfig.Layout = FTabManager::NewLayout("GameplayInsightsTimingLayout_v1.1")
+		->AddArea
+		(
+			FTabManager::NewPrimaryArea()
+			->SetOrientation(Orient_Horizontal)
+			->Split
+			(
+				FTabManager::NewSplitter()
+				->SetOrientation(Orient_Vertical)
+				->SetSizeCoefficient(0.7f)
+				->Split
+				(
+					FTabManager::NewStack()
+					->SetSizeCoefficient(0.1f)
+					->SetHideTabWell(true)
+					->AddTab(FTimingProfilerTabs::FramesTrackID, ETabState::OpenedTab)
+				)
+				->Split
+				(
+					FTabManager::NewStack()
+					->SetSizeCoefficient(0.9f)
+					->SetHideTabWell(true)
+					->AddTab(FTimingProfilerTabs::TimingViewID, ETabState::OpenedTab)
+				)
 			)
 			->Split
 			(
 				FTabManager::NewStack()
-				->SetSizeCoefficient(0.9f)
-				->SetHideTabWell(true)
-				->AddTab(FTimingProfilerTabs::TimingViewID, ETabState::OpenedTab)
+				->SetSizeCoefficient(0.3f)
+				->AddTab(GameplayInsightsTabs::DocumentTab, ETabState::ClosedTab)
 			)
-		)
-		->Split
-		(
-			FTabManager::NewStack()
-			->SetSizeCoefficient(0.3f)
-			->AddTab(GameplayInsightsTabs::DocumentTab, ETabState::ClosedTab)
-		)
-	);
+		);
+		TimingProfilerConfig.WorkspaceGroup = WorkspaceMenu::GetMenuStructure().GetDeveloperToolsProfilingCategory();
+
+		UnrealInsightsModule.RegisterMajorTabConfig(FInsightsManagerTabs::TimingProfilerTabId, TimingProfilerConfig);
+		UnrealInsightsModule.RegisterMajorTabConfig(FInsightsManagerTabs::StartPageTabId, FInsightsMajorTabConfig::Unavailable());
+		UnrealInsightsModule.RegisterMajorTabConfig(FInsightsManagerTabs::SessionInfoTabId, FInsightsMajorTabConfig::Unavailable());
+		UnrealInsightsModule.RegisterMajorTabConfig(FInsightsManagerTabs::LoadingProfilerTabId, FInsightsMajorTabConfig::Unavailable());
+		UnrealInsightsModule.RegisterMajorTabConfig(FInsightsManagerTabs::NetworkingProfilerTabId, FInsightsMajorTabConfig::Unavailable());
+
+		UnrealInsightsModule.CreateSessionViewer(false);
+		UnrealInsightsModule.StartAnalysisForLastLiveSession();
+	}
+
+#else
+	FOnRegisterMajorTabExtensions& TimingProfilerExtension = UnrealInsightsModule.OnRegisterMajorTabExtension(FInsightsManagerTabs::TimingProfilerTabId);
+	TimingProfilerExtension.AddRaw(this, &FGameplayInsightsModule::RegisterTimingProfilerLayoutExtensions);
+#endif
 
 	UnrealInsightsModule.OnMajorTabCreated().AddLambda([this](const FName& InMajorTabId, TSharedPtr<FTabManager> InTabManager)
 	{
-		if(InMajorTabId == FInsightsManagerTabs::TimingProfilerTabId)
+		if (InMajorTabId == FInsightsManagerTabs::TimingProfilerTabId)
 		{
 			WeakTimingProfilerTabManager = InTabManager;
 		}
 	});
-
-	TimingProfilerConfig.WorkspaceGroup = WorkspaceMenu::GetMenuStructure().GetDeveloperToolsProfilingCategory();
-
-	UnrealInsightsModule.RegisterMajorTabConfig(FInsightsManagerTabs::TimingProfilerTabId, TimingProfilerConfig);
-	UnrealInsightsModule.RegisterMajorTabConfig(FInsightsManagerTabs::StartPageTabId, FInsightsMajorTabConfig::Unavailable());
-	UnrealInsightsModule.RegisterMajorTabConfig(FInsightsManagerTabs::SessionInfoTabId, FInsightsMajorTabConfig::Unavailable());
-	UnrealInsightsModule.RegisterMajorTabConfig(FInsightsManagerTabs::LoadingProfilerTabId, FInsightsMajorTabConfig::Unavailable());
-	UnrealInsightsModule.RegisterMajorTabConfig(FInsightsManagerTabs::NetworkingProfilerTabId, FInsightsMajorTabConfig::Unavailable());
-
-	UnrealInsightsModule.CreateSessionViewer(false);
-	UnrealInsightsModule.StartAnalysisForLastLiveSession();
-#endif
 }
 
 void FGameplayInsightsModule::ShutdownModule()
@@ -146,6 +149,10 @@ void FGameplayInsightsModule::ShutdownModule()
 	{
 		AnimationBlueprintEditorModule->OnGetCustomDebugObjects().Remove(CustomDebugObjectHandle);
 	}
+#else
+	IUnrealInsightsModule& UnrealInsightsModule = FModuleManager::LoadModuleChecked<IUnrealInsightsModule>("TraceInsights");
+	FOnRegisterMajorTabExtensions& TimingProfilerLayoutExtension = UnrealInsightsModule.OnRegisterMajorTabExtension(FInsightsManagerTabs::TimingProfilerTabId);
+	TimingProfilerLayoutExtension.RemoveAll(this);
 #endif
 
 	FTicker::GetCoreTicker().RemoveTicker(TickerHandle);
@@ -165,6 +172,11 @@ TSharedRef<SDockTab> FGameplayInsightsModule::SpawnTimingProfilerDocumentTab(con
 	return NewTab;
 }
 
+void FGameplayInsightsModule::RegisterTimingProfilerLayoutExtensions(FInsightsMajorTabExtender& InOutExtender)
+{
+	InOutExtender.GetLayoutExtender().ExtendLayout(FTimingProfilerTabs::TimersID, ELayoutExtensionPosition::Before, FTabManager::FTab(GameplayInsightsTabs::DocumentTab, ETabState::ClosedTab));
+}
+
 IMPLEMENT_MODULE(FGameplayInsightsModule, GameplayInsights);
 
 #undef LOCTEXT_NAMESPACE
\ No newline at end of file
diff --git a/Engine/Plugins/Animation/GameplayInsights/Source/GameplayInsights/Private/GameplayInsightsModule.h b/Engine/Plugins/Animation/GameplayInsights/Source/GameplayInsights/Private/GameplayInsightsModule.h
--- a/Engine/Plugins/Animation/GameplayInsights/Source/GameplayInsights/Private/GameplayInsightsModule.h
+++ b/Engine/Plugins/Animation/GameplayInsights/Source/GameplayInsights/Private/GameplayInsightsModule.h
@@ -7,6 +7,9 @@
 #include "GameplayTraceModule.h"
 #include "GameplayTimingViewExtender.h"
 #include "Framework/Docking/TabManager.h"
+#include "IGameplayInsightsModule.h"
+
+struct FInsightsMajorTabExtender;
 
 namespace Trace 
 {
@@ -24,6 +27,9 @@ public:
 	// Spawn a document tab
 	TSharedRef<SDockTab> SpawnTimingProfilerDocumentTab(const FTabManager::FSearchPreference& InSearchPreference);
 
+protected:
+	void RegisterTimingProfilerLayoutExtensions(FInsightsMajorTabExtender& InOutExtender);
+
 private:
 	FGameplayTraceModule GameplayTraceModule;
 
diff --git a/Engine/Plugins/Animation/GameplayInsights/Source/GameplayInsights/Private/IAnimationProvider.h b/Engine/Plugins/Animation/GameplayInsights/Source/GameplayInsights/Private/IAnimationProvider.h
deleted file mode 100644
diff --git a/Engine/Plugins/Animation/GameplayInsights/Source/GameplayInsights/Private/IGameplayProvider.h b/Engine/Plugins/Animation/GameplayInsights/Source/GameplayInsights/Private/IGameplayProvider.h
deleted file mode 100644
diff --git a/Engine/Plugins/Animation/GameplayInsights/Source/GameplayInsights/Private/InsightsSkeletalMeshComponent.cpp b/Engine/Plugins/Animation/GameplayInsights/Source/GameplayInsights/Private/InsightsSkeletalMeshComponent.cpp
deleted file mode 100644
diff --git a/Engine/Plugins/Animation/GameplayInsights/Source/GameplayInsights/Private/InsightsSkeletalMeshComponent.h b/Engine/Plugins/Animation/GameplayInsights/Source/GameplayInsights/Private/InsightsSkeletalMeshComponent.h
deleted file mode 100644
diff --git a/Engine/Plugins/Animation/GameplayInsights/Source/GameplayInsights/Private/SkeletalMeshPoseTrack.cpp b/Engine/Plugins/Animation/GameplayInsights/Source/GameplayInsights/Private/SkeletalMeshPoseTrack.cpp
--- a/Engine/Plugins/Animation/GameplayInsights/Source/GameplayInsights/Private/SkeletalMeshPoseTrack.cpp
+++ b/Engine/Plugins/Animation/GameplayInsights/Source/GameplayInsights/Private/SkeletalMeshPoseTrack.cpp
@@ -166,6 +166,7 @@ void FSkeletalMeshPoseTrack::FindSkeletalMeshPoseMessage(const FTimingEventSearc
 
 void FSkeletalMeshPoseTrack::BuildContextMenu(FMenuBuilder& MenuBuilder)
 {
+#if WITH_ENGINE
 	MenuBuilder.BeginSection(TEXT("DrawingSection"), LOCTEXT("Drawing", "Drawing (Component)"));
 	{
 		MenuBuilder.AddMenuEntry(
@@ -317,6 +318,7 @@ void FSkeletalMeshPoseTrack::BuildContextMenu(FMenuBuilder& MenuBuilder)
 			}
 		}
 	}
+#endif
 }
 
 #if WITH_ENGINE
diff --git a/Engine/Plugins/Animation/GameplayInsights/Source/GameplayInsights/Private/SkeletalMeshPoseTrack.h b/Engine/Plugins/Animation/GameplayInsights/Source/GameplayInsights/Private/SkeletalMeshPoseTrack.h
--- a/Engine/Plugins/Animation/GameplayInsights/Source/GameplayInsights/Private/SkeletalMeshPoseTrack.h
+++ b/Engine/Plugins/Animation/GameplayInsights/Source/GameplayInsights/Private/SkeletalMeshPoseTrack.h
@@ -64,9 +64,10 @@ private:
 	// Helper function used to find a skeletal mesh pose
 	void FindSkeletalMeshPoseMessage(const FTimingEventSearchParameters& InParameters, TFunctionRef<void(double, double, uint32, const FSkeletalMeshPoseMessage&)> InFoundPredicate) const;
 
+#if WITH_ENGINE
 	// Updates component visibility based on draw pose flag
 	void UpdateComponentVisibility();
-
+#endif
 private:
 	/** The shared data */
 	const FAnimationSharedData& SharedData;
diff --git a/Engine/Plugins/Animation/GameplayInsights/Source/GameplayInsights/Public/IAnimationProvider.h b/Engine/Plugins/Animation/GameplayInsights/Source/GameplayInsights/Public/IAnimationProvider.h
new file mode 100644
--- /dev/null
+++ b/Engine/Plugins/Animation/GameplayInsights/Source/GameplayInsights/Public/IAnimationProvider.h
@@ -0,0 +1,262 @@
+// Copyright Epic Games, Inc. All Rights Reserved.
+
+#pragma once
+
+#include "CoreMinimal.h"
+#include "TraceServices/Model/AnalysisSession.h"
+#include "TraceServices/Containers/Timelines.h"
+
+struct FAnimTickRecordTrace;
+
+struct FSkeletalMeshInfo
+{
+	uint64 ParentIndicesStartIndex = 0;
+	uint64 Id = 0;
+	uint32 BoneCount = 0;
+};
+
+struct FSkeletalMeshNamedCurve
+{
+	uint32 Id = 0;
+	float Value = 0.0f;
+};
+
+struct FSkeletalMeshPoseMessage
+{
+	FTransform ComponentToWorld;
+	uint64 TransformStartIndex = 0;
+	uint64 CurveStartIndex = 0;
+	uint64 ComponentId = 0;	
+	uint64 MeshId = 0;	
+	uint16 NumTransforms = 0;
+	uint16 NumCurves = 0;
+	uint16 FrameCounter = 0;
+	uint16 LodIndex = 0;
+};
+
+struct FSkeletalMeshFrameMessage
+{
+	uint64 ComponentId = 0;
+	uint16 FrameCounter = 0;
+};
+
+struct FTickRecordMessage
+{
+	uint64 ComponentId = 0;
+	uint64 AnimInstanceId = 0;
+	uint64 AssetId = 0;
+	int32 NodeId = -1;
+	float BlendWeight = 0.0f;
+	float PlaybackTime = 0.0f;
+	float RootMotionWeight = 0.0f;
+	float PlayRate = 0.0f;
+	float BlendSpacePositionX = 0.0f;
+	float BlendSpacePositionY = 0.0f;
+	uint16 FrameCounter = 0;
+	bool bLooping = false;
+	bool bIsBlendSpace = false;
+	bool bContinuous = true;
+};
+
+enum class EAnimGraphPhase : uint8
+{
+	Initialize = 0,
+	PreUpdate = 1,
+	Update = 2,
+	CacheBones = 3,
+	Evaluate = 4,
+};
+
+struct FAnimGraphMessage
+{
+	uint64 AnimInstanceId = 0;
+	int32 NodeCount = 0;
+	uint16 FrameCounter = 0;
+	EAnimGraphPhase Phase = EAnimGraphPhase::Initialize;
+};
+
+struct FAnimNodeMessage
+{
+	const TCHAR* NodeName = nullptr;
+	uint64 AnimInstanceId = 0;
+	int32 PreviousNodeId = -1;
+	int32 NodeId = -1;
+	float Weight = 0.0f;
+	float RootMotionWeight = 0.0f;
+	uint16 FrameCounter = 0;
+	EAnimGraphPhase Phase = EAnimGraphPhase::Initialize;
+};
+
+enum class EAnimNodeValueType : uint8
+{
+	Bool,
+	Int32,
+	Float,
+	Vector2D,
+	Vector,
+	String,
+	Object,
+	Class,
+};
+
+struct FVariantValue
+{
+	FVariantValue()
+		: Vector(FVector::ZeroVector)
+	{
+	}
+
+	struct FVector2DEntry
+	{
+		FVector2DEntry(const FVector2D& InVector)
+			: Value(InVector)
+		{}
+
+		FVector2D Value;
+	};
+
+	struct FVectorEntry
+	{
+		FVectorEntry(const FVector& InVector)
+			: Value(InVector)
+		{}
+
+		FVector Value;
+	};
+
+	union
+	{
+		struct
+		{
+			bool bValue;
+		} Bool;
+		struct
+		{
+			int32 Value;
+		} Int32;
+		struct
+		{
+			float Value;
+		} Float;
+		FVector2DEntry Vector2D;
+		FVectorEntry Vector;
+		struct
+		{
+			const TCHAR* Value;
+		} String;
+		struct
+		{
+			uint64 Value;
+		} Object;
+		struct
+		{
+			uint64 Value;
+		} Class;
+	};
+
+	EAnimNodeValueType Type = EAnimNodeValueType::Bool;
+};
+
+struct FAnimNodeValueMessage
+{
+	const TCHAR* Key = nullptr;
+	uint64 AnimInstanceId = 0;
+	FVariantValue Value;
+	int32 NodeId = -1;
+	uint16 FrameCounter = 0;
+};
+
+struct FAnimSequencePlayerMessage
+{
+	uint64 AnimInstanceId = 0;
+	int32 NodeId = -1;
+	float Position = 0.0f;
+	float Length = 0.0f;
+	int32 FrameCount = 0;
+};
+
+struct FBlendSpacePlayerMessage
+{
+	uint64 AnimInstanceId = 0;
+	uint64 BlendSpaceId = 0;
+	int32 NodeId = -1;
+	float PositionX = 0.0f;
+	float PositionY = 0.0f;
+	float PositionZ = 0.0f;
+};
+
+struct FAnimStateMachineMessage
+{
+	uint64 AnimInstanceId = 0;
+	int32 NodeId = -1;
+	int32 StateMachineIndex = -1;
+	int32 StateIndex = -1;
+	float StateWeight = 0.0f;
+	float ElapsedTime = 0.0f;
+};
+
+enum class EAnimNotifyMessageType : uint8
+{
+	Event = 0,
+	Begin = 1,
+	End = 2,
+	Tick = 3,
+	SyncMarker = 4	// We 'fake' sync markers with a notify type for convenience
+};
+
+struct FAnimNotifyMessage
+{
+	uint64 AnimInstanceId = 0;
+	uint64 AssetId = 0;
+	uint64 NotifyId = 0;
+	uint32 NameId = 0;
+	float Time = 0.0f; 
+	float Duration = 0.0f;
+	EAnimNotifyMessageType NotifyEventType = EAnimNotifyMessageType::Event;
+};
+
+struct FAnimMontageMessage
+{
+	uint64 AnimInstanceId = 0;
+	uint64 MontageId = 0;
+	uint32 CurrentSectionNameId = 0;
+	uint32 NextSectionNameId = 0;
+	float Weight = 0.0f;
+	float DesiredWeight = 0.0f;
+};
+
+class IAnimationProvider : public Trace::IProvider
+{
+public:
+	typedef Trace::ITimeline<FTickRecordMessage> TickRecordTimeline;
+	typedef Trace::ITimeline<FSkeletalMeshPoseMessage> SkeletalMeshPoseTimeline;
+	typedef Trace::ITimeline<FAnimGraphMessage> AnimGraphTimeline;
+	typedef Trace::ITimeline<FAnimNodeMessage> AnimNodesTimeline;
+	typedef Trace::ITimeline<FAnimNodeValueMessage> AnimNodeValuesTimeline;
+	typedef Trace::ITimeline<FAnimSequencePlayerMessage> AnimSequencePlayersTimeline;
+	typedef Trace::ITimeline<FAnimStateMachineMessage> StateMachinesTimeline;
+	typedef Trace::ITimeline<FBlendSpacePlayerMessage> BlendSpacePlayersTimeline;
+	typedef Trace::ITimeline<FAnimNotifyMessage> AnimNotifyTimeline;
+	typedef Trace::ITimeline<FAnimMontageMessage> AnimMontageTimeline;
+
+	virtual bool ReadSkeletalMeshPoseTimeline(uint64 InObjectId, TFunctionRef<void(const SkeletalMeshPoseTimeline&, bool)> Callback) const = 0;
+	virtual void GetSkeletalMeshComponentSpacePose(const FSkeletalMeshPoseMessage& InMessage, const FSkeletalMeshInfo& InMeshInfo, FTransform& OutComponentToWorld, TArray<FTransform>& OutTransforms) const = 0;
+	virtual void EnumerateSkeletalMeshCurveIds(uint64 InObjectId, TFunctionRef<void(uint32)> Callback) const = 0;
+	virtual void EnumerateSkeletalMeshCurves(const FSkeletalMeshPoseMessage& InMessage, TFunctionRef<void(const FSkeletalMeshNamedCurve&)> Callback) const = 0;
+	virtual void EnumerateTickRecordTimelines(uint64 InObjectId, TFunctionRef<void(uint64, int32, const TickRecordTimeline&)> Callback) const = 0;
+	virtual bool ReadTickRecordTimeline(uint64 InObjectId, uint64 InAssetId, int32 InNodeId, TFunctionRef<void(const TickRecordTimeline&)> Callback) const = 0;
+	virtual bool ReadAnimGraphTimeline(uint64 InObjectId, TFunctionRef<void(const AnimGraphTimeline&)> Callback) const = 0;
+	virtual bool ReadAnimNodesTimeline(uint64 InObjectId, TFunctionRef<void(const AnimNodesTimeline&)> Callback) const = 0;
+	virtual bool ReadAnimNodeValuesTimeline(uint64 InObjectId, TFunctionRef<void(const AnimNodeValuesTimeline&)> Callback) const = 0;
+	virtual bool ReadAnimSequencePlayersTimeline(uint64 InObjectId, TFunctionRef<void(const AnimSequencePlayersTimeline&)> Callback) const = 0;
+	virtual bool ReadAnimBlendSpacePlayersTimeline(uint64 InObjectId, TFunctionRef<void(const BlendSpacePlayersTimeline&)> Callback) const = 0;
+	virtual bool ReadStateMachinesTimeline(uint64 InObjectId, TFunctionRef<void(const StateMachinesTimeline&)> Callback) const = 0;
+	virtual bool ReadNotifyTimeline(uint64 InObjectId, TFunctionRef<void(const AnimNotifyTimeline&)> Callback) const = 0;
+	virtual void EnumerateNotifyStateTimelines(uint64 InObjectId, TFunctionRef<void(uint64, const AnimNotifyTimeline&)> Callback) const = 0;
+	virtual bool ReadMontageTimeline(uint64 InObjectId, TFunctionRef<void(const AnimMontageTimeline&)> Callback) const = 0;
+	virtual void EnumerateMontageIds(uint64 InObjectId, TFunctionRef<void(uint64)> Callback) const = 0;
+	virtual const FSkeletalMeshInfo* FindSkeletalMeshInfo(uint64 InObjectId) const = 0;
+	virtual const TCHAR* GetName(uint32 InId) const = 0;
+	virtual FText FormatNodeKeyValue(const FAnimNodeValueMessage& InMessage) const = 0;
+	virtual FText FormatNodeValue(const FAnimNodeValueMessage& InMessage) const = 0;
+};
\ No newline at end of file
diff --git a/Engine/Plugins/Animation/GameplayInsights/Source/GameplayInsights/Public/IGameplayInsightsModule.h b/Engine/Plugins/Animation/GameplayInsights/Source/GameplayInsights/Public/IGameplayInsightsModule.h
new file mode 100644
--- /dev/null
+++ b/Engine/Plugins/Animation/GameplayInsights/Source/GameplayInsights/Public/IGameplayInsightsModule.h
@@ -0,0 +1,10 @@
+// Copyright Epic Games, Inc. All Rights Reserved.
+
+#pragma once
+
+#include "CoreMinimal.h"
+
+namespace GameplayInsightsTabs
+{
+	extern GAMEPLAYINSIGHTS_API const FName DocumentTab;
+};
diff --git a/Engine/Plugins/Animation/GameplayInsights/Source/GameplayInsights/Public/IGameplayProvider.h b/Engine/Plugins/Animation/GameplayInsights/Source/GameplayInsights/Public/IGameplayProvider.h
new file mode 100644
--- /dev/null
+++ b/Engine/Plugins/Animation/GameplayInsights/Source/GameplayInsights/Public/IGameplayProvider.h
@@ -0,0 +1,49 @@
+// Copyright Epic Games, Inc. All Rights Reserved.
+
+#pragma once
+
+#include "CoreMinimal.h"
+#include "TraceServices/Model/AnalysisSession.h"
+#include "TraceServices/Containers/Timelines.h"
+
+struct FClassInfo
+{	
+	uint64 Id = 0;
+	uint64 SuperId = 0;
+	const TCHAR* Name = nullptr;
+	const TCHAR* PathName = nullptr;
+};
+
+struct FObjectInfo
+{	
+	uint64 Id = 0;
+	uint64 OuterId = 0;
+	uint64 ClassId = 0;
+	const TCHAR* Name = nullptr;
+	const TCHAR* PathName = nullptr;
+};
+
+struct FObjectEventMessage
+{
+	uint64 Id = 0;
+	const TCHAR* Name = nullptr;
+};
+
+// Delegate fired when an object receives an end play event
+DECLARE_MULTICAST_DELEGATE_ThreeParams(FOnObjectEndPlay, uint64 /*ObjectId*/, double /*Time*/, const FObjectInfo& /*ObjectInfo*/);
+
+class IGameplayProvider : public Trace::IProvider
+{
+public:
+	typedef Trace::ITimeline<FObjectEventMessage> ObjectEventsTimeline;
+
+	virtual bool ReadObjectEventsTimeline(uint64 InObjectId, TFunctionRef<void(const ObjectEventsTimeline&)> Callback) const = 0;
+	virtual bool ReadObjectEvent(uint64 InObjectId, uint64 InMessageId, TFunctionRef<void(const FObjectEventMessage&)> Callback) const = 0;
+	virtual void EnumerateObjects(TFunctionRef<void(const FObjectInfo&)> Callback) const = 0;
+	virtual const FClassInfo* FindClassInfo(uint64 InClassId) const = 0;
+	virtual const FObjectInfo* FindObjectInfo(uint64 InObjectId) const = 0;
+	virtual const FClassInfo& GetClassInfo(uint64 InClassId) const = 0;
+	virtual const FClassInfo& GetClassInfoFromObject(uint64 InObjectId) const = 0;
+	virtual const FObjectInfo& GetObjectInfo(uint64 InObjectId) const = 0;
+	virtual FOnObjectEndPlay& OnObjectEndPlay() = 0;
+};
\ No newline at end of file
diff --git a/Engine/Plugins/Animation/GameplayInsights/Source/GameplayInsightsEditor/GameplayInsightsEditor.Build.cs b/Engine/Plugins/Animation/GameplayInsights/Source/GameplayInsightsEditor/GameplayInsightsEditor.Build.cs
new file mode 100644
--- /dev/null
+++ b/Engine/Plugins/Animation/GameplayInsights/Source/GameplayInsightsEditor/GameplayInsightsEditor.Build.cs
@@ -0,0 +1,43 @@
+// Copyright Epic Games, Inc. All Rights Reserved.
+
+namespace UnrealBuildTool.Rules
+{
+	public class GameplayInsightsEditor : ModuleRules
+	{
+		public GameplayInsightsEditor(ReadOnlyTargetRules Target) : base(Target)
+		{
+			PrivateDependencyModuleNames.AddRange(
+			new string[]
+			{
+				"Core",
+                "TraceServices"
+            });
+
+            PrivateIncludePathModuleNames.AddRange(
+			new string[]
+			{
+				"GameplayInsights"
+			});
+
+            if (Target.bCompileAgainstEngine)
+			{
+				PrivateDependencyModuleNames.AddRange(
+				new string[]
+				{
+					"Engine",
+					"CoreUObject",
+				});
+			}
+
+			if (Target.bBuildEditor)
+			{
+				PrivateDependencyModuleNames.AddRange(
+				new string[]
+				{
+					"UnrealEd",
+				});
+			}
+		}
+	}
+}
+
diff --git a/Engine/Plugins/Animation/GameplayInsights/Source/GameplayInsightsEditor/Private/GamePlayInsightsEditor.cpp b/Engine/Plugins/Animation/GameplayInsights/Source/GameplayInsightsEditor/Private/GamePlayInsightsEditor.cpp
new file mode 100644
--- /dev/null
+++ b/Engine/Plugins/Animation/GameplayInsights/Source/GameplayInsightsEditor/Private/GamePlayInsightsEditor.cpp
@@ -0,0 +1,8 @@
+// Copyright Epic Games, Inc. All Rights Reserved.
+
+#pragma once
+
+#include "CoreMinimal.h"
+#include "Modules/ModuleManager.h"
+
+IMPLEMENT_MODULE(FDefaultModuleImpl, GameplayInsightsEditor);
diff --git a/Engine/Plugins/Animation/GameplayInsights/Source/GameplayInsightsEditor/Private/InsightsSkeletalMeshComponent.cpp b/Engine/Plugins/Animation/GameplayInsights/Source/GameplayInsightsEditor/Private/InsightsSkeletalMeshComponent.cpp
new file mode 100644
--- /dev/null
+++ b/Engine/Plugins/Animation/GameplayInsights/Source/GameplayInsightsEditor/Private/InsightsSkeletalMeshComponent.cpp
@@ -0,0 +1,45 @@
+// Copyright Epic Games, Inc. All Rights Reserved.
+
+#include "InsightsSkeletalMeshComponent.h"
+#include "IAnimationProvider.h"
+
+void UInsightsSkeletalMeshComponent::SetPoseFromProvider(const IAnimationProvider& InProvider, const FSkeletalMeshPoseMessage& InMessage, const FSkeletalMeshInfo& SkeletalMeshInfo)
+{
+	// Grab transform and bone transforms from provider
+	FTransform ComponentToWorldTransform;
+	InProvider.GetSkeletalMeshComponentSpacePose(InMessage, SkeletalMeshInfo, ComponentToWorldTransform, GetEditableComponentSpaceTransforms());
+
+	// Set transform
+	SetRelativeTransform(ComponentToWorldTransform);
+
+	// Force LOD to message LOD
+	SetForcedLOD(InMessage.LodIndex + 1);
+
+	// Flip buffers once to copy the directly-written component space transforms
+	bNeedToFlipSpaceBaseBuffers = true;
+	bHasValidBoneTransform = false;
+	FlipEditableSpaceBases();
+	bHasValidBoneTransform = true;
+
+	InvalidateCachedBounds();
+	UpdateBounds();
+	MarkRenderTransformDirty();
+	MarkRenderDynamicDataDirty();
+}
+
+void UInsightsSkeletalMeshComponent::InitAnim(bool bForceReInit)
+{
+	if(SkeletalMesh)
+	{
+		// set up bone visibility states as this gets skipped since we allocate the component array before registration
+		for (int32 BaseIndex = 0; BaseIndex < 2; ++BaseIndex)
+		{
+			BoneVisibilityStates[BaseIndex].SetNum(SkeletalMesh->RefSkeleton.GetNum());
+
+			for (int32 BoneIndex = 0; BoneIndex < SkeletalMesh->RefSkeleton.GetNum(); BoneIndex++)
+			{
+				BoneVisibilityStates[BaseIndex][BoneIndex] = BVS_Visible;
+			}
+		}
+	}
+}
\ No newline at end of file
diff --git a/Engine/Plugins/Animation/GameplayInsights/Source/GameplayInsightsEditor/Public/InsightsSkeletalMeshComponent.h b/Engine/Plugins/Animation/GameplayInsights/Source/GameplayInsightsEditor/Public/InsightsSkeletalMeshComponent.h
new file mode 100644
--- /dev/null
+++ b/Engine/Plugins/Animation/GameplayInsights/Source/GameplayInsightsEditor/Public/InsightsSkeletalMeshComponent.h
@@ -0,0 +1,24 @@
+// Copyright Epic Games, Inc. All Rights Reserved.
+
+#pragma once
+
+#include "CoreMinimal.h"
+#include "Components/SkeletalMeshComponent.h"
+#include "InsightsSkeletalMeshComponent.generated.h"
+
+class IAnimationProvider;
+struct FSkeletalMeshPoseMessage;
+struct FSkeletalMeshInfo;
+
+UCLASS()
+class GAMEPLAYINSIGHTSEDITOR_API UInsightsSkeletalMeshComponent : public USkeletalMeshComponent
+{
+	GENERATED_BODY()
+
+public:
+	// Set this component up from a provider & message
+	void SetPoseFromProvider(const IAnimationProvider& InProvider, const FSkeletalMeshPoseMessage& InMessage, const FSkeletalMeshInfo& SkeletalMeshInfo);
+
+	// USkeletalMeshComponent interface
+	virtual void InitAnim(bool bForceReInit) override;
+};
diff --git a/Engine/Programs/UnrealInsights/Config/DefaultEngine.ini b/Engine/Programs/UnrealInsights/Config/DefaultEngine.ini
new file mode 100644
--- /dev/null
+++ b/Engine/Programs/UnrealInsights/Config/DefaultEngine.ini
@@ -0,0 +1,4 @@
+[Plugins]
++ProgramEnabledPlugins=GameplayInsights
+
+
diff --git a/Engine/Source/Developer/TraceInsights/Private/Insights/TraceInsightsModule.cpp b/Engine/Source/Developer/TraceInsights/Private/Insights/TraceInsightsModule.cpp
--- a/Engine/Source/Developer/TraceInsights/Private/Insights/TraceInsightsModule.cpp
+++ b/Engine/Source/Developer/TraceInsights/Private/Insights/TraceInsightsModule.cpp
@@ -352,6 +352,14 @@ void FTraceInsightsModule::UnregisterMajorTabConfig(const FName& InMajorTabId)
 	TabConfigs.Remove(InMajorTabId);
 }
 
+
+////////////////////////////////////////////////////////////////////////////////////////////////////
+
+FOnRegisterMajorTabExtensions& FTraceInsightsModule::OnRegisterMajorTabExtension(const FName& InMajorTabId)
+{
+	return MajorTabExtensionDelegates.FindOrAdd(InMajorTabId);
+}
+
 ////////////////////////////////////////////////////////////////////////////////////////////////////
 
 const FInsightsMajorTabConfig& FTraceInsightsModule::FindMajorTabConfig(const FName& InMajorTabId) const
@@ -368,6 +376,13 @@ const FInsightsMajorTabConfig& FTraceInsightsModule::FindMajorTabConfig(const FN
 
 ////////////////////////////////////////////////////////////////////////////////////////////////////
 
+const FOnRegisterMajorTabExtensions* FTraceInsightsModule::FindMajorTabLayoutExtension(const FName& InMajorTabId) const
+{
+	return MajorTabExtensionDelegates.Find(InMajorTabId);
+}
+
+////////////////////////////////////////////////////////////////////////////////////////////////////
+
 void FTraceInsightsModule::StartAnalysisForTraceFile(const TCHAR* InTraceFile)
 {
 	if (InTraceFile != nullptr)
diff --git a/Engine/Source/Developer/TraceInsights/Private/Insights/TraceInsightsModule.h b/Engine/Source/Developer/TraceInsights/Private/Insights/TraceInsightsModule.h
--- a/Engine/Source/Developer/TraceInsights/Private/Insights/TraceInsightsModule.h
+++ b/Engine/Source/Developer/TraceInsights/Private/Insights/TraceInsightsModule.h
@@ -42,10 +42,12 @@ public:
 	virtual void RegisterMajorTabConfig(const FName& InMajorTabId, const FInsightsMajorTabConfig& InConfig) override;
 	virtual void UnregisterMajorTabConfig(const FName& InMajorTabId) override;
 	virtual FOnInsightsMajorTabCreated& OnMajorTabCreated() override { return OnInsightsMajorTabCreatedDelegate; }
+	virtual FOnRegisterMajorTabExtensions& OnRegisterMajorTabExtension(const FName& InMajorTabId) override;
 
 	/** Find a major tab config for the specified ID */
 	const FInsightsMajorTabConfig& FindMajorTabConfig(const FName& InMajorTabId) const;
 
+	const FOnRegisterMajorTabExtensions* FindMajorTabLayoutExtension(const FName& InMajorTabId) const;
 protected:
 	void InitTraceStore();
 
@@ -95,6 +97,7 @@ protected:
 	TSharedPtr<Trace::IModuleService> TraceModuleService;
 
 	TMap<FName, FInsightsMajorTabConfig> TabConfigs;
+	TMap<FName, FOnRegisterMajorTabExtensions> MajorTabExtensionDelegates;
 
 	FOnInsightsMajorTabCreated OnInsightsMajorTabCreatedDelegate;
 
diff --git a/Engine/Source/Developer/TraceInsights/Private/Insights/Widgets/STimingProfilerWindow.cpp b/Engine/Source/Developer/TraceInsights/Private/Insights/Widgets/STimingProfilerWindow.cpp
--- a/Engine/Source/Developer/TraceInsights/Private/Insights/Widgets/STimingProfilerWindow.cpp
+++ b/Engine/Source/Developer/TraceInsights/Private/Insights/Widgets/STimingProfilerWindow.cpp
@@ -5,6 +5,7 @@
 #include "EditorStyleSet.h"
 #include "Framework/Docking/WorkspaceItem.h"
 #include "Framework/MultiBox/MultiBoxBuilder.h"
+#include "Framework/Docking/LayoutExtender.h"
 #include "SlateOptMacros.h"
 #include "Widgets/Docking/SDockTab.h"
 #include "Widgets/Images/SImage.h"
@@ -325,6 +326,8 @@ void STimingProfilerWindow::Construct(const FArguments& InArgs, const TSharedRef
 	TabManager = FGlobalTabmanager::Get()->NewTabManager(ConstructUnderMajorTab);
 	TSharedRef<FWorkspaceItem> AppMenuGroup = TabManager->AddLocalWorkspaceMenuCategory(LOCTEXT("TimingProfilerMenuGroupName", "Timing Insights"));
 
+	Extension = MakeShared<FInsightsMajorTabExtender>(TabManager);
+
 	TabManager->RegisterTabSpawner(FTimingProfilerTabs::ToolbarID, FOnSpawnTab::CreateRaw(this, &STimingProfilerWindow::SpawnTab_Toolbar))
 		.SetDisplayName(LOCTEXT("DeviceToolbarTabTitle", "Toolbar"))
 		.SetIcon(FSlateIcon(FInsightsStyle::GetStyleSetName(), "Toolbar.Icon.Small"))
@@ -374,8 +377,14 @@ void STimingProfilerWindow::Construct(const FArguments& InArgs, const TSharedRef
 	FTraceInsightsModule& TraceInsightsModule = FModuleManager::GetModuleChecked<FTraceInsightsModule>("TraceInsights");
 	FInsightsMajorTabConfig TabConfig = TraceInsightsModule.FindMajorTabConfig(FInsightsManagerTabs::TimingProfilerTabId);
 
+	const FOnRegisterMajorTabExtensions* ExtensionDelegate = TraceInsightsModule.FindMajorTabLayoutExtension(FInsightsManagerTabs::TimingProfilerTabId);
+	if (ExtensionDelegate)
+	{
+		ExtensionDelegate->Broadcast(*Extension);
+	}
+
 	// Register any new minor tabs
-	for(const FInsightsMinorTabConfig& MinorTabConfig : TabConfig.MinorTabs)
+	for(const FInsightsMinorTabConfig& MinorTabConfig : Extension->GetMinorTabs())
 	{
 		FTabSpawnerEntry& TabSpawnerEntry = TabManager->RegisterTabSpawner(MinorTabConfig.TabId, MinorTabConfig.OnSpawnTab);
 
@@ -473,8 +482,10 @@ void STimingProfilerWindow::Construct(const FArguments& InArgs, const TSharedRef
 			);
 	}
 
+	Layout->ProcessExtensions(Extension->GetLayoutExtender());
+
 	// Create & initialize main menu.
-	FMenuBarBuilder MenuBarBuilder = FMenuBarBuilder(TSharedPtr<FUICommandList>(), TabConfig.MenuExtender);
+	FMenuBarBuilder MenuBarBuilder = FMenuBarBuilder(TSharedPtr<FUICommandList>(), Extension->GetMenuExtender());
 
 	MenuBarBuilder.AddPullDownMenu(
 		LOCTEXT("MenuLabel", "Menu"),
diff --git a/Engine/Source/Developer/TraceInsights/Private/Insights/Widgets/STimingProfilerWindow.h b/Engine/Source/Developer/TraceInsights/Private/Insights/Widgets/STimingProfilerWindow.h
--- a/Engine/Source/Developer/TraceInsights/Private/Insights/Widgets/STimingProfilerWindow.h
+++ b/Engine/Source/Developer/TraceInsights/Private/Insights/Widgets/STimingProfilerWindow.h
@@ -14,6 +14,7 @@
 
 // Insights
 #include "Insights/InsightsManager.h"
+#include "Insights/IUnrealInsightsModule.h"
 
 ////////////////////////////////////////////////////////////////////////////////////////////////////
 
@@ -187,4 +188,7 @@ private:
 
 	/** The number of seconds the profiler has been active */
 	float DurationActive;
+
+	/** Tab specific slate extender structure, has the lifetime of this widget */
+	TSharedPtr<FInsightsMajorTabExtender> Extension;
 };
diff --git a/Engine/Source/Developer/TraceInsights/Public/Insights/IUnrealInsightsModule.h b/Engine/Source/Developer/TraceInsights/Public/Insights/IUnrealInsightsModule.h
--- a/Engine/Source/Developer/TraceInsights/Public/Insights/IUnrealInsightsModule.h
+++ b/Engine/Source/Developer/TraceInsights/Public/Insights/IUnrealInsightsModule.h
@@ -5,6 +5,8 @@
 #include "CoreMinimal.h"
 #include "Modules/ModuleInterface.h"
 #include "Framework/Docking/TabManager.h"
+#include "Framework/Docking/LayoutExtender.h"
+#include "Framework/MultiBox/MultiBoxExtender.h"
 
 class FExtender;
 
@@ -96,16 +98,37 @@ struct TRACEINSIGHTS_API FInsightsMajorTabConfig
 	/** The menu workspace group to use. If not specified, the default will be used. */
 	TSharedPtr<FWorkspaceItem> WorkspaceGroup;
 
+	/** Whether the tab is available for selection (i.e. registered with the tab manager) */
+	bool bIsAvailable;
+};
+
+/** Combination of extenders applied to the individual major tabs within Insights */
+struct TRACEINSIGHTS_API FInsightsMajorTabExtender
+{
+	FInsightsMajorTabExtender(TSharedPtr<FTabManager>& InTabManager) : MenuExtender(MakeShared<FExtender>()), TabManager(InTabManager) {}
+	
+	TSharedPtr<FExtender>& GetMenuExtender() { return MenuExtender; }
+	FLayoutExtender& GetLayoutExtender() { return LayoutExtender; }
+	FInsightsMinorTabConfig& AddMinorTabConfig() { return MinorTabs.AddDefaulted_GetRef(); }	
+	TSharedPtr<FTabManager> GetTabManager() const { return TabManager; }
+	const TArray<FInsightsMinorTabConfig>& GetMinorTabs() const { return MinorTabs; }
+protected:
 	/** Extender used to add to the menu for this tab */
 	TSharedPtr<FExtender> MenuExtender;
 
 	/** Any additional minor tabs to add */
 	TArray<FInsightsMinorTabConfig> MinorTabs;
 
-	/** Whether the tab is available for selection (i.e. registered with the tab manager) */
-	bool bIsAvailable;
+	/** Extender used when creating the layout for this tab */
+	FLayoutExtender LayoutExtender;
+
+	/** Tab manager for this major tab*/
+	TSharedPtr<FTabManager> TabManager;
 };
 
+/** Called back to register common layout extensions */
+DECLARE_MULTICAST_DELEGATE_OneParam(FOnRegisterMajorTabExtensions, FInsightsMajorTabExtender&);
+
 /** Interface for an Unreal Insights module. */
 class IUnrealInsightsModule : public IModuleInterface
 {
@@ -165,6 +188,12 @@ public:
 	 * @param	InMajorTabId	The major tab ID we are supplying a layout to
 	 */
 	virtual void UnregisterMajorTabConfig(const FName& InMajorTabId) = 0;
+	
+	/**
+	 * Allows for registering a delegate callback for populating a FInsightsMajorTabExtender structure	 
+	 * @param	InMajorTabId	The major tab ID to register the delegate for
+	 */
+	virtual FOnRegisterMajorTabExtensions& OnRegisterMajorTabExtension(const FName& InMajorTabId) = 0;
 
 	/** Callback invoked when a major tab is created */
 	virtual FOnInsightsMajorTabCreated& OnMajorTabCreated() = 0;
diff --git a/Engine/Source/Programs/UnrealInsights/Private/UserInterfaceCommand.cpp b/Engine/Source/Programs/UnrealInsights/Private/UserInterfaceCommand.cpp
--- a/Engine/Source/Programs/UnrealInsights/Private/UserInterfaceCommand.cpp
+++ b/Engine/Source/Programs/UnrealInsights/Private/UserInterfaceCommand.cpp
@@ -123,6 +123,7 @@ void FUserInterfaceCommand::Run()
 	// Load plug-ins.
 	// @todo: allow for better plug-in support in standalone Slate applications
 	IPluginManager::Get().LoadModulesForEnabledPlugins(ELoadingPhase::PreDefault);
+	IPluginManager::Get().LoadModulesForEnabledPlugins(ELoadingPhase::Default);
 
 	// Load optional modules.
 	FModuleManager::Get().LoadModule("SettingsEditor");
