Changed GameplayInsights module from UncookedOnly to EditorAndProgram type instead

diff --git a/Engine/Plugins/Animation/GameplayInsights/GameplayInsights.uplugin b/Engine/Plugins/Animation/GameplayInsights/GameplayInsights.uplugin
--- a/Engine/Plugins/Animation/GameplayInsights/GameplayInsights.uplugin
+++ b/Engine/Plugins/Animation/GameplayInsights/GameplayInsights.uplugin
@@ -18,7 +18,7 @@
 	"Modules": [
 		{
 			"Name": "GameplayInsights",
-			"Type": "UncookedOnly",
+			"Type": "EditorAndProgram",
 			"LoadingPhase": "Default",
 			"WhitelistPrograms": [
 				"UnrealInsights"
