CIS fix, stray pragma once

diff --git a/Engine/Plugins/Animation/GameplayInsights/Source/GameplayInsightsEditor/Private/GamePlayInsightsEditor.cpp b/Engine/Plugins/Animation/GameplayInsights/Source/GameplayInsightsEditor/Private/GamePlayInsightsEditor.cpp
--- a/Engine/Plugins/Animation/GameplayInsights/Source/GameplayInsightsEditor/Private/GamePlayInsightsEditor.cpp
+++ b/Engine/Plugins/Animation/GameplayInsights/Source/GameplayInsightsEditor/Private/GamePlayInsightsEditor.cpp
@@ -1,7 +1,5 @@
 // Copyright Epic Games, Inc. All Rights Reserved.
 
-#pragma once
-
 #include "CoreMinimal.h"
 #include "Modules/ModuleManager.h"
 
