Added transport controls to control timeline 'playback'

Controls allow playing/reverse-playing, pausing, skipping to start/end and manipulating playback speed.

diff --git a/Engine/Plugins/Animation/GameplayInsights/Source/GameplayInsights/GameplayInsights.Build.cs b/Engine/Plugins/Animation/GameplayInsights/Source/GameplayInsights/GameplayInsights.Build.cs
--- a/Engine/Plugins/Animation/GameplayInsights/Source/GameplayInsights/GameplayInsights.Build.cs
+++ b/Engine/Plugins/Animation/GameplayInsights/Source/GameplayInsights/GameplayInsights.Build.cs
@@ -39,7 +39,8 @@ namespace UnrealBuildTool.Rules
 					"AnimationBlueprintEditor",
 					"Persona",
 					"UnrealEd",
-  					"GameplayInsightsEditor"
+  					"GameplayInsightsEditor",
+					"EditorWidgets",
 				});
 			}
 		}
diff --git a/Engine/Plugins/Animation/GameplayInsights/Source/GameplayInsights/Private/GameplayInsightsStyle.h b/Engine/Plugins/Animation/GameplayInsights/Source/GameplayInsights/Private/GameplayInsightsStyle.h
--- a/Engine/Plugins/Animation/GameplayInsights/Source/GameplayInsights/Private/GameplayInsightsStyle.h
+++ b/Engine/Plugins/Animation/GameplayInsights/Source/GameplayInsights/Private/GameplayInsightsStyle.h
@@ -45,6 +45,16 @@ public:
 
 		Set("SchematicViewViewButtonIcon", new IMAGE_BRUSH("view_button", Icon20x20));
 
+		Set("TransportControls.HyperlinkSpinBox", FSpinBoxStyle()
+			.SetTextPadding(FMargin(0))
+			.SetBackgroundBrush(BORDER_BRUSH("HyperlinkDotted", FMargin(0, 0, 0, 3 / 16.0f), FSlateColor::UseSubduedForeground()))
+			.SetHoveredBackgroundBrush(FSlateNoResource())
+			.SetInactiveFillBrush(FSlateNoResource())
+			.SetActiveFillBrush(FSlateNoResource())
+			.SetForegroundColor(FSlateColor::UseSubduedForeground())
+			.SetArrowsImage(FSlateNoResource())
+		);
+
 		FSlateStyleRegistry::RegisterSlateStyle(*this);
 	}
 
diff --git a/Engine/Plugins/Animation/GameplayInsights/Source/GameplayInsights/Private/GameplaySharedData.h b/Engine/Plugins/Animation/GameplayInsights/Source/GameplayInsights/Private/GameplaySharedData.h
--- a/Engine/Plugins/Animation/GameplayInsights/Source/GameplayInsights/Private/GameplaySharedData.h
+++ b/Engine/Plugins/Animation/GameplayInsights/Source/GameplayInsights/Private/GameplaySharedData.h
@@ -35,6 +35,9 @@ public:
 	// Get the last cached analysis session
 	const Trace::IAnalysisSession& GetAnalysisSession() const { return *AnalysisSession; }
 
+	// Get the timing view session
+	Insights::ITimingViewSession& GetTimingViewSession() { return *TimingViewSession; }
+
 	// Enumerate object tracks
 	void EnumerateObjectTracks(TFunctionRef<void(const TSharedRef<FObjectEventsTrack>&)> InCallback) const;
 
diff --git a/Engine/Plugins/Animation/GameplayInsights/Source/GameplayInsights/Private/GameplayTimingViewExtender.cpp b/Engine/Plugins/Animation/GameplayInsights/Source/GameplayInsights/Private/GameplayTimingViewExtender.cpp
--- a/Engine/Plugins/Animation/GameplayInsights/Source/GameplayInsights/Private/GameplayTimingViewExtender.cpp
+++ b/Engine/Plugins/Animation/GameplayInsights/Source/GameplayInsights/Private/GameplayTimingViewExtender.cpp
@@ -7,6 +7,11 @@
 #include "UObject/WeakObjectPtr.h"
 #include "Framework/MultiBox/MultiBoxBuilder.h"
 
+#if WITH_EDITOR
+#include "SGameplayInsightsTransportControls.h"
+#include "Widgets/Layout/SBorder.h"
+#endif
+
 #if WITH_ENGINE
 #include "Engine/World.h"
 #include "Editor/EditorEngine.h"
@@ -22,6 +27,21 @@ void FGameplayTimingViewExtender::OnBeginSession(Insights::ITimingViewSession& I
 		PerSessionData = &PerSessionDataMap.Add(&InSession);
 		PerSessionData->GameplaySharedData = new FGameplaySharedData();
 		PerSessionData->AnimationSharedData = new FAnimationSharedData(*PerSessionData->GameplaySharedData);
+
+		InSession.AddOverlayWidget(
+			SNew(SOverlay)
+			.Visibility(EVisibility::SelfHitTestInvisible)
+			+SOverlay::Slot()
+			.HAlign(HAlign_Center)
+			.VAlign(VAlign_Bottom)
+			.Padding(20.0f)
+			[
+				SNew(SBorder)
+				.BorderImage(FCoreStyle::Get().GetBrush("ToolPanel.GroupBorder"))
+				[
+					SAssignNew(PerSessionData->TransportControls, SGameplayInsightsTransportControls, *PerSessionData->GameplaySharedData)
+				]
+			]);
 	}
 
 	PerSessionData->GameplaySharedData->OnBeginSession(InSession);
@@ -40,6 +60,8 @@ void FGameplayTimingViewExtender::OnEndSession(Insights::ITimingViewSession& InS
 		PerSessionData->GameplaySharedData = nullptr;
 		delete PerSessionData->AnimationSharedData;
 		PerSessionData->AnimationSharedData = nullptr;
+
+		PerSessionData->TransportControls.Reset();
 	}
 
 	PerSessionDataMap.Remove(&InSession);
@@ -60,12 +82,8 @@ void FGameplayTimingViewExtender::ExtendFilterMenu(Insights::ITimingViewSession&
 	FPerSessionData* PerSessionData = PerSessionDataMap.Find(&InSession);
 	if(PerSessionData != nullptr)
 	{
-		InMenuBuilder.BeginSection("GameplayTracks", LOCTEXT("GameplayTracksSection", "Gameplay Tracks"));
-		{
-			PerSessionData->GameplaySharedData->ExtendFilterMenu(InMenuBuilder);
-			PerSessionData->AnimationSharedData->ExtendFilterMenu(InMenuBuilder);
-		}
-		InMenuBuilder.EndSection();
+		PerSessionData->GameplaySharedData->ExtendFilterMenu(InMenuBuilder);
+		PerSessionData->AnimationSharedData->ExtendFilterMenu(InMenuBuilder);
 	}
 }
 
diff --git a/Engine/Plugins/Animation/GameplayInsights/Source/GameplayInsights/Private/GameplayTimingViewExtender.h b/Engine/Plugins/Animation/GameplayInsights/Source/GameplayInsights/Private/GameplayTimingViewExtender.h
--- a/Engine/Plugins/Animation/GameplayInsights/Source/GameplayInsights/Private/GameplayTimingViewExtender.h
+++ b/Engine/Plugins/Animation/GameplayInsights/Source/GameplayInsights/Private/GameplayTimingViewExtender.h
@@ -10,6 +10,7 @@ class FGameplaySharedData;
 class FAnimationSharedData;
 class IAnimationBlueprintEditor;
 struct FCustomDebugObject;
+class SGameplayInsightsTransportControls;
 
 class FGameplayTimingViewExtender : public Insights::ITimingViewExtender
 {
@@ -37,6 +38,7 @@ private:
 		// Shared data
 		FGameplaySharedData* GameplaySharedData;
 		FAnimationSharedData* AnimationSharedData;
+		TSharedPtr<SGameplayInsightsTransportControls> TransportControls;
 	};
 
 	// The data we host per-session
diff --git a/Engine/Plugins/Animation/GameplayInsights/Source/GameplayInsights/Private/SGameplayInsightsTransportControls.cpp b/Engine/Plugins/Animation/GameplayInsights/Source/GameplayInsights/Private/SGameplayInsightsTransportControls.cpp
new file mode 100644
--- /dev/null
+++ b/Engine/Plugins/Animation/GameplayInsights/Source/GameplayInsights/Private/SGameplayInsightsTransportControls.cpp
@@ -0,0 +1,238 @@
+// Copyright Epic Games, Inc. All Rights Reserved.
+
+#include "SGameplayInsightsTransportControls.h"
+
+#if WITH_EDITOR
+
+#include "EditorWidgetsModule.h"
+#include "Modules/ModuleManager.h"
+#include "Insights/ITimingViewSession.h"
+#include "TraceServices/Model/Frames.h"
+#include "TraceServices/Model/AnalysisSession.h"
+#include "GameplaySharedData.h"
+#include "GameplayInsightsStyle.h"
+#include "Widgets/Input/SSpinBox.h"
+#include "Math/UnitConversion.h"
+#include "Widgets/Input/NumericUnitTypeInterface.inl"
+
+#include <limits>
+
+#define LOCTEXT_NAMESPACE "SGameplayInsightsTransportControls"
+
+void SGameplayInsightsTransportControls::Construct(const FArguments& InArgs, FGameplaySharedData& InSharedData)
+{
+	SharedData = &InSharedData;
+
+	PlayRate = 1.0;
+	bPlaying = false;
+	bReverse = false;
+
+	FEditorWidgetsModule& EditorWidgetsModule = FModuleManager::LoadModuleChecked<FEditorWidgetsModule>("EditorWidgets");
+
+	FTransportControlArgs TransportControlArgs;
+	TransportControlArgs.OnForwardPlay = FOnClicked::CreateSP(this, &SGameplayInsightsTransportControls::OnClick_Forward);
+	TransportControlArgs.OnBackwardPlay = FOnClicked::CreateSP(this, &SGameplayInsightsTransportControls::OnClick_Backward);
+	TransportControlArgs.OnForwardStep = FOnClicked::CreateSP(this, &SGameplayInsightsTransportControls::OnClick_Forward_Step);
+	TransportControlArgs.OnBackwardStep = FOnClicked::CreateSP(this, &SGameplayInsightsTransportControls::OnClick_Backward_Step);
+	TransportControlArgs.OnForwardEnd = FOnClicked::CreateSP(this, &SGameplayInsightsTransportControls::OnClick_Forward_End);
+	TransportControlArgs.OnBackwardEnd = FOnClicked::CreateSP(this, &SGameplayInsightsTransportControls::OnClick_Backward_End);
+	TransportControlArgs.OnGetPlaybackMode = FOnGetPlaybackMode::CreateSP(this, &SGameplayInsightsTransportControls::GetPlaybackMode);
+
+	ChildSlot
+	[
+		SNew(SHorizontalBox)
+		+SHorizontalBox::Slot()
+		.AutoWidth()
+		[
+			EditorWidgetsModule.CreateTransportControl(TransportControlArgs)
+		]
+		+SHorizontalBox::Slot()
+		.AutoWidth()
+		.VAlign(VAlign_Center)
+		.Padding(4.0f, 0.0f)
+		[
+			SNew(SSpinBox<double>)
+			.Value_Lambda([this](){ return PlayRate; })
+			.ToolTipText(LOCTEXT("PlayRate", "Playback speed"))
+			.OnValueCommitted_Lambda([this](double InValue, ETextCommit::Type InCommitType){ PlayRate = InValue; })
+			.MinValue(0.001f)
+			.MaxValue(100.0f)
+			.Style(&FGameplayInsightsStyle::Get().GetWidgetStyle<FSpinBoxStyle>("TransportControls.HyperlinkSpinBox"))
+			.ClearKeyboardFocusOnCommit(true)
+			.Delta(0.01f)
+			.LinearDeltaSensitivity(25)
+			.TypeInterface(MakeShared<TNumericUnitTypeInterface<double>>(EUnit::Multiplier))
+		]
+	];
+
+	RegisterActiveTimer(0.0f, FWidgetActiveTimerDelegate::CreateLambda([this](double InCurrentTime, float InDeltaTime)
+	{
+		if(bPlaying)
+		{
+			double CurrentTime = SharedData->GetTimingViewSession().GetTimeMarker();
+			double Delta = (double)(bReverse ? -InDeltaTime : InDeltaTime);
+			SharedData->GetTimingViewSession().SetTimeMarker(CurrentTime + (Delta * PlayRate));
+		}
+
+		return EActiveTimerReturnType::Continue;
+	}));
+}
+
+FReply SGameplayInsightsTransportControls::OnClick_Forward_Step()
+{
+	Trace::FAnalysisSessionReadScope SessionReadScope(SharedData->GetAnalysisSession());
+	const Trace::IFrameProvider& FramesProvider = Trace::ReadFrameProvider(SharedData->GetAnalysisSession());
+
+	double CurrentTime = SharedData->GetTimingViewSession().GetTimeMarker();
+	if(CurrentTime == std::numeric_limits<double>::infinity())
+	{
+		const Trace::FFrame* FirstFrame = FramesProvider.GetFrame(ETraceFrameType::TraceFrameType_Game, 0);
+		if(FirstFrame)
+		{
+			CurrentTime = FirstFrame->StartTime + (double)KINDA_SMALL_NUMBER;
+		}
+	}
+	
+	Trace::FFrame Frame;
+	if(FramesProvider.GetFrameFromTime(ETraceFrameType::TraceFrameType_Game, CurrentTime, Frame))
+	{
+		if(Frame.Index < FramesProvider.GetFrameCount(ETraceFrameType::TraceFrameType_Game) - 1)
+		{
+			const Trace::FFrame* NextFrame = FramesProvider.GetFrame(ETraceFrameType::TraceFrameType_Game, Frame.Index + 1);
+			if(NextFrame)
+			{
+				SharedData->GetTimingViewSession().SetTimeMarker(NextFrame->StartTime + (double)KINDA_SMALL_NUMBER);
+			}
+		}
+	}
+
+	bPlaying = false;
+	bReverse = false;
+
+	return FReply::Handled();
+}
+
+FReply SGameplayInsightsTransportControls::OnClick_Forward_End()
+{
+	Trace::FAnalysisSessionReadScope SessionReadScope(SharedData->GetAnalysisSession());
+	const Trace::IFrameProvider& FramesProvider = Trace::ReadFrameProvider(SharedData->GetAnalysisSession());
+	
+	if(FramesProvider.GetFrameCount(ETraceFrameType::TraceFrameType_Game) > 0)
+	{
+		const Trace::FFrame* LastFrame = FramesProvider.GetFrame(ETraceFrameType::TraceFrameType_Game, FramesProvider.GetFrameCount(ETraceFrameType::TraceFrameType_Game) - 1);
+		if(LastFrame)
+		{
+			SharedData->GetTimingViewSession().SetTimeMarker(LastFrame->StartTime + (double)KINDA_SMALL_NUMBER);
+		}
+	}
+
+	bPlaying = false;
+	bReverse = false;
+
+	return FReply::Handled();
+}
+
+FReply SGameplayInsightsTransportControls::OnClick_Backward_Step()
+{
+	Trace::FAnalysisSessionReadScope SessionReadScope(SharedData->GetAnalysisSession());
+	const Trace::IFrameProvider& FramesProvider = Trace::ReadFrameProvider(SharedData->GetAnalysisSession());
+
+	double CurrentTime = SharedData->GetTimingViewSession().GetTimeMarker();
+	if(CurrentTime == std::numeric_limits<double>::infinity())
+	{
+		const Trace::FFrame* LastFrame = FramesProvider.GetFrame(ETraceFrameType::TraceFrameType_Game, FramesProvider.GetFrameCount(TraceFrameType_Game) - 1);
+		if(LastFrame)
+		{
+			CurrentTime = LastFrame->StartTime + (double)KINDA_SMALL_NUMBER;
+		}
+	}
+
+	Trace::FFrame Frame;
+	if(FramesProvider.GetFrameFromTime(ETraceFrameType::TraceFrameType_Game, CurrentTime, Frame))
+	{
+		if(Frame.Index > 0)
+		{
+			const Trace::FFrame* PrevFrame = FramesProvider.GetFrame(ETraceFrameType::TraceFrameType_Game, Frame.Index - 1);
+			if(PrevFrame)
+			{
+				SharedData->GetTimingViewSession().SetTimeMarker(PrevFrame->StartTime + (double)KINDA_SMALL_NUMBER);
+			}
+		}
+	}
+
+	bPlaying = false;
+	bReverse = false;
+
+	return FReply::Handled();
+}
+
+FReply SGameplayInsightsTransportControls::OnClick_Backward_End()
+{
+	Trace::FAnalysisSessionReadScope SessionReadScope(SharedData->GetAnalysisSession());
+	const Trace::IFrameProvider& FramesProvider = Trace::ReadFrameProvider(SharedData->GetAnalysisSession());
+	
+	const Trace::FFrame* FirstFrame = FramesProvider.GetFrame(ETraceFrameType::TraceFrameType_Game, 0);
+	if(FirstFrame)
+	{
+		SharedData->GetTimingViewSession().SetTimeMarker(FirstFrame->StartTime + (double)KINDA_SMALL_NUMBER);
+	}
+
+	bPlaying = false;
+	bReverse = false;
+
+	return FReply::Handled();
+}
+
+FReply SGameplayInsightsTransportControls::OnClick_Forward()
+{
+	double CurrentTime = SharedData->GetTimingViewSession().GetTimeMarker();
+	if(CurrentTime == std::numeric_limits<double>::infinity())
+	{
+		Trace::FAnalysisSessionReadScope SessionReadScope(SharedData->GetAnalysisSession());
+		const Trace::IFrameProvider& FramesProvider = Trace::ReadFrameProvider(SharedData->GetAnalysisSession());
+	
+		const Trace::FFrame* FirstFrame = FramesProvider.GetFrame(ETraceFrameType::TraceFrameType_Game, 0);
+		if(FirstFrame)
+		{
+			SharedData->GetTimingViewSession().SetTimeMarker(FirstFrame->StartTime + (double)KINDA_SMALL_NUMBER);
+		}
+	}
+
+	bPlaying = bReverse ? true : !bPlaying;
+	bReverse = false;
+	return FReply::Handled();
+}
+
+FReply SGameplayInsightsTransportControls::OnClick_Backward()
+{
+	double CurrentTime = SharedData->GetTimingViewSession().GetTimeMarker();
+	if(CurrentTime == std::numeric_limits<double>::infinity())
+	{
+		Trace::FAnalysisSessionReadScope SessionReadScope(SharedData->GetAnalysisSession());
+		const Trace::IFrameProvider& FramesProvider = Trace::ReadFrameProvider(SharedData->GetAnalysisSession());
+	
+		const Trace::FFrame* LastFrame = FramesProvider.GetFrame(ETraceFrameType::TraceFrameType_Game, FramesProvider.GetFrameCount(TraceFrameType_Game) - 1);
+		if(LastFrame)
+		{
+			SharedData->GetTimingViewSession().SetTimeMarker(LastFrame->StartTime + (double)KINDA_SMALL_NUMBER);
+		}
+	}
+
+	bPlaying = !bReverse ? true : !bPlaying;
+	bReverse = true;
+	return FReply::Handled();
+}
+
+EPlaybackMode::Type SGameplayInsightsTransportControls::GetPlaybackMode() const
+{
+	if(bPlaying)
+	{
+		return bReverse ? EPlaybackMode::PlayingReverse : EPlaybackMode::PlayingForward;
+	}
+
+	return EPlaybackMode::Stopped;
+}
+
+#undef LOCTEXT_NAMESPACE
+
+#endif
\ No newline at end of file
diff --git a/Engine/Plugins/Animation/GameplayInsights/Source/GameplayInsights/Private/SGameplayInsightsTransportControls.h b/Engine/Plugins/Animation/GameplayInsights/Source/GameplayInsights/Private/SGameplayInsightsTransportControls.h
new file mode 100644
--- /dev/null
+++ b/Engine/Plugins/Animation/GameplayInsights/Source/GameplayInsights/Private/SGameplayInsightsTransportControls.h
@@ -0,0 +1,50 @@
+// Copyright Epic Games, Inc. All Rights Reserved.
+
+#pragma once
+
+#include "CoreMinimal.h"
+#include "Widgets/SCompoundWidget.h"
+#include "Widgets/DeclarativeSyntaxSupport.h"
+#include "Input/Reply.h"
+
+#if WITH_EDITOR
+
+#include "ITransportControl.h"
+
+class FGameplaySharedData;
+
+class SGameplayInsightsTransportControls : public SCompoundWidget
+{
+public:
+	SLATE_BEGIN_ARGS(SGameplayInsightsTransportControls) {}
+
+	SLATE_END_ARGS()
+
+	void Construct(const FArguments& InArgs, FGameplaySharedData& InSharedData);
+
+private:
+	FReply OnClick_Forward_Step();
+
+	FReply OnClick_Forward_End();
+
+	FReply OnClick_Backward_Step();
+
+	FReply OnClick_Backward_End();
+
+	FReply OnClick_Forward();
+
+	FReply OnClick_Backward();
+
+	EPlaybackMode::Type GetPlaybackMode() const;
+
+private:
+	FGameplaySharedData* SharedData;
+
+	double PlayRate;
+
+	bool bPlaying;
+
+	bool bReverse;
+};
+
+#endif
\ No newline at end of file
