Insights non-editor compile CIS

diff --git a/Engine/Plugins/Animation/GameplayInsights/Source/GameplayInsights/Private/GameplayTimingViewExtender.cpp b/Engine/Plugins/Animation/GameplayInsights/Source/GameplayInsights/Private/GameplayTimingViewExtender.cpp
--- a/Engine/Plugins/Animation/GameplayInsights/Source/GameplayInsights/Private/GameplayTimingViewExtender.cpp
+++ b/Engine/Plugins/Animation/GameplayInsights/Source/GameplayInsights/Private/GameplayTimingViewExtender.cpp
@@ -27,7 +27,7 @@ void FGameplayTimingViewExtender::OnBeginSession(Insights::ITimingViewSession& I
 		PerSessionData = &PerSessionDataMap.Add(&InSession);
 		PerSessionData->GameplaySharedData = new FGameplaySharedData();
 		PerSessionData->AnimationSharedData = new FAnimationSharedData(*PerSessionData->GameplaySharedData);
-
+#if WITH_EDITOR
 		InSession.AddOverlayWidget(
 			SNew(SOverlay)
 			.Visibility(EVisibility::SelfHitTestInvisible)
@@ -42,6 +42,7 @@ void FGameplayTimingViewExtender::OnBeginSession(Insights::ITimingViewSession& I
 					SAssignNew(PerSessionData->TransportControls, SGameplayInsightsTransportControls, *PerSessionData->GameplaySharedData)
 				]
 			]);
+#endif WITH_EDITOR
 	}
 
 	PerSessionData->GameplaySharedData->OnBeginSession(InSession);
@@ -60,8 +61,9 @@ void FGameplayTimingViewExtender::OnEndSession(Insights::ITimingViewSession& InS
 		PerSessionData->GameplaySharedData = nullptr;
 		delete PerSessionData->AnimationSharedData;
 		PerSessionData->AnimationSharedData = nullptr;
-
+#if WITH_EDITOR
 		PerSessionData->TransportControls.Reset();
+#endif // WITH_EDITOR
 	}
 
 	PerSessionDataMap.Remove(&InSession);
diff --git a/Engine/Plugins/Animation/GameplayInsights/Source/GameplayInsights/Private/GameplayTimingViewExtender.h b/Engine/Plugins/Animation/GameplayInsights/Source/GameplayInsights/Private/GameplayTimingViewExtender.h
--- a/Engine/Plugins/Animation/GameplayInsights/Source/GameplayInsights/Private/GameplayTimingViewExtender.h
+++ b/Engine/Plugins/Animation/GameplayInsights/Source/GameplayInsights/Private/GameplayTimingViewExtender.h
@@ -38,7 +38,9 @@ private:
 		// Shared data
 		FGameplaySharedData* GameplaySharedData;
 		FAnimationSharedData* AnimationSharedData;
+#if WITH_EDITOR
 		TSharedPtr<SGameplayInsightsTransportControls> TransportControls;
+#endif WITH_EDITOR
 	};
 
 	// The data we host per-session
