Added feature to turn all skeletal mesh rendering on or off

Allows for nicer demos!

diff --git a/Engine/Plugins/Animation/GameplayInsights/Source/GameplayInsights/Private/AnimationSharedData.cpp b/Engine/Plugins/Animation/GameplayInsights/Source/GameplayInsights/Private/AnimationSharedData.cpp
--- a/Engine/Plugins/Animation/GameplayInsights/Source/GameplayInsights/Private/AnimationSharedData.cpp
+++ b/Engine/Plugins/Animation/GameplayInsights/Source/GameplayInsights/Private/AnimationSharedData.cpp
@@ -251,89 +251,115 @@ void FAnimationSharedData::Tick(Insights::ITimingViewSession& InTimingViewSessio
 
 void FAnimationSharedData::ExtendFilterMenu(FMenuBuilder& InMenuBuilder)
 {
-	InMenuBuilder.AddMenuEntry(
-		LOCTEXT("ToggleAnimationTracks", "Animation Tracks"),
-		LOCTEXT("ToggleAnimationTracks_Tooltip", "Show/hide all animation tracks"),
-		FSlateIcon(),
-		FUIAction(
-			FExecuteAction::CreateRaw(this, &FAnimationSharedData::ToggleAnimationTracks),
-			FCanExecuteAction(),
-			FIsActionChecked::CreateRaw(this, &FAnimationSharedData::AreAnimationTracksEnabled)),
-		NAME_None,
-		EUserInterfaceActionType::ToggleButton
-	);
-
-	InMenuBuilder.AddMenuEntry(
-		LOCTEXT("ToggleSkelMeshPoseTracks", "Pose Tracks"),
-		LOCTEXT("ToggleSkelMeshPoseTracks_Tooltip", "Show/hide the skeletal mesh pose tracks"),
-		FSlateIcon(),
-		FUIAction(
-			FExecuteAction::CreateRaw(this, &FAnimationSharedData::ToggleSkeletalMeshPoseTracks),
-			FCanExecuteAction(),
-			FIsActionChecked::CreateLambda([this](){ return bSkeletalMeshPoseTracksEnabled; })),
-		NAME_None,
-		EUserInterfaceActionType::ToggleButton
-	);
-
-	InMenuBuilder.AddMenuEntry(
-		LOCTEXT("ToggleSkelMeshCurveTracks", "Curve Tracks"),
-		LOCTEXT("ToggleSkelMeshCurveTracks_Tooltip", "Show/hide the skeletal mesh curve tracks"),
-		FSlateIcon(),
-		FUIAction(
-			FExecuteAction::CreateRaw(this, &FAnimationSharedData::ToggleSkeletalMeshCurveTracks),
-			FCanExecuteAction(),
-			FIsActionChecked::CreateLambda([this](){ return bSkeletalMeshCurveTracksEnabled; })),
-		NAME_None,
-		EUserInterfaceActionType::ToggleButton
-	);
-
-	InMenuBuilder.AddMenuEntry(
-		LOCTEXT("ToggleAnimTickRecordTracks", "Blend Weights Tracks"),
-		LOCTEXT("ToggleAnimTickRecordTracks_Tooltip", "Show/hide the blend weights (tick records) tracks"),
-		FSlateIcon(),
-		FUIAction(
-			FExecuteAction::CreateRaw(this, &FAnimationSharedData::ToggleTickRecordTracks),
-			FCanExecuteAction(),
-			FIsActionChecked::CreateLambda([this](){ return bTickRecordTracksEnabled; })),
-		NAME_None,
-		EUserInterfaceActionType::ToggleButton
-	);
-
-	InMenuBuilder.AddMenuEntry(
-		LOCTEXT("ToggleAnimNodeTracks", "Graph Tracks"),
-		LOCTEXT("ToggleAnimNodeTracks_Tooltip", "Show/hide the animation graph tracks"),
-		FSlateIcon(),
-		FUIAction(
-			FExecuteAction::CreateRaw(this, &FAnimationSharedData::ToggleAnimNodeTracks),
-			FCanExecuteAction(),
-			FIsActionChecked::CreateLambda([this](){ return bAnimNodeTracksEnabled; })),
-		NAME_None,
-		EUserInterfaceActionType::ToggleButton
-	);
-
-	InMenuBuilder.AddMenuEntry(
-		LOCTEXT("ToggleAnimNotifyTracks", "Notify Tracks"),
-		LOCTEXT("ToggleAnimNotifyTracks_Tooltip", "Show/hide the animation notify/sync marker tracks"),
-		FSlateIcon(),
-		FUIAction(
-			FExecuteAction::CreateRaw(this, &FAnimationSharedData::ToggleAnimNotifyTracks),
-			FCanExecuteAction(),
-			FIsActionChecked::CreateLambda([this](){ return bAnimNotifyTracksEnabled; })),
-		NAME_None,
-		EUserInterfaceActionType::ToggleButton
-	);
-
-	InMenuBuilder.AddMenuEntry(
-		LOCTEXT("ToggleMontageTracks", "Montage Tracks"),
-		LOCTEXT("ToggleMontageTracks_Tooltip", "Show/hide the montage tracks"),
-		FSlateIcon(),
-		FUIAction(
-			FExecuteAction::CreateRaw(this, &FAnimationSharedData::ToggleMontageTracks),
-			FCanExecuteAction(),
-			FIsActionChecked::CreateLambda([this](){ return bMontageTracksEnabled; })),
-		NAME_None,
-		EUserInterfaceActionType::ToggleButton
-	);
+	InMenuBuilder.BeginSection("AnimationTracks", LOCTEXT("AnimationHeader", "Animation"));
+	{
+#if WITH_EDITOR
+		InMenuBuilder.AddMenuEntry(
+			LOCTEXT("ShowAllMeshes", "Show All Mesh Poses"),
+			LOCTEXT("ShowAllMeshes_Tooltip", "Show all skeletal mesh poses"),
+			FSlateIcon(),
+			FUIAction(
+				FExecuteAction::CreateRaw(this, &FAnimationSharedData::ShowAllMeshes)),
+			NAME_None,
+			EUserInterfaceActionType::Button
+		);
+
+		InMenuBuilder.AddMenuEntry(
+			LOCTEXT("HideAllMeshes", "Hide All Mesh Poses"),
+			LOCTEXT("HideAllMeshes_Tooltip", "Hide all skeletal mesh poses"),
+			FSlateIcon(),
+			FUIAction(
+				FExecuteAction::CreateRaw(this, &FAnimationSharedData::HideAllMeshes)),
+			NAME_None,
+			EUserInterfaceActionType::Button
+		);
+#endif
+
+		InMenuBuilder.AddMenuEntry(
+			LOCTEXT("ToggleAnimationTracks", "Animation Tracks"),
+			LOCTEXT("ToggleAnimationTracks_Tooltip", "Show/hide all animation tracks"),
+			FSlateIcon(),
+			FUIAction(
+				FExecuteAction::CreateRaw(this, &FAnimationSharedData::ToggleAnimationTracks),
+				FCanExecuteAction(),
+				FIsActionChecked::CreateRaw(this, &FAnimationSharedData::AreAnimationTracksEnabled)),
+			NAME_None,
+			EUserInterfaceActionType::ToggleButton
+		);
+
+		InMenuBuilder.AddMenuEntry(
+			LOCTEXT("ToggleSkelMeshPoseTracks", "Pose Tracks"),
+			LOCTEXT("ToggleSkelMeshPoseTracks_Tooltip", "Show/hide the skeletal mesh pose tracks"),
+			FSlateIcon(),
+			FUIAction(
+				FExecuteAction::CreateRaw(this, &FAnimationSharedData::ToggleSkeletalMeshPoseTracks),
+				FCanExecuteAction(),
+				FIsActionChecked::CreateLambda([this](){ return bSkeletalMeshPoseTracksEnabled; })),
+			NAME_None,
+			EUserInterfaceActionType::ToggleButton
+		);
+
+		InMenuBuilder.AddMenuEntry(
+			LOCTEXT("ToggleSkelMeshCurveTracks", "Curve Tracks"),
+			LOCTEXT("ToggleSkelMeshCurveTracks_Tooltip", "Show/hide the skeletal mesh curve tracks"),
+			FSlateIcon(),
+			FUIAction(
+				FExecuteAction::CreateRaw(this, &FAnimationSharedData::ToggleSkeletalMeshCurveTracks),
+				FCanExecuteAction(),
+				FIsActionChecked::CreateLambda([this](){ return bSkeletalMeshCurveTracksEnabled; })),
+			NAME_None,
+			EUserInterfaceActionType::ToggleButton
+		);
+
+		InMenuBuilder.AddMenuEntry(
+			LOCTEXT("ToggleAnimTickRecordTracks", "Blend Weights Tracks"),
+			LOCTEXT("ToggleAnimTickRecordTracks_Tooltip", "Show/hide the blend weights (tick records) tracks"),
+			FSlateIcon(),
+			FUIAction(
+				FExecuteAction::CreateRaw(this, &FAnimationSharedData::ToggleTickRecordTracks),
+				FCanExecuteAction(),
+				FIsActionChecked::CreateLambda([this](){ return bTickRecordTracksEnabled; })),
+			NAME_None,
+			EUserInterfaceActionType::ToggleButton
+		);
+
+		InMenuBuilder.AddMenuEntry(
+			LOCTEXT("ToggleAnimNodeTracks", "Graph Tracks"),
+			LOCTEXT("ToggleAnimNodeTracks_Tooltip", "Show/hide the animation graph tracks"),
+			FSlateIcon(),
+			FUIAction(
+				FExecuteAction::CreateRaw(this, &FAnimationSharedData::ToggleAnimNodeTracks),
+				FCanExecuteAction(),
+				FIsActionChecked::CreateLambda([this](){ return bAnimNodeTracksEnabled; })),
+			NAME_None,
+			EUserInterfaceActionType::ToggleButton
+		);
+
+		InMenuBuilder.AddMenuEntry(
+			LOCTEXT("ToggleAnimNotifyTracks", "Notify Tracks"),
+			LOCTEXT("ToggleAnimNotifyTracks_Tooltip", "Show/hide the animation notify/sync marker tracks"),
+			FSlateIcon(),
+			FUIAction(
+				FExecuteAction::CreateRaw(this, &FAnimationSharedData::ToggleAnimNotifyTracks),
+				FCanExecuteAction(),
+				FIsActionChecked::CreateLambda([this](){ return bAnimNotifyTracksEnabled; })),
+			NAME_None,
+			EUserInterfaceActionType::ToggleButton
+		);
+
+		InMenuBuilder.AddMenuEntry(
+			LOCTEXT("ToggleMontageTracks", "Montage Tracks"),
+			LOCTEXT("ToggleMontageTracks_Tooltip", "Show/hide the montage tracks"),
+			FSlateIcon(),
+			FUIAction(
+				FExecuteAction::CreateRaw(this, &FAnimationSharedData::ToggleMontageTracks),
+				FCanExecuteAction(),
+				FIsActionChecked::CreateLambda([this](){ return bMontageTracksEnabled; })),
+			NAME_None,
+			EUserInterfaceActionType::ToggleButton
+		);
+	}
+	InMenuBuilder.EndSection();
 }
 
 void FAnimationSharedData::ToggleAnimationTracks()
@@ -469,7 +495,7 @@ void FAnimationSharedData::OnTimeMarkerChanged(Insights::ETimeChangedFlags InFla
 }
 
 #if WITH_EDITOR
-void FAnimationSharedData::InvalidateViewports()
+void FAnimationSharedData::InvalidateViewports() const
 {
 	UEditorEngine* Engine = Cast<UEditorEngine>(GEngine);
 	if (GIsEditor && Engine != nullptr)
@@ -562,6 +588,26 @@ void FAnimationSharedData::GetCustomDebugObjects(const IAnimationBlueprintEditor
 	}
 }
 
+void FAnimationSharedData::ShowAllMeshes()
+{
+	for(TSharedRef<FSkeletalMeshPoseTrack> PoseTrack : SkeletalMeshPoseTracks)
+	{
+		PoseTrack->SetDrawPose(true);
+	}
+
+	InvalidateViewports();
+}
+
+void FAnimationSharedData::HideAllMeshes()
+{
+	for(TSharedRef<FSkeletalMeshPoseTrack> PoseTrack : SkeletalMeshPoseTracks)
+	{
+		PoseTrack->SetDrawPose(false);
+	}
+
+	InvalidateViewports();
+}
+
 #endif
 
 void FAnimationSharedData::OpenAnimGraphTab(uint64 InAnimInstanceId) const
diff --git a/Engine/Plugins/Animation/GameplayInsights/Source/GameplayInsights/Private/AnimationSharedData.h b/Engine/Plugins/Animation/GameplayInsights/Source/GameplayInsights/Private/AnimationSharedData.h
--- a/Engine/Plugins/Animation/GameplayInsights/Source/GameplayInsights/Private/AnimationSharedData.h
+++ b/Engine/Plugins/Animation/GameplayInsights/Source/GameplayInsights/Private/AnimationSharedData.h
@@ -40,7 +40,13 @@ public:
 	void GetCustomDebugObjects(const IAnimationBlueprintEditor& InAnimationBlueprintEditor, TArray<FCustomDebugObject>& OutDebugList);
 
 	// Helper function to invalidate all viewports so non-realtime viewports update correctly.
-	void InvalidateViewports();
+	void InvalidateViewports() const;
+
+	// Show all the meshes
+	void ShowAllMeshes();
+
+	// Hide all the meshes
+	void HideAllMeshes();
 #endif
 
 	// Check whether animation tracks are enabled
diff --git a/Engine/Plugins/Animation/GameplayInsights/Source/GameplayInsights/Private/GameplaySharedData.cpp b/Engine/Plugins/Animation/GameplayInsights/Source/GameplayInsights/Private/GameplaySharedData.cpp
--- a/Engine/Plugins/Animation/GameplayInsights/Source/GameplayInsights/Private/GameplaySharedData.cpp
+++ b/Engine/Plugins/Animation/GameplayInsights/Source/GameplayInsights/Private/GameplaySharedData.cpp
@@ -148,33 +148,38 @@ void FGameplaySharedData::Tick(Insights::ITimingViewSession& InTimingViewSession
 
 void FGameplaySharedData::ExtendFilterMenu(FMenuBuilder& InMenuBuilder)
 {
-	InMenuBuilder.AddSubMenu(
-		LOCTEXT("ToggleGameplayTracks", "Gameplay Tracks"),
-		LOCTEXT("ToggleGameplayTracks_Tooltip", "Show/hide individual gameplay tracks"),
-		FNewMenuDelegate::CreateLambda([this](FMenuBuilder& InSubMenuBuilder)
-		{ 
-			InSubMenuBuilder.AddWidget(
-				SNew(SBox)
-				.MaxDesiredHeight(300.0f)
-				.MinDesiredWidth(300.0f)
-				[
-					SNew(SGameplayTrackTree, *this)
-				],
-				FText(), true);
-		})
-	);
-
-	InMenuBuilder.AddMenuEntry(
-		LOCTEXT("ToggleEventTracks", "Event Tracks"),
-		LOCTEXT("ToggleEventTracks_Tooltip", "Show/hide the gameplay event tracks"),
-		FSlateIcon(),
-		FUIAction(
-			FExecuteAction::CreateRaw(this, &FGameplaySharedData::ToggleGameplayTracks),
-			FCanExecuteAction(),
-			FIsActionChecked::CreateRaw(this, &FGameplaySharedData::AreGameplayTracksEnabled)),
-		NAME_None,
-		EUserInterfaceActionType::ToggleButton
-	);
+	InMenuBuilder.BeginSection("GameplayTracks", LOCTEXT("GameplayTracksSection", "Gameplay Tracks"));
+	{
+		InMenuBuilder.AddSubMenu(
+			LOCTEXT("ToggleGameplayTracks", "Gameplay Tracks"),
+			LOCTEXT("ToggleGameplayTracks_Tooltip", "Show/hide individual gameplay tracks"),
+			FNewMenuDelegate::CreateLambda([this](FMenuBuilder& InSubMenuBuilder)
+			{ 
+				InSubMenuBuilder.AddWidget(
+					SNew(SBox)
+					.MaxDesiredHeight(300.0f)
+					.MinDesiredWidth(300.0f)
+					.MaxDesiredWidth(300.0f)
+					[
+						SNew(SGameplayTrackTree, *this)
+					],
+					FText(), true);
+			})
+		);
+
+		InMenuBuilder.AddMenuEntry(
+			LOCTEXT("ToggleEventTracks", "Event Tracks"),
+			LOCTEXT("ToggleEventTracks_Tooltip", "Show/hide the gameplay event tracks"),
+			FSlateIcon(),
+			FUIAction(
+				FExecuteAction::CreateRaw(this, &FGameplaySharedData::ToggleGameplayTracks),
+				FCanExecuteAction(),
+				FIsActionChecked::CreateRaw(this, &FGameplaySharedData::AreGameplayTracksEnabled)),
+			NAME_None,
+			EUserInterfaceActionType::ToggleButton
+		);
+	}
+	InMenuBuilder.EndSection();
 }
 
 void FGameplaySharedData::SortTracks()
diff --git a/Engine/Plugins/Animation/GameplayInsights/Source/GameplayInsights/Private/SkeletalMeshPoseTrack.cpp b/Engine/Plugins/Animation/GameplayInsights/Source/GameplayInsights/Private/SkeletalMeshPoseTrack.cpp
--- a/Engine/Plugins/Animation/GameplayInsights/Source/GameplayInsights/Private/SkeletalMeshPoseTrack.cpp
+++ b/Engine/Plugins/Animation/GameplayInsights/Source/GameplayInsights/Private/SkeletalMeshPoseTrack.cpp
@@ -178,6 +178,7 @@ void FSkeletalMeshPoseTrack::BuildContextMenu(FMenuBuilder& MenuBuilder)
 				{
 					bDrawPose = !bDrawPose;
 					UpdateComponentVisibility();
+					SharedData.InvalidateViewports();
 				}),
 				FCanExecuteAction(),
 				FIsActionChecked::CreateLambda([this](){ return bDrawPose; })),
@@ -190,7 +191,11 @@ void FSkeletalMeshPoseTrack::BuildContextMenu(FMenuBuilder& MenuBuilder)
 			LOCTEXT("ToggleDrawSkeleton_Tooltip", "Draw the skeleton for poses in this track (when pose drawing is also enabled)"),
 			FSlateIcon(),
 			FUIAction(
-				FExecuteAction::CreateLambda([this](){ bDrawSkeleton = !bDrawSkeleton; }),
+				FExecuteAction::CreateLambda([this]()
+				{ 
+					bDrawSkeleton = !bDrawSkeleton; 
+					SharedData.InvalidateViewports();
+				}),
 				FCanExecuteAction(),
 				FIsActionChecked::CreateLambda([this](){ return bDrawSkeleton; })),
 			NAME_None,
@@ -242,6 +247,8 @@ void FSkeletalMeshPoseTrack::BuildContextMenu(FMenuBuilder& MenuBuilder)
 										InTrack->UpdateComponentVisibility();
 									}
 								});
+
+								SharedData.InvalidateViewports();
 							}),
 							FCanExecuteAction(),
 							FIsActionChecked::CreateLambda([this, ActorObjectInfo, GameplayProvider]()
@@ -292,6 +299,8 @@ void FSkeletalMeshPoseTrack::BuildContextMenu(FMenuBuilder& MenuBuilder)
 										InTrack->UpdateComponentVisibility();
 									}
 								});
+
+								SharedData.InvalidateViewports();
 							}),
 							FCanExecuteAction(),
 							FIsActionChecked::CreateLambda([this, ActorObjectInfo, GameplayProvider]()
diff --git a/Engine/Plugins/Animation/GameplayInsights/Source/GameplayInsights/Private/SkeletalMeshPoseTrack.h b/Engine/Plugins/Animation/GameplayInsights/Source/GameplayInsights/Private/SkeletalMeshPoseTrack.h
--- a/Engine/Plugins/Animation/GameplayInsights/Source/GameplayInsights/Private/SkeletalMeshPoseTrack.h
+++ b/Engine/Plugins/Animation/GameplayInsights/Source/GameplayInsights/Private/SkeletalMeshPoseTrack.h
@@ -38,6 +38,7 @@ public:
 	virtual void BuildContextMenu(FMenuBuilder& MenuBuilder) override;
 
 	// Access drawing flags
+	void SetDrawPose(bool bInDrawPose) { bDrawPose = bInDrawPose; }
 	bool ShouldDrawPose() const { return bDrawPose; }
 	bool ShouldDrawSkeleton() const { return bDrawSkeleton; }
 
