Clang build fix for CIS

#rnx

diff --git a/Engine/Plugins/Animation/GameplayInsights/Source/GameplayInsights/Private/GameplayTimingViewExtender.cpp b/Engine/Plugins/Animation/GameplayInsights/Source/GameplayInsights/Private/GameplayTimingViewExtender.cpp
--- a/Engine/Plugins/Animation/GameplayInsights/Source/GameplayInsights/Private/GameplayTimingViewExtender.cpp
+++ b/Engine/Plugins/Animation/GameplayInsights/Source/GameplayInsights/Private/GameplayTimingViewExtender.cpp
@@ -42,7 +42,7 @@ void FGameplayTimingViewExtender::OnBeginSession(Insights::ITimingViewSession& I
 					SAssignNew(PerSessionData->TransportControls, SGameplayInsightsTransportControls, *PerSessionData->GameplaySharedData)
 				]
 			]);
-#endif WITH_EDITOR
+#endif //WITH_EDITOR
 	}
 
 	PerSessionData->GameplaySharedData->OnBeginSession(InSession);
diff --git a/Engine/Plugins/Animation/GameplayInsights/Source/GameplayInsights/Private/GameplayTimingViewExtender.h b/Engine/Plugins/Animation/GameplayInsights/Source/GameplayInsights/Private/GameplayTimingViewExtender.h
--- a/Engine/Plugins/Animation/GameplayInsights/Source/GameplayInsights/Private/GameplayTimingViewExtender.h
+++ b/Engine/Plugins/Animation/GameplayInsights/Source/GameplayInsights/Private/GameplayTimingViewExtender.h
@@ -40,7 +40,7 @@ private:
 		FAnimationSharedData* AnimationSharedData;
 #if WITH_EDITOR
 		TSharedPtr<SGameplayInsightsTransportControls> TransportControls;
-#endif WITH_EDITOR
+#endif //WITH_EDITOR
 	};
 
 	// The data we host per-session
