Transport controls now scroll the view when jumping to begin/end

Scrubbing the timeline now disables playing

diff --git a/Engine/Plugins/Animation/GameplayInsights/Source/GameplayInsights/Private/GameplayTimingViewExtender.cpp b/Engine/Plugins/Animation/GameplayInsights/Source/GameplayInsights/Private/GameplayTimingViewExtender.cpp
--- a/Engine/Plugins/Animation/GameplayInsights/Source/GameplayInsights/Private/GameplayTimingViewExtender.cpp
+++ b/Engine/Plugins/Animation/GameplayInsights/Source/GameplayInsights/Private/GameplayTimingViewExtender.cpp
@@ -27,6 +27,10 @@ void FGameplayTimingViewExtender::OnBeginSession(Insights::ITimingViewSession& I
 		PerSessionData = &PerSessionDataMap.Add(&InSession);
 		PerSessionData->GameplaySharedData = new FGameplaySharedData();
 		PerSessionData->AnimationSharedData = new FAnimationSharedData(*PerSessionData->GameplaySharedData);
+
+		PerSessionData->GameplaySharedData->OnBeginSession(InSession);
+		PerSessionData->AnimationSharedData->OnBeginSession(InSession);
+
 #if WITH_EDITOR
 		InSession.AddOverlayWidget(
 			SNew(SOverlay)
@@ -44,9 +48,11 @@ void FGameplayTimingViewExtender::OnBeginSession(Insights::ITimingViewSession& I
 			]);
 #endif
 	}
-
-	PerSessionData->GameplaySharedData->OnBeginSession(InSession);
-	PerSessionData->AnimationSharedData->OnBeginSession(InSession);
+	else
+	{
+		PerSessionData->GameplaySharedData->OnBeginSession(InSession);
+		PerSessionData->AnimationSharedData->OnBeginSession(InSession);
+	}
 }
 
 void FGameplayTimingViewExtender::OnEndSession(Insights::ITimingViewSession& InSession)
diff --git a/Engine/Plugins/Animation/GameplayInsights/Source/GameplayInsights/Private/SGameplayInsightsTransportControls.cpp b/Engine/Plugins/Animation/GameplayInsights/Source/GameplayInsights/Private/SGameplayInsightsTransportControls.cpp
--- a/Engine/Plugins/Animation/GameplayInsights/Source/GameplayInsights/Private/SGameplayInsightsTransportControls.cpp
+++ b/Engine/Plugins/Animation/GameplayInsights/Source/GameplayInsights/Private/SGameplayInsightsTransportControls.cpp
@@ -26,6 +26,9 @@ void SGameplayInsightsTransportControls::Construct(const FArguments& InArgs, FGa
 	PlayRate = 1.0;
 	bPlaying = false;
 	bReverse = false;
+	bSettingMarker = false;
+
+	SharedData->GetTimingViewSession().OnTimeMarkerChanged().AddSP(this, &SGameplayInsightsTransportControls::HandleTimeMarkerChanged);
 
 	FEditorWidgetsModule& EditorWidgetsModule = FModuleManager::LoadModuleChecked<FEditorWidgetsModule>("EditorWidgets");
 
@@ -71,7 +74,7 @@ void SGameplayInsightsTransportControls::Construct(const FArguments& InArgs, FGa
 		{
 			double CurrentTime = SharedData->GetTimingViewSession().GetTimeMarker();
 			double Delta = (double)(bReverse ? -InDeltaTime : InDeltaTime);
-			SharedData->GetTimingViewSession().SetTimeMarker(CurrentTime + (Delta * PlayRate));
+			SetTimeMarker(CurrentTime + (Delta * PlayRate), false);
 		}
 
 		return EActiveTimerReturnType::Continue;
@@ -101,7 +104,7 @@ FReply SGameplayInsightsTransportControls::OnClick_Forward_Step()
 			const Trace::FFrame* NextFrame = FramesProvider.GetFrame(ETraceFrameType::TraceFrameType_Game, Frame.Index + 1);
 			if(NextFrame)
 			{
-				SharedData->GetTimingViewSession().SetTimeMarker(NextFrame->StartTime + (double)KINDA_SMALL_NUMBER);
+				SetTimeMarker(NextFrame->StartTime + (double)KINDA_SMALL_NUMBER, false);
 			}
 		}
 	}
@@ -122,7 +125,7 @@ FReply SGameplayInsightsTransportControls::OnClick_Forward_End()
 		const Trace::FFrame* LastFrame = FramesProvider.GetFrame(ETraceFrameType::TraceFrameType_Game, FramesProvider.GetFrameCount(ETraceFrameType::TraceFrameType_Game) - 1);
 		if(LastFrame)
 		{
-			SharedData->GetTimingViewSession().SetTimeMarker(LastFrame->StartTime + (double)KINDA_SMALL_NUMBER);
+			SetTimeMarker(LastFrame->StartTime + (double)KINDA_SMALL_NUMBER, true);
 		}
 	}
 
@@ -155,7 +158,7 @@ FReply SGameplayInsightsTransportControls::OnClick_Backward_Step()
 			const Trace::FFrame* PrevFrame = FramesProvider.GetFrame(ETraceFrameType::TraceFrameType_Game, Frame.Index - 1);
 			if(PrevFrame)
 			{
-				SharedData->GetTimingViewSession().SetTimeMarker(PrevFrame->StartTime + (double)KINDA_SMALL_NUMBER);
+				SetTimeMarker(PrevFrame->StartTime + (double)KINDA_SMALL_NUMBER, false);
 			}
 		}
 	}
@@ -174,7 +177,7 @@ FReply SGameplayInsightsTransportControls::OnClick_Backward_End()
 	const Trace::FFrame* FirstFrame = FramesProvider.GetFrame(ETraceFrameType::TraceFrameType_Game, 0);
 	if(FirstFrame)
 	{
-		SharedData->GetTimingViewSession().SetTimeMarker(FirstFrame->StartTime + (double)KINDA_SMALL_NUMBER);
+		SetTimeMarker(FirstFrame->StartTime + (double)KINDA_SMALL_NUMBER, true);
 	}
 
 	bPlaying = false;
@@ -194,7 +197,7 @@ FReply SGameplayInsightsTransportControls::OnClick_Forward()
 		const Trace::FFrame* FirstFrame = FramesProvider.GetFrame(ETraceFrameType::TraceFrameType_Game, 0);
 		if(FirstFrame)
 		{
-			SharedData->GetTimingViewSession().SetTimeMarker(FirstFrame->StartTime + (double)KINDA_SMALL_NUMBER);
+			SetTimeMarker(FirstFrame->StartTime + (double)KINDA_SMALL_NUMBER, false);
 		}
 	}
 
@@ -214,7 +217,7 @@ FReply SGameplayInsightsTransportControls::OnClick_Backward()
 		const Trace::FFrame* LastFrame = FramesProvider.GetFrame(ETraceFrameType::TraceFrameType_Game, FramesProvider.GetFrameCount(TraceFrameType_Game) - 1);
 		if(LastFrame)
 		{
-			SharedData->GetTimingViewSession().SetTimeMarker(LastFrame->StartTime + (double)KINDA_SMALL_NUMBER);
+			SetTimeMarker(LastFrame->StartTime + (double)KINDA_SMALL_NUMBER, false);
 		}
 	}
 
@@ -233,6 +236,29 @@ EPlaybackMode::Type SGameplayInsightsTransportControls::GetPlaybackMode() const
 	return EPlaybackMode::Stopped;
 }
 
+void SGameplayInsightsTransportControls::SetTimeMarker(double InTime, bool bInScroll)
+{
+	bSettingMarker = true;
+	if(bInScroll)
+	{
+		SharedData->GetTimingViewSession().SetAndCenterOnTimeMarker(InTime);
+	}
+	else
+	{
+		SharedData->GetTimingViewSession().SetTimeMarker(InTime);
+	}
+	bSettingMarker = false;
+}
+
+void SGameplayInsightsTransportControls::HandleTimeMarkerChanged(Insights::ETimeChangedFlags InFlags, double InTimeMarker)
+{
+	if(!bSettingMarker)
+	{
+		// turn off playback if someone else scrubbed the timeline
+		bPlaying = false;
+	}
+}
+
 #undef LOCTEXT_NAMESPACE
 
 #endif
\ No newline at end of file
diff --git a/Engine/Plugins/Animation/GameplayInsights/Source/GameplayInsights/Private/SGameplayInsightsTransportControls.h b/Engine/Plugins/Animation/GameplayInsights/Source/GameplayInsights/Private/SGameplayInsightsTransportControls.h
--- a/Engine/Plugins/Animation/GameplayInsights/Source/GameplayInsights/Private/SGameplayInsightsTransportControls.h
+++ b/Engine/Plugins/Animation/GameplayInsights/Source/GameplayInsights/Private/SGameplayInsightsTransportControls.h
@@ -12,6 +12,7 @@
 #include "ITransportControl.h"
 
 class FGameplaySharedData;
+namespace Insights { enum class ETimeChangedFlags; }
 
 class SGameplayInsightsTransportControls : public SCompoundWidget
 {
@@ -37,6 +38,10 @@ private:
 
 	EPlaybackMode::Type GetPlaybackMode() const;
 
+	void SetTimeMarker(double InTime, bool bInScroll);
+
+	void HandleTimeMarkerChanged(Insights::ETimeChangedFlags InFlags, double InTimeMarker);
+
 private:
 	FGameplaySharedData* SharedData;
 
@@ -45,6 +50,8 @@ private:
 	bool bPlaying;
 
 	bool bReverse;
+
+	bool bSettingMarker;
 };
 
 #endif
\ No newline at end of file
