Added undetermined state to gameplay track tree entries

diff --git a/Engine/Plugins/Animation/GameplayInsights/Source/GameplayInsights/Private/SGameplayTrackTree.cpp b/Engine/Plugins/Animation/GameplayInsights/Source/GameplayInsights/Private/SGameplayTrackTree.cpp
--- a/Engine/Plugins/Animation/GameplayInsights/Source/GameplayInsights/Private/SGameplayTrackTree.cpp
+++ b/Engine/Plugins/Animation/GameplayInsights/Source/GameplayInsights/Private/SGameplayTrackTree.cpp
@@ -23,11 +23,13 @@ enum class EGameplayTrackFilterState
 };
 
 // Simple wrapper around a gameplay track for filtering
-struct FGameplayTrackTreeEntry
+struct FGameplayTrackTreeEntry : public TSharedFromThis<FGameplayTrackTreeEntry>
 {
 	FGameplayTrackTreeEntry(TSharedRef<FBaseTimingTrack> InTimingTrack)
 		: WeakTimingTrack(InTimingTrack)
 		, FilterState(EGameplayTrackFilterState::Hidden)
+		, bIsAnyChildVisible(false)
+		, bAreAllChildrenVisible(false)
 	{}
 
 	FText GetName() const
@@ -43,6 +45,20 @@ struct FGameplayTrackTreeEntry
 		return Name;
 	}
 
+	ECheckBoxState GetCheckBoxState() const
+	{
+		if(IsVisible())
+		{
+			return ECheckBoxState::Checked;
+		}
+		else if(bIsAnyChildVisible)
+		{
+			return ECheckBoxState::Undetermined;
+		}
+	
+		return ECheckBoxState::Unchecked;
+	}
+
 	bool IsVisible() const
 	{
 		TSharedPtr<FBaseTimingTrack> TimingTrack = WeakTimingTrack.Pin();
@@ -71,9 +87,37 @@ struct FGameplayTrackTreeEntry
 		}
 	}
 
+	void CacheVisibility()
+	{
+		CacheVisibility_Helper(bIsAnyChildVisible, bAreAllChildrenVisible);
+	}
+
+private:
+
+	void CacheVisibility_Helper(bool &bOutAnyChildVisible, bool& bOutAllChildrenVisible)
+	{
+		bIsAnyChildVisible = false;
+		bAreAllChildrenVisible = Children.Num() > 0;
+
+		for(const TSharedRef<FGameplayTrackTreeEntry>& Child : Children)
+		{
+			bIsAnyChildVisible |= Child->IsVisible();
+			bAreAllChildrenVisible &= Child->IsVisible();
+
+			Child->CacheVisibility_Helper(Child->bIsAnyChildVisible, Child->bAreAllChildrenVisible);
+
+			bOutAnyChildVisible |= Child->bIsAnyChildVisible;
+			bOutAllChildrenVisible &=  Child->bAreAllChildrenVisible;
+		}
+	}
+
+public:
+
 	TWeakPtr<FBaseTimingTrack> WeakTimingTrack;
 	TArray<TSharedRef<FGameplayTrackTreeEntry>> Children;
 	EGameplayTrackFilterState FilterState;
+	bool bIsAnyChildVisible;
+	bool bAreAllChildrenVisible;
 };
 
 // A list entry widget for a gameplay track
@@ -88,10 +132,11 @@ public:
 
 	SLATE_END_ARGS()
 
-	void Construct(const FArguments& InArgs)
+	void Construct(const FArguments& InArgs, const TSharedRef<SGameplayTrackTree>& InTree)
 	{
 		WeakTreeEntry = InArgs._TreeEntry;
 		SearchText = InArgs._SearchText;
+		Tree = InTree;
 
 		ChildSlot
 		[
@@ -101,7 +146,7 @@ public:
 				TSharedPtr<FGameplayTrackTreeEntry> TreeEntry = WeakTreeEntry.Pin();
 				if(TreeEntry.IsValid())
 				{
-					return TreeEntry->IsVisible() ? ECheckBoxState::Checked : ECheckBoxState::Unchecked; 
+					return TreeEntry->GetCheckBoxState(); 
 				}
 				return ECheckBoxState::Unchecked; 
 			})
@@ -111,6 +156,8 @@ public:
 				if(TreeEntry.IsValid())
 				{
 					TreeEntry->SetVisibilityFlag(InCheckBoxState == ECheckBoxState::Checked);
+
+					Tree.Pin()->CacheVisibility();
 				}
 			})
 			.Content()
@@ -125,6 +172,9 @@ public:
 	// The tree entry we represent
 	TWeakPtr<FGameplayTrackTreeEntry> WeakTreeEntry;
 
+	// The tree we are contained in
+	TWeakPtr<SGameplayTrackTree> Tree;
+
 	// The search text to highlight
 	TAttribute<FText> SearchText;
 };
@@ -184,6 +234,8 @@ void SGameplayTrackTree::Construct(const FArguments& InArgs, FGameplaySharedData
 					{
 						Track->SetVisibilityFlag(bVisible);
 					}
+
+					CacheVisibility();
 				})
 			]
 			+ SHorizontalBox::Slot()
@@ -220,8 +272,9 @@ TSharedRef<ITableRow> SGameplayTrackTree::OnGenerateRow(TSharedRef<FGameplayTrac
 {
 	return
 		SNew(STableRow<TSharedRef<FGameplayTrackTreeEntry>>, InOwnerTable)
+		.ToolTipText(InItem->GetName())
 		[
-			SNew(SGameplayTrackTreeEntry)
+			SNew(SGameplayTrackTreeEntry, SharedThis(this))
 			.TreeEntry(InItem)
 			.SearchText_Lambda([this](){ return SearchText; })
 		];
@@ -258,6 +311,9 @@ EGameplayTrackFilterState SGameplayTrackTree::RefreshFilter_Helper(const TShared
 		{
 			TSharedRef<FGameplayTrackTreeEntry> NewTreeEntry = MakeShared<FGameplayTrackTreeEntry>(ChildTrack->GetTimingTrack());
 			InTreeEntry->Children.Add(NewTreeEntry);
+
+			TreeView->SetItemExpansion(NewTreeEntry, true);
+
 			EGameplayTrackFilterState ChildFilterState = RefreshFilter_Helper(NewTreeEntry);
 			InTreeEntry->FilterState = FMath::Max(ChildFilterState, InTreeEntry->FilterState);
 		}
@@ -290,18 +346,31 @@ EGameplayTrackFilterState SGameplayTrackTree::RefreshFilter_Helper(const TShared
 void SGameplayTrackTree::RefreshFilter()
 {
 	FilteredTracks.Reset();
+	RootEntries.Reset();
 
 	for(const TSharedRef<FBaseTimingTrack>& RootTrack : SharedData->GetRootTracks())
 	{
 		TSharedRef<FGameplayTrackTreeEntry> NewTreeEntry = MakeShared<FGameplayTrackTreeEntry>(RootTrack);
+		RootEntries.Add(NewTreeEntry);
 		EGameplayTrackFilterState FilterState = RefreshFilter_Helper(NewTreeEntry);
 		if(FilterState != EGameplayTrackFilterState::Hidden)
 		{
-			FilteredTracks.Add(NewTreeEntry);	
+			FilteredTracks.Add(NewTreeEntry);
+			TreeView->SetItemExpansion(NewTreeEntry, true);
 		}
 	}
 
+	CacheVisibility();
+
 	TreeView->RequestTreeRefresh();
 }
 
+void SGameplayTrackTree::CacheVisibility()
+{
+	for(const TSharedRef<FGameplayTrackTreeEntry>& RootEntry : RootEntries)
+	{
+		RootEntry->CacheVisibility();
+	}
+}
+
 #undef LOCTEXT_NAMESPACE
diff --git a/Engine/Plugins/Animation/GameplayInsights/Source/GameplayInsights/Private/SGameplayTrackTree.h b/Engine/Plugins/Animation/GameplayInsights/Source/GameplayInsights/Private/SGameplayTrackTree.h
--- a/Engine/Plugins/Animation/GameplayInsights/Source/GameplayInsights/Private/SGameplayTrackTree.h
+++ b/Engine/Plugins/Animation/GameplayInsights/Source/GameplayInsights/Private/SGameplayTrackTree.h
@@ -22,6 +22,8 @@ public:
 
 	void Construct(const FArguments& InArgs, FGameplaySharedData& InSharedData);
 
+	void CacheVisibility();
+
 private:
 	// Generate a row for the tree
 	TSharedRef<ITableRow> OnGenerateRow(TSharedRef<FGameplayTrackTreeEntry> Item, const TSharedRef<STableViewBase>& OwnerTable);
@@ -53,4 +55,7 @@ private:
 
 	// Filtered list of tracks
 	TArray<TSharedRef<FGameplayTrackTreeEntry>> FilteredTracks;
+
+	// Unfiltered list of tracks
+	TArray<TSharedRef<FGameplayTrackTreeEntry>> RootEntries;
 };
