Fix crash when an anim BP class could not be loaded

Also fixes crash opening a child anim BP

diff --git a/Engine/Plugins/Animation/GameplayInsights/Source/GameplayInsights/Private/AnimNodesTrack.cpp b/Engine/Plugins/Animation/GameplayInsights/Source/GameplayInsights/Private/AnimNodesTrack.cpp
--- a/Engine/Plugins/Animation/GameplayInsights/Source/GameplayInsights/Private/AnimNodesTrack.cpp
+++ b/Engine/Plugins/Animation/GameplayInsights/Source/GameplayInsights/Private/AnimNodesTrack.cpp
@@ -203,7 +203,10 @@ void FAnimNodesTrack::BuildContextMenu(FMenuBuilder& MenuBuilder)
 										if(USkeletalMeshComponent* SkeletalMeshComponent = SkeletalMeshPoseTrack->GetComponent(FGameplayTimingViewExtender::GetWorldToVisualize()))
 										{
 											UAnimInstance* Instance = LazyCreateAnimInstance(SkeletalMeshComponent);
-											AnimBlueprint->SetObjectBeingDebugged(Instance);
+											if(Instance)
+											{
+												AnimBlueprint->SetObjectBeingDebugged(Instance);
+											}
 										}
 									}
 								}
@@ -232,7 +235,10 @@ void FAnimNodesTrack::BuildContextMenu(FMenuBuilder& MenuBuilder)
 										if(USkeletalMeshComponent* SkeletalMeshComponent = SkeletalMeshPoseTrack->GetComponent(FGameplayTimingViewExtender::GetWorldToVisualize()))
 										{
 											UAnimInstance* Instance = LazyCreateAnimInstance(SkeletalMeshComponent);
-											AnimBlueprint->IsObjectBeingDebugged(Instance);
+											if(Instance)
+											{
+												AnimBlueprint->IsObjectBeingDebugged(Instance);
+											}
 										}
 									}
 								}
@@ -389,7 +395,10 @@ void FAnimNodesTrack::GetCustomDebugObjects(const IAnimationBlueprintEditor& InA
 						if(USkeletalMeshComponent* SkeletalMeshComponent = SkeletalMeshPoseTrack->GetComponent(FGameplayTimingViewExtender::GetWorldToVisualize()))
 						{
 							UAnimInstance* Instance = LazyCreateAnimInstance(SkeletalMeshComponent);
-							OutDebugList.Emplace(Instance, FText::Format(LOCTEXT("PreviewObjectLabel", "Insights - {0}"), FText::FromString(GetName())).ToString());
+							if(Instance)
+							{
+								OutDebugList.Emplace(Instance, FText::Format(LOCTEXT("PreviewObjectLabel", "Insights - {0}"), FText::FromString(GetName())).ToString());
+							}
 						}
 					}
 				}
diff --git a/Engine/Plugins/Animation/GameplayInsights/Source/GameplayInsights/Private/AnimationTickRecordsTrack.cpp b/Engine/Plugins/Animation/GameplayInsights/Source/GameplayInsights/Private/AnimationTickRecordsTrack.cpp
--- a/Engine/Plugins/Animation/GameplayInsights/Source/GameplayInsights/Private/AnimationTickRecordsTrack.cpp
+++ b/Engine/Plugins/Animation/GameplayInsights/Source/GameplayInsights/Private/AnimationTickRecordsTrack.cpp
@@ -383,10 +383,10 @@ void FAnimationTickRecordsTrack::BuildContextMenu(FMenuBuilder& MenuBuilder)
 							if(IAnimationBlueprintEditor* AnimBlueprintEditor = static_cast<IAnimationBlueprintEditor*>(GEditor->GetEditorSubsystem<UAssetEditorSubsystem>()->FindEditorForAsset(AnimBlueprint, true)))
 							{
 								int32 AnimNodeIndex = InstanceClass.Get()->AnimNodeProperties.Num() - NodeId - 1;
-								TWeakObjectPtr<const UEdGraphNode> GraphNode = InstanceClass.Get()->AnimBlueprintDebugData.NodePropertyIndexToNodeMap[AnimNodeIndex];
-								if(GraphNode.Get())
+								TWeakObjectPtr<const UEdGraphNode>* GraphNode = InstanceClass.Get()->AnimBlueprintDebugData.NodePropertyIndexToNodeMap.Find(AnimNodeIndex);
+								if(GraphNode != nullptr && GraphNode->Get())
 								{
-									AnimBlueprintEditor->JumpToHyperlink(GraphNode.Get());
+									AnimBlueprintEditor->JumpToHyperlink(GraphNode->Get());
 								}
 							}
 						}
