Trace Data Filtering and Gameplay Insights plugins do not share a Plugin category, Plugin Dependency, or very many keywords making it hard to discover and enable both
#fix Make GameplayInsights, TraceDataFilter and TraceSourceFilter share a plugin category

diff --git a/Engine/Plugins/Animation/GameplayInsights/GameplayInsights.uplugin b/Engine/Plugins/Animation/GameplayInsights/GameplayInsights.uplugin
--- a/Engine/Plugins/Animation/GameplayInsights/GameplayInsights.uplugin
+++ b/Engine/Plugins/Animation/GameplayInsights/GameplayInsights.uplugin
@@ -4,7 +4,7 @@
 	"VersionName": "1.0",
 	"FriendlyName": "Gameplay Insights",
 	"Description": "Allows debugging of gameplay and animation systems via Unreal Insights",
-	"Category": "GameplayInsights",
+	"Category": "Insights",
 	"CreatedBy": "Epic Games, Inc.",
 	"CreatedByURL": "http://epicgames.com",
 	"DocsURL": "",
diff --git a/Engine/Plugins/Developer/TraceDataFiltering/TraceDataFilters.uplugin b/Engine/Plugins/Developer/TraceDataFiltering/TraceDataFilters.uplugin
--- a/Engine/Plugins/Developer/TraceDataFiltering/TraceDataFilters.uplugin
+++ b/Engine/Plugins/Developer/TraceDataFiltering/TraceDataFilters.uplugin
@@ -4,7 +4,7 @@
 	"VersionName": "1.0",
 	"FriendlyName": "Trace Data Filtering",
 	"Description": "Allows for turning on/off individual or sets of Trace Channels.",
-	"Category": "Editor",
+	"Category": "Insights",
 	"CreatedBy": "Epic Games, Inc.",
 	"CreatedByURL": "http://epicgames.com",
 	"DocsURL": "",
diff --git a/Engine/Plugins/Developer/TraceSourceFiltering/TraceSourceFilters.uplugin b/Engine/Plugins/Developer/TraceSourceFiltering/TraceSourceFilters.uplugin
--- a/Engine/Plugins/Developer/TraceSourceFiltering/TraceSourceFilters.uplugin
+++ b/Engine/Plugins/Developer/TraceSourceFiltering/TraceSourceFilters.uplugin
@@ -4,7 +4,7 @@
 	"VersionName": "1.0",
 	"FriendlyName": "Insights Data Source Filters",
 	"Description": "Source data filtering for Unreal Insights.",
-	"Category": "Editor",
+	"Category": "Insights",
 	"CreatedBy": "Epic Games, Inc.",
 	"CreatedByURL": "http://epicgames.com",
 	"DocsURL": "",
