Gameplay Insights release fixes

Added UWorld information so data can be differentiated more easily.
Collapsed anim sequence player tracks into a single track to make better use of space and to reduce track-explosion.
Fixed world tick so that contiguous curves are shown correctly (this uses a world subsystem to keep track of the workds' actual tick count).
Fix duplicate values being reported by TPointTimeline iterations if bEnumerateOutsideRange was set.
Fix FPS drops when viewing pose and notify tracks (removed Snprintf per-event etc.).
Fixed frame rounding and removed duplicate calls to the frame provider.
Fixed notify event names.
Gameplay graph tracks now expand without animating, as graph tracks recompute their points snapped Y values each time they resize, this becomes very expensive.

diff --git a/Engine/Plugins/Animation/GameplayInsights/Source/GameplayInsights/Private/AnimNodesTrack.cpp b/Engine/Plugins/Animation/GameplayInsights/Source/GameplayInsights/Private/AnimNodesTrack.cpp
--- a/Engine/Plugins/Animation/GameplayInsights/Source/GameplayInsights/Private/AnimNodesTrack.cpp
+++ b/Engine/Plugins/Animation/GameplayInsights/Source/GameplayInsights/Private/AnimNodesTrack.cpp
@@ -138,6 +138,7 @@ void FAnimNodesTrack::InitTooltip(FTooltipDrawState& Tooltip, const ITimingEvent
 
 		Tooltip.AddNameValueTextLine(LOCTEXT("EventDuration", "Duration").ToString(), TimeUtils::FormatTimeAuto(InFoundEndTime - InFoundStartTime));
 		Tooltip.AddNameValueTextLine(LOCTEXT("EventTime", "Time").ToString(), TimeUtils::FormatTimeAuto(InFoundStartTime));
+		Tooltip.AddNameValueTextLine(LOCTEXT("EventWorld", "World").ToString(), GetGameplayTrack().GetWorldName(SharedData.GetAnalysisSession()).ToString());
 
 		Tooltip.UpdateLayout();
 	});
@@ -297,7 +298,7 @@ UAnimInstance* FAnimNodesTrack::LazyCreateAnimInstance(USkeletalMeshComponent* I
 	return nullptr;
 }
 
-void FAnimNodesTrack::UpdateDebugData(double InTime)
+void FAnimNodesTrack::UpdateDebugData(const Trace::FFrame& InFrame)
 {
 	if(InstanceClass.LoadSynchronous())
 	{
@@ -308,81 +309,73 @@ void FAnimNodesTrack::UpdateDebugData(double InTime)
 		{
 			Trace::FAnalysisSessionReadScope SessionReadScope(SharedData.GetAnalysisSession());
 
-			AnimationProvider->ReadAnimGraphTimeline(GetGameplayTrack().GetObjectId(), [this, AnimationProvider, GameplayProvider, InTime](const FAnimationProvider::AnimGraphTimeline& InGraphTimeline)
+			AnimationProvider->ReadAnimGraphTimeline(GetGameplayTrack().GetObjectId(), [this, AnimationProvider, GameplayProvider, InFrame](const FAnimationProvider::AnimGraphTimeline& InGraphTimeline)
 			{
-				const Trace::IFrameProvider& FramesProvider = Trace::ReadFrameProvider(SharedData.GetAnalysisSession());
-
-
 				FAnimBlueprintDebugData& DebugData = InstanceClass.Get()->GetAnimBlueprintDebugData();
 				DebugData.ResetNodeVisitSites();
 
-				// round to nearest frame boundary
-				Trace::FFrame Frame;
-				if(FramesProvider.GetFrameFromTime(ETraceFrameType::TraceFrameType_Game, InTime, Frame))
+				InGraphTimeline.EnumerateEvents(InFrame.StartTime, InFrame.EndTime, [this, AnimationProvider, GameplayProvider, &DebugData](double InGraphStartTime, double InGraphEndTime, uint32 InDepth, const FAnimGraphMessage& InMessage)
 				{
-					InGraphTimeline.EnumerateEvents(Frame.StartTime, Frame.EndTime, [this, AnimationProvider, GameplayProvider, &DebugData](double InGraphStartTime, double InGraphEndTime, uint32 InDepth, const FAnimGraphMessage& InMessage)
+					// Check for an update phase (which contains weights)
+					if(InMessage.Phase == EAnimGraphPhase::Update)
 					{
-						// Check for an update phase (which contains weights)
-						if(InMessage.Phase == EAnimGraphPhase::Update)
+						// Basic verification - check node count is the same
+						// @TODO: could add some form of node hash/CRC to the class to improve this
+						if(InMessage.NodeCount == InstanceClass.Get()->AnimNodeProperties.Num())
 						{
-							// Basic verification - check node count is the same
-							// @TODO: could add some form of node hash/CRC to the class to improve this
-							if(InMessage.NodeCount == InstanceClass.Get()->AnimNodeProperties.Num())
+							AnimationProvider->ReadAnimNodesTimeline(GetGameplayTrack().GetObjectId(), [InGraphStartTime, InGraphEndTime, &DebugData](const FAnimationProvider::AnimNodesTimeline& InNodesTimeline)
 							{
-								AnimationProvider->ReadAnimNodesTimeline(GetGameplayTrack().GetObjectId(), [InGraphStartTime, InGraphEndTime, &DebugData](const FAnimationProvider::AnimNodesTimeline& InNodesTimeline)
+								InNodesTimeline.EnumerateEvents(InGraphStartTime, InGraphEndTime, [&DebugData](double InStartTime, double InEndTime, uint32 InDepth, const FAnimNodeMessage& InMessage)
 								{
-									InNodesTimeline.EnumerateEvents(InGraphStartTime, InGraphEndTime, [&DebugData](double InStartTime, double InEndTime, uint32 InDepth, const FAnimNodeMessage& InMessage)
+									if(InMessage.Phase == EAnimGraphPhase::Update)
 									{
-										if(InMessage.Phase == EAnimGraphPhase::Update)
-										{
-											DebugData.RecordNodeVisit(InMessage.NodeId, InMessage.PreviousNodeId, InMessage.Weight);
-										}
-									});
+										DebugData.RecordNodeVisit(InMessage.NodeId, InMessage.PreviousNodeId, InMessage.Weight);
+									}
 								});
+							});
 
-								AnimationProvider->ReadStateMachinesTimeline(GetGameplayTrack().GetObjectId(), [InGraphStartTime, InGraphEndTime, &DebugData](const FAnimationProvider::StateMachinesTimeline& InStateMachinesTimeline)
+							AnimationProvider->ReadStateMachinesTimeline(GetGameplayTrack().GetObjectId(), [InGraphStartTime, InGraphEndTime, &DebugData](const FAnimationProvider::StateMachinesTimeline& InStateMachinesTimeline)
+							{
+								InStateMachinesTimeline.EnumerateEvents(InGraphStartTime, InGraphEndTime, [&DebugData](double InStartTime, double InEndTime, uint32 InDepth, const FAnimStateMachineMessage& InMessage)
 								{
-									InStateMachinesTimeline.EnumerateEvents(InGraphStartTime, InGraphEndTime, [&DebugData](double InStartTime, double InEndTime, uint32 InDepth, const FAnimStateMachineMessage& InMessage)
-									{
-										DebugData.RecordStateData(InMessage.StateMachineIndex, InMessage.StateIndex, InMessage.StateWeight, InMessage.ElapsedTime);
-									});
+									DebugData.RecordStateData(InMessage.StateMachineIndex, InMessage.StateIndex, InMessage.StateWeight, InMessage.ElapsedTime);
 								});
+							});
 
-								AnimationProvider->ReadAnimNodeValuesTimeline(GetGameplayTrack().GetObjectId(), [InGraphStartTime, InGraphEndTime, AnimationProvider, &DebugData](const FAnimationProvider::AnimNodeValuesTimeline& InNodeValuesTimeline)
+							AnimationProvider->ReadAnimNodeValuesTimeline(GetGameplayTrack().GetObjectId(), [InGraphStartTime, InGraphEndTime, AnimationProvider, &DebugData](const FAnimationProvider::AnimNodeValuesTimeline& InNodeValuesTimeline)
+							{
+								InNodeValuesTimeline.EnumerateEvents(InGraphStartTime, InGraphEndTime, [AnimationProvider, &DebugData](double InStartTime, double InEndTime, uint32 InDepth, const FAnimNodeValueMessage& InMessage)
 								{
-									InNodeValuesTimeline.EnumerateEvents(InGraphStartTime, InGraphEndTime, [AnimationProvider, &DebugData](double InStartTime, double InEndTime, uint32 InDepth, const FAnimNodeValueMessage& InMessage)
-									{
-										FText Text = AnimationProvider->FormatNodeKeyValue(InMessage);
-										DebugData.RecordNodeValue(InMessage.NodeId, Text.ToString());
-									});
+									FText Text = AnimationProvider->FormatNodeKeyValue(InMessage);
+									DebugData.RecordNodeValue(InMessage.NodeId, Text.ToString());
 								});
+							});
 
-								AnimationProvider->ReadAnimSequencePlayersTimeline(GetGameplayTrack().GetObjectId(), [InGraphStartTime, InGraphEndTime, GameplayProvider, &DebugData](const FAnimationProvider::AnimSequencePlayersTimeline& InSequencePlayersTimeline)
+							AnimationProvider->ReadAnimSequencePlayersTimeline(GetGameplayTrack().GetObjectId(), [InGraphStartTime, InGraphEndTime, GameplayProvider, &DebugData](const FAnimationProvider::AnimSequencePlayersTimeline& InSequencePlayersTimeline)
+							{
+								InSequencePlayersTimeline.EnumerateEvents(InGraphStartTime, InGraphEndTime, [&DebugData](double InStartTime, double InEndTime, uint32 InDepth, const FAnimSequencePlayerMessage& InMessage)
 								{
-									InSequencePlayersTimeline.EnumerateEvents(InGraphStartTime, InGraphEndTime, [&DebugData](double InStartTime, double InEndTime, uint32 InDepth, const FAnimSequencePlayerMessage& InMessage)
-									{
-										DebugData.RecordSequencePlayer(InMessage.NodeId, InMessage.Position, InMessage.Length, InMessage.FrameCount);
-									});
+									DebugData.RecordSequencePlayer(InMessage.NodeId, InMessage.Position, InMessage.Length, InMessage.FrameCounter);
 								});
+							});
 
-								AnimationProvider->ReadAnimBlendSpacePlayersTimeline(GetGameplayTrack().GetObjectId(), [InGraphStartTime, InGraphEndTime, GameplayProvider, &DebugData](const FAnimationProvider::BlendSpacePlayersTimeline& InBlendSpacePlayersTimeline)
+							AnimationProvider->ReadAnimBlendSpacePlayersTimeline(GetGameplayTrack().GetObjectId(), [InGraphStartTime, InGraphEndTime, GameplayProvider, &DebugData](const FAnimationProvider::BlendSpacePlayersTimeline& InBlendSpacePlayersTimeline)
+							{
+								InBlendSpacePlayersTimeline.EnumerateEvents(InGraphStartTime, InGraphEndTime, [GameplayProvider, &DebugData](double InStartTime, double InEndTime, uint32 InDepth, const FBlendSpacePlayerMessage& InMessage)
 								{
-									InBlendSpacePlayersTimeline.EnumerateEvents(InGraphStartTime, InGraphEndTime, [GameplayProvider, &DebugData](double InStartTime, double InEndTime, uint32 InDepth, const FBlendSpacePlayerMessage& InMessage)
+									UBlendSpaceBase* BlendSpaceBase = nullptr;
+									const FObjectInfo* BlendSpaceInfo = GameplayProvider->FindObjectInfo(InMessage.BlendSpaceId);
+									if(BlendSpaceInfo)
 									{
-										UBlendSpaceBase* BlendSpaceBase = nullptr;
-										const FObjectInfo* BlendSpaceInfo = GameplayProvider->FindObjectInfo(InMessage.BlendSpaceId);
-										if(BlendSpaceInfo)
-										{
-											BlendSpaceBase = TSoftObjectPtr<UBlendSpaceBase>(FSoftObjectPath(BlendSpaceInfo->PathName)).LoadSynchronous();
-										}
+										BlendSpaceBase = TSoftObjectPtr<UBlendSpaceBase>(FSoftObjectPath(BlendSpaceInfo->PathName)).LoadSynchronous();
+									}
 
-										DebugData.RecordBlendSpacePlayer(InMessage.NodeId, BlendSpaceBase, InMessage.PositionX, InMessage.PositionY, InMessage.PositionZ);
-									});
+									DebugData.RecordBlendSpacePlayer(InMessage.NodeId, BlendSpaceBase, InMessage.PositionX, InMessage.PositionY, InMessage.PositionZ);
 								});
-							}
+							});
 						}
-					});
-				}
+					}
+				});
 			});
 		}
 	}
diff --git a/Engine/Plugins/Animation/GameplayInsights/Source/GameplayInsights/Private/AnimNodesTrack.h b/Engine/Plugins/Animation/GameplayInsights/Source/GameplayInsights/Private/AnimNodesTrack.h
--- a/Engine/Plugins/Animation/GameplayInsights/Source/GameplayInsights/Private/AnimNodesTrack.h
+++ b/Engine/Plugins/Animation/GameplayInsights/Source/GameplayInsights/Private/AnimNodesTrack.h
@@ -37,7 +37,7 @@ public:
 
 #if WITH_EDITOR
 	// Update the debug data for any classes that are running on this track
-	void UpdateDebugData(double InTime);
+	void UpdateDebugData(const Trace::FFrame& InFrame);
 
 	// Get custom debug objects for integration with anim blueprint debugging
 	void GetCustomDebugObjects(const IAnimationBlueprintEditor& InAnimationBlueprintEditor, TArray<FCustomDebugObject>& OutDebugList);
diff --git a/Engine/Plugins/Animation/GameplayInsights/Source/GameplayInsights/Private/AnimNotifiesTrack.cpp b/Engine/Plugins/Animation/GameplayInsights/Source/GameplayInsights/Private/AnimNotifiesTrack.cpp
--- a/Engine/Plugins/Animation/GameplayInsights/Source/GameplayInsights/Private/AnimNotifiesTrack.cpp
+++ b/Engine/Plugins/Animation/GameplayInsights/Source/GameplayInsights/Private/AnimNotifiesTrack.cpp
@@ -34,11 +34,11 @@ void FAnimNotifiesTrack::BuildDrawState(ITimingEventsTrackDrawStateBuilder& Buil
 
 		bHasNotifies = false;
 
-		AnimationProvider->ReadNotifyTimeline(GetGameplayTrack().GetObjectId(), [this, &AnimationProvider, &Context, &Builder](const FAnimationProvider::AnimNotifyTimeline& InTimeline)
+		AnimationProvider->ReadNotifyTimeline(GetGameplayTrack().GetObjectId(), [this, &Context, &Builder](const FAnimationProvider::AnimNotifyTimeline& InTimeline)
 		{
-			InTimeline.EnumerateEvents(Context.GetViewport().GetStartTime(), Context.GetViewport().GetEndTime(), [this, &AnimationProvider, &Builder](double InStartTime, double InEndTime, uint32 InDepth, const FAnimNotifyMessage& InMessage)
+			InTimeline.EnumerateEvents(Context.GetViewport().GetStartTime(), Context.GetViewport().GetEndTime(), [this, &Builder](double InStartTime, double InEndTime, uint32 InDepth, const FAnimNotifyMessage& InMessage)
 			{
-				Builder.AddEvent(InStartTime, InEndTime, 0, AnimationProvider->GetName(InMessage.NameId));
+				Builder.AddEvent(InStartTime, InEndTime, 0, InMessage.Name);
 				bHasNotifies = true;
 			});
 		});
@@ -46,7 +46,7 @@ void FAnimNotifiesTrack::BuildDrawState(ITimingEventsTrackDrawStateBuilder& Buil
 		uint32 Depth = bHasNotifies ? 1 : 0;
 		IdToDepthMap.Reset();
 
-		AnimationProvider->EnumerateNotifyStateTimelines(GetGameplayTrack().GetObjectId(), [this, &AnimationProvider, &Context, &Builder, &Depth](uint64 InNotifyId, const FAnimationProvider::AnimNotifyTimeline& InTimeline)
+		AnimationProvider->EnumerateNotifyStateTimelines(GetGameplayTrack().GetObjectId(), [this, &Context, &Builder, &Depth](uint64 InNotifyId, const FAnimationProvider::AnimNotifyTimeline& InTimeline)
 		{
 			uint32 DepthToUse = 0;
 			if(uint32* FoundDepthPtr = IdToDepthMap.Find(InNotifyId))
@@ -59,9 +59,9 @@ void FAnimNotifiesTrack::BuildDrawState(ITimingEventsTrackDrawStateBuilder& Buil
 				DepthToUse = Depth;
 			}
 
-			InTimeline.EnumerateEvents(Context.GetViewport().GetStartTime(), Context.GetViewport().GetEndTime(), [&AnimationProvider, &Builder, &DepthToUse](double InStartTime, double InEndTime, uint32 InDepth, const FAnimNotifyMessage& InMessage)
+			InTimeline.EnumerateEvents(Context.GetViewport().GetStartTime(), Context.GetViewport().GetEndTime(), [&Builder, &DepthToUse](double InStartTime, double InEndTime, uint32 InDepth, const FAnimNotifyMessage& InMessage)
 			{
-				Builder.AddEvent(InStartTime, InEndTime, DepthToUse, AnimationProvider->GetName(InMessage.NameId));
+				Builder.AddEvent(InStartTime, InEndTime, DepthToUse, InMessage.Name);
 			});
 		});
 	}
@@ -73,6 +73,51 @@ void FAnimNotifiesTrack::Draw(const ITimingTrackDrawContext& Context) const
 	GetGameplayTrack().DrawHeaderForTimingTrack(Context, *this, false);
 }
 
+static const TCHAR* GetMessageType(EAnimNotifyMessageType InType, bool bInForEvent)
+{
+	switch (InType)
+	{
+	case EAnimNotifyMessageType::Begin:
+	case EAnimNotifyMessageType::End:
+	case EAnimNotifyMessageType::Tick:
+	{
+		static FText StateType = LOCTEXT("AnimNotifyState", "Anim Notify State");
+		static FText BeginType = LOCTEXT("AnimNotifyStateBegin", "Anim Notify State (Begin)");
+		static FText EndType = LOCTEXT("AnimNotifyStateEnd", "Anim Notify State (End)");
+		static FText TickType = LOCTEXT("AnimNotifyStateTick", "Anim Notify State (Tick)");
+		if (bInForEvent)
+		{
+			switch (InType)
+			{
+			case EAnimNotifyMessageType::Begin:
+				return *BeginType.ToString();
+			case EAnimNotifyMessageType::End:
+				return *EndType.ToString();
+			case EAnimNotifyMessageType::Tick:
+				return *TickType.ToString();
+			}
+		}
+		else
+		{
+			return *StateType.ToString();
+		}
+	}
+	case EAnimNotifyMessageType::Event:
+	{
+		static FText Type = LOCTEXT("AnimNotify", "Anim Notify");
+		return *Type.ToString();
+	}
+	case EAnimNotifyMessageType::SyncMarker:
+	{
+		static FText Type = LOCTEXT("SyncMarker", "Sync Marker");
+		return *Type.ToString();
+	}
+	}
+
+	static FText UnknownType = LOCTEXT("UnknownType", "Unknown");
+	return *UnknownType.ToString();
+}
+
 void FAnimNotifiesTrack::InitTooltip(FTooltipDrawState& Tooltip, const ITimingEvent& HoveredTimingEvent) const
 {
 	auto EventFilter = [&HoveredTimingEvent](double EventStartTime, double EventEndTime, uint32 EventDepth) 
@@ -88,6 +133,8 @@ void FAnimNotifiesTrack::InitTooltip(FTooltipDrawState& Tooltip, const ITimingEv
 
 		Tooltip.AddTitle(GetName());
 
+		Tooltip.AddNameValueTextLine(LOCTEXT("NotifyType", "Type").ToString(), GetMessageType(InMessage.NotifyEventType, InFoundDepth == 0));
+
 		{
 			Trace::FAnalysisSessionReadScope SessionReadScope(SharedData.GetAnalysisSession());
 
@@ -119,6 +166,8 @@ void FAnimNotifiesTrack::InitTooltip(FTooltipDrawState& Tooltip, const ITimingEv
 			}
 		}
 
+		Tooltip.AddNameValueTextLine(LOCTEXT("EventWorld", "World").ToString(), GetGameplayTrack().GetWorldName(SharedData.GetAnalysisSession()).ToString());
+
 		Tooltip.UpdateLayout();
 	});
 }
@@ -181,34 +230,7 @@ void FAnimNotifiesTrack::FindAnimNotifyMessage(const FTimingEventSearchParameter
 		});
 }
 
-static const TCHAR* GetMessageType(EAnimNotifyMessageType InType)
-{
-	switch(InType)
-	{
-	case EAnimNotifyMessageType::Begin:
-	case EAnimNotifyMessageType::End:
-	case EAnimNotifyMessageType::Tick:
-	{
-		static FText Type = LOCTEXT("AnimNotifyState", "Anim Notify State");
-		return *Type.ToString();
-	}
-	case EAnimNotifyMessageType::Event:
-	{
-		static FText Type = LOCTEXT("AnimNotify", "Anim Notify");
-		return *Type.ToString();
-	}
-	case EAnimNotifyMessageType::SyncMarker:
-	{
-		static FText Type = LOCTEXT("SyncMarker", "Sync Marker");
-		return *Type.ToString();
-	}
-	}
-
-	static FText UnknownType = LOCTEXT("UnknownType", "Unknown");
-	return *UnknownType.ToString();
-}
-
-void FAnimNotifiesTrack::GetVariantsAtTime(double InTime, TArray<TSharedRef<FVariantTreeNode>>& OutVariants) const 
+void FAnimNotifiesTrack::GetVariantsAtFrame(const Trace::FFrame& InFrame, TArray<TSharedRef<FVariantTreeNode>>& OutVariants) const 
 {
 	TSharedRef<FVariantTreeNode> Header = OutVariants.Add_GetRef(FVariantTreeNode::MakeHeader(LOCTEXT("AnimNotifiesHeader", "Notifies and Sync Markers")));
 
@@ -220,39 +242,36 @@ void FAnimNotifiesTrack::GetVariantsAtTime(double InTime, TArray<TSharedRef<FVar
 	{
 		Trace::FAnalysisSessionReadScope SessionReadScope(SharedData.GetAnalysisSession());
 
-		// round to nearest frame boundary
-		Trace::FFrame Frame;
-		if(FramesProvider.GetFrameFromTime(ETraceFrameType::TraceFrameType_Game, InTime, Frame))
+		bool bForEvent = true;
+		auto ProcessEvent = [this, &AnimationProvider, &GameplayProvider, &Header, &bForEvent](double InStartTime, double InEndTime, uint32 InDepth, const FAnimNotifyMessage& InMessage)
 		{
-			auto ProcessEvent = [this, &AnimationProvider, &GameplayProvider, &Header](double InStartTime, double InEndTime, uint32 InDepth, const FAnimNotifyMessage& InMessage)
-			{
-				const TCHAR* Name = AnimationProvider->GetName(InMessage.NameId);
-				TSharedRef<FVariantTreeNode> NotifyHeader = Header->AddChild(FVariantTreeNode::MakeHeader(FText::FromString(Name)));
+			const TCHAR* Name = AnimationProvider->GetName(InMessage.NameId);
+			TSharedRef<FVariantTreeNode> NotifyHeader = Header->AddChild(FVariantTreeNode::MakeHeader(FText::FromString(Name)));
 
-				NotifyHeader->AddChild(FVariantTreeNode::MakeString(LOCTEXT("EventType", "Type"), GetMessageType(InMessage.NotifyEventType)));
-				NotifyHeader->AddChild(FVariantTreeNode::MakeFloat(LOCTEXT("EventDuration", "Duration"), InEndTime - InStartTime));
-				NotifyHeader->AddChild(FVariantTreeNode::MakeFloat(LOCTEXT("EventTime", "Time"), InStartTime));
+			NotifyHeader->AddChild(FVariantTreeNode::MakeString(LOCTEXT("EventType", "Type"), GetMessageType(InMessage.NotifyEventType, bForEvent)));
+			NotifyHeader->AddChild(FVariantTreeNode::MakeFloat(LOCTEXT("EventDuration", "Duration"), InEndTime - InStartTime));
+			NotifyHeader->AddChild(FVariantTreeNode::MakeFloat(LOCTEXT("EventTime", "Time"), InStartTime));
 
-				if(InMessage.NotifyEventType != EAnimNotifyMessageType::SyncMarker)
-				{
-					NotifyHeader->AddChild(FVariantTreeNode::MakeObject(LOCTEXT("Asset", "Asset"), InMessage.AssetId));
+			if(InMessage.NotifyEventType != EAnimNotifyMessageType::SyncMarker)
+			{
+				NotifyHeader->AddChild(FVariantTreeNode::MakeObject(LOCTEXT("Asset", "Asset"), InMessage.AssetId));
 
-					const FObjectInfo& NotifyInfo = GameplayProvider->GetObjectInfo(InMessage.NotifyId);
-					bool bIsState = InMessage.NotifyEventType == EAnimNotifyMessageType::Begin || InMessage.NotifyEventType == EAnimNotifyMessageType::End || InMessage.NotifyEventType == EAnimNotifyMessageType::Tick;
-					NotifyHeader->AddChild(FVariantTreeNode::MakeClass((bIsState ? LOCTEXT("Notify State Class", "Notify State Class") : LOCTEXT("NotifyClass", "Notify Class")), NotifyInfo.ClassId));
-				}
-			};
+				const FObjectInfo& NotifyInfo = GameplayProvider->GetObjectInfo(InMessage.NotifyId);
+				bool bIsState = InMessage.NotifyEventType == EAnimNotifyMessageType::Begin || InMessage.NotifyEventType == EAnimNotifyMessageType::End || InMessage.NotifyEventType == EAnimNotifyMessageType::Tick;
+				NotifyHeader->AddChild(FVariantTreeNode::MakeClass((bIsState ? LOCTEXT("Notify State Class", "Notify State Class") : LOCTEXT("NotifyClass", "Notify Class")), NotifyInfo.ClassId));
+			}
+		};
 
-			AnimationProvider->ReadNotifyTimeline(GetGameplayTrack().GetObjectId(), [&Frame, &ProcessEvent](const FAnimationProvider::AnimNotifyTimeline& InTimeline)
-			{
-				InTimeline.EnumerateEvents(Frame.StartTime, Frame.EndTime, ProcessEvent);
-			});
+		AnimationProvider->ReadNotifyTimeline(GetGameplayTrack().GetObjectId(), [&InFrame, &ProcessEvent](const FAnimationProvider::AnimNotifyTimeline& InTimeline)
+		{
+			InTimeline.EnumerateEvents(InFrame.StartTime, InFrame.EndTime, ProcessEvent);
+		});
 
-			AnimationProvider->EnumerateNotifyStateTimelines(GetGameplayTrack().GetObjectId(), [&Frame, &ProcessEvent](uint32 InNotifyNameId, const FAnimationProvider::AnimNotifyTimeline& InTimeline)
-			{
-				InTimeline.EnumerateEvents(Frame.StartTime, Frame.EndTime, ProcessEvent);
-			});
-		}
+		bForEvent = false;
+		AnimationProvider->EnumerateNotifyStateTimelines(GetGameplayTrack().GetObjectId(), [&InFrame, &ProcessEvent](uint32 InNotifyNameId, const FAnimationProvider::AnimNotifyTimeline& InTimeline)
+		{
+			InTimeline.EnumerateEvents(InFrame.StartTime, InFrame.EndTime, ProcessEvent);
+		});
 	}
 }
 
diff --git a/Engine/Plugins/Animation/GameplayInsights/Source/GameplayInsights/Private/AnimNotifiesTrack.h b/Engine/Plugins/Animation/GameplayInsights/Source/GameplayInsights/Private/AnimNotifiesTrack.h
--- a/Engine/Plugins/Animation/GameplayInsights/Source/GameplayInsights/Private/AnimNotifiesTrack.h
+++ b/Engine/Plugins/Animation/GameplayInsights/Source/GameplayInsights/Private/AnimNotifiesTrack.h
@@ -20,7 +20,7 @@ public:
 	virtual void Draw(const ITimingTrackDrawContext& Context) const override;
 	virtual void InitTooltip(FTooltipDrawState& Tooltip, const ITimingEvent& HoveredTimingEvent) const override;
 	virtual const TSharedPtr<const ITimingEvent> SearchEvent(const FTimingEventSearchParameters& InSearchParameters) const override;
-	virtual void GetVariantsAtTime(double InTime, TArray<TSharedRef<FVariantTreeNode>>& OutVariants) const override;
+	virtual void GetVariantsAtFrame(const Trace::FFrame& InFrame, TArray<TSharedRef<FVariantTreeNode>>& OutVariants) const override;
 
 private:
 	// Helper function used to find an anim notify message
diff --git a/Engine/Plugins/Animation/GameplayInsights/Source/GameplayInsights/Private/AnimationAnalyzer.cpp b/Engine/Plugins/Animation/GameplayInsights/Source/GameplayInsights/Private/AnimationAnalyzer.cpp
--- a/Engine/Plugins/Animation/GameplayInsights/Source/GameplayInsights/Private/AnimationAnalyzer.cpp
+++ b/Engine/Plugins/Animation/GameplayInsights/Source/GameplayInsights/Private/AnimationAnalyzer.cpp
@@ -229,8 +229,8 @@ bool FAnimationAnalyzer::OnEvent(uint16 RouteId, const FOnEventContext& Context)
 		int32 NodeId = EventData.GetValue<int32>("NodeId");
 		float Position = EventData.GetValue<float>("Position");
 		float Length = EventData.GetValue<float>("Length");
-		int32 FrameCount = EventData.GetValue<int32>("FrameCount");
-		AnimationProvider.AppendAnimSequencePlayer(AnimInstanceId, Context.SessionContext.TimestampFromCycle(Cycle), NodeId, Position, Length, FrameCount);
+		uint16 FrameCounter = EventData.GetValue<uint16>("FrameCounter");
+		AnimationProvider.AppendAnimSequencePlayer(AnimInstanceId, Context.SessionContext.TimestampFromCycle(Cycle), NodeId, Position, Length, FrameCounter);
 		break;
 	}
 	case RouteId_BlendSpacePlayer:
@@ -287,7 +287,8 @@ bool FAnimationAnalyzer::OnEvent(uint16 RouteId, const FOnEventContext& Context)
 		uint32 NextSectionNameId = EventData.GetValue<uint32>("NextSectionNameId");
 		float Weight = EventData.GetValue<float>("Weight");
 		float DesiredWeight = EventData.GetValue<float>("DesiredWeight");
-		AnimationProvider.AppendMontage(AnimInstanceId, Context.SessionContext.TimestampFromCycle(Cycle), MontageId, CurrentSectionNameId, NextSectionNameId, Weight, DesiredWeight);
+		uint16 FrameCounter = EventData.GetValue<uint16>("FrameCounter");
+		AnimationProvider.AppendMontage(AnimInstanceId, Context.SessionContext.TimestampFromCycle(Cycle), MontageId, CurrentSectionNameId, NextSectionNameId, Weight, DesiredWeight, FrameCounter);
 		break;
 	}
 	}
diff --git a/Engine/Plugins/Animation/GameplayInsights/Source/GameplayInsights/Private/AnimationProvider.cpp b/Engine/Plugins/Animation/GameplayInsights/Source/GameplayInsights/Private/AnimationProvider.cpp
--- a/Engine/Plugins/Animation/GameplayInsights/Source/GameplayInsights/Private/AnimationProvider.cpp
+++ b/Engine/Plugins/Animation/GameplayInsights/Source/GameplayInsights/Private/AnimationProvider.cpp
@@ -89,40 +89,35 @@ void FAnimationProvider::EnumerateSkeletalMeshCurves(const FSkeletalMeshPoseMess
 	}
 }
 
-void FAnimationProvider::EnumerateTickRecordTimelines(uint64 InObjectId, TFunctionRef<void(uint64, int32, const TickRecordTimeline&)> Callback) const
+void FAnimationProvider::EnumerateTickRecordIds(uint64 InObjectId, TFunctionRef<void(uint64, int32)> Callback) const
 {
 	Session.ReadAccessCheck();
 
 	const uint32* IndexPtr = ObjectIdToTickRecordTimelineStorage.Find(InObjectId);
-	if(IndexPtr != nullptr)
+	if (IndexPtr != nullptr)
 	{
 		if (*IndexPtr < uint32(TickRecordTimelineStorage.Num()))
 		{
 			const TSharedRef<FTickRecordTimelineStorage>& TimelineStorage = TickRecordTimelineStorage[*IndexPtr];
-			for(auto AssetIdPair : TimelineStorage->AssetIdAndPlayerToTickRecordTimeline)
+			for (auto AssetIdPair : TimelineStorage->AssetIdAndPlayers)
 			{
-				Callback(AssetIdPair.Key.Key, AssetIdPair.Key.Value, TimelineStorage->Timelines[AssetIdPair.Value].Get());
+				Callback(AssetIdPair.Key, AssetIdPair.Value);
 			}
 		}
 	}
 }
 
-bool FAnimationProvider::ReadTickRecordTimeline(uint64 InObjectId, uint64 InAssetId, int32 InNodeId, TFunctionRef<void(const TickRecordTimeline&)> Callback) const
+bool FAnimationProvider::ReadTickRecordTimeline(uint64 InObjectId, TFunctionRef<void(const TickRecordTimeline&)> Callback) const
 {
 	Session.ReadAccessCheck();
 
 	const uint32* IndexPtr = ObjectIdToTickRecordTimelineStorage.Find(InObjectId);
-	if(IndexPtr != nullptr)
+	if (IndexPtr != nullptr)
 	{
 		if (*IndexPtr < uint32(TickRecordTimelineStorage.Num()))
 		{
-			const TSharedRef<FTickRecordTimelineStorage>& TimelineStorage = TickRecordTimelineStorage[*IndexPtr];
-			const uint32* TimelineIndexPtr = TimelineStorage->AssetIdAndPlayerToTickRecordTimeline.Find(TTuple<uint64, int32>(InAssetId, InNodeId));
-			if(TimelineIndexPtr != nullptr)
-			{
-				Callback(TimelineStorage->Timelines[*TimelineIndexPtr].Get());
-				return true;
-			}
+			Callback(*TickRecordTimelineStorage[*IndexPtr]->Timeline.Get());
+			return true;
 		}
 	}
 
@@ -377,7 +372,6 @@ void FAnimationProvider::AppendTickRecord(uint64 InAnimInstanceId, double InTime
 {
 	Session.WriteAccessCheck();
 
-	TSharedPtr<Trace::TPointTimeline<FTickRecordMessage>> Timeline;
 	TSharedPtr<FTickRecordTimelineStorage> TimelineStorage;
 	uint32* TimelineStorageIndexPtr = ObjectIdToTickRecordTimelineStorage.Find(InAnimInstanceId);
 	if(TimelineStorageIndexPtr != nullptr)
@@ -388,41 +382,13 @@ void FAnimationProvider::AppendTickRecord(uint64 InAnimInstanceId, double InTime
 	{
 		ObjectIdToTickRecordTimelineStorage.Add(InAnimInstanceId, TickRecordTimelineStorage.Num());
 		TimelineStorage = TickRecordTimelineStorage.Add_GetRef(MakeShared<FTickRecordTimelineStorage>());
+		TimelineStorage->Timeline = MakeShared<Trace::TPointTimeline<FTickRecordMessage>>(Session.GetLinearAllocator());
+		TimelineStorage->Timeline->SetEnumerateOutsideRange(true);
 	}
 
 	check(TimelineStorage.IsValid());
 
-	uint32* TimelineIndexPtr = TimelineStorage->AssetIdAndPlayerToTickRecordTimeline.Find(TTuple<uint64, int32>(InAssetId, InNodeId));
-	if(TimelineIndexPtr != nullptr)
-	{
-		Timeline = TimelineStorage->Timelines[*TimelineIndexPtr];
-	}
-	else
-	{
-		Timeline = MakeShared<Trace::TPointTimeline<FTickRecordMessage>>(Session.GetLinearAllocator());
-
-		// We need to enumerate outside range to make sure we draw lines to the edge of the viewport when iterating to draw
-		Timeline->SetEnumerateOutsideRange(true);
-		TimelineStorage->AssetIdAndPlayerToTickRecordTimeline.Add(TTuple<uint64, int32>(InAssetId, InNodeId), TimelineStorage->Timelines.Num());
-		TimelineStorage->Timelines.Add(Timeline.ToSharedRef());
-	}
-
-	// check for continuity
-	bool bContinuous = false;
-	const uint64 NumEvents = Timeline->GetEventCount();
-	if(NumEvents > 0)
-	{
-		if(InFrameCounter - 1 != Timeline->GetEvent(NumEvents - 1).FrameCounter)
-		{
-			// discontinuous frame count
-			bContinuous = false;
-		}
-		else
-		{
-			// continuous frame time
-			bContinuous = true;
-		}
-	}
+	TimelineStorage->AssetIdAndPlayers.Add(TTuple<uint64, int32>(InAssetId, InNodeId));
 
 	FTickRecordMessage Message;
 	Message.AnimInstanceId = InAnimInstanceId;
@@ -437,9 +403,8 @@ void FAnimationProvider::AppendTickRecord(uint64 InAnimInstanceId, double InTime
 	Message.FrameCounter = InFrameCounter;
 	Message.bLooping = bInLooping;
 	Message.bIsBlendSpace = bInIsBlendSpace;
-	Message.bContinuous = bContinuous;
 
-	Timeline->AppendEvent(InTime, Message);
+	TimelineStorage->Timeline->AppendEvent(InTime, Message);
 
 	Session.UpdateDurationSeconds(InTime);
 }
@@ -497,6 +462,7 @@ void FAnimationProvider::AppendSkeletalMeshComponent(uint64 InObjectId, uint64 I
 	Message.CurveStartIndex = SkeletalMeshCurves.Num();
 	Message.ComponentId = InObjectId;
 	Message.MeshId = InMeshId;
+	Message.MeshName = GameplayProvider.GetObjectInfo(Message.MeshId).Name;
 	Message.NumTransforms = (uint16)InPose.Num() - 1;
 	Message.NumCurves = (uint16)InCurves.Num();
 	Message.LodIndex = InLodIndex;
@@ -518,7 +484,7 @@ void FAnimationProvider::AppendSkeletalMeshComponent(uint64 InObjectId, uint64 I
 	Session.UpdateDurationSeconds(InTime);
 }
 
-void FAnimationProvider::AppendSkeletalMeshComponent(uint64 InObjectId, uint64 InMeshId, double InTime, uint16 InLodIndex, uint16 InFrameCounter, const FTransform& InComponentToworld, const TArrayView<const FTransform>& InPose, const TArrayView<const uint32>& InCurveIds, const TArrayView<const float>& InCurveValues)
+void FAnimationProvider::AppendSkeletalMeshComponent(uint64 InObjectId, uint64 InMeshId, double InTime, uint16 InLodIndex, uint16 InFrameCounter, const FTransform& InComponentToWorld, const TArrayView<const FTransform>& InPose, const TArrayView<const uint32>& InCurveIds, const TArrayView<const float>& InCurveValues)
 {
 	Session.WriteAccessCheck();
 
@@ -547,11 +513,12 @@ void FAnimationProvider::AppendSkeletalMeshComponent(uint64 InObjectId, uint64 I
 	const int32 NumCurves = InCurveIds.Num();
 
 	FSkeletalMeshPoseMessage Message;
-	Message.ComponentToWorld = InComponentToworld;
+	Message.ComponentToWorld = InComponentToWorld;
 	Message.TransformStartIndex = SkeletalMeshPoseTransforms.Num();
 	Message.CurveStartIndex = SkeletalMeshCurves.Num();
 	Message.ComponentId = InObjectId;
 	Message.MeshId = InMeshId;
+	Message.MeshName = GameplayProvider.GetObjectInfo(Message.MeshId).Name;
 	Message.NumTransforms = (uint16)InPose.Num();
 	Message.NumCurves = (uint16)NumCurves;
 	Message.LodIndex = InLodIndex;
@@ -693,7 +660,7 @@ void FAnimationProvider::AppendAnimNodeStart(uint64 InAnimInstanceId, double InS
 	Session.UpdateDurationSeconds(InStartTime);
 }
 
-void FAnimationProvider::AppendAnimSequencePlayer(uint64 InAnimInstanceId, double InTime, int32 InNodeId, float InPosition, float InLength, int32 InFrameCount)
+void FAnimationProvider::AppendAnimSequencePlayer(uint64 InAnimInstanceId, double InTime, int32 InNodeId, float InPosition, float InLength, uint16 InFrameCounter)
 {
 	Session.WriteAccessCheck();
 
@@ -716,7 +683,7 @@ void FAnimationProvider::AppendAnimSequencePlayer(uint64 InAnimInstanceId, doubl
 	Message.NodeId = InNodeId;
 	Message.Position = InPosition;
 	Message.Length = InLength;
-	Message.FrameCount = InFrameCount;
+	Message.FrameCounter = InFrameCounter;
 
 	Timeline->AppendEvent(InTime, Message);
 
@@ -924,6 +891,7 @@ void FAnimationProvider::AppendNotify(uint64 InAnimInstanceId, double InTime, ui
 			Message.AssetId = InAssetId;
 			Message.NotifyId = InNotifyId;
 			Message.NameId = InNameId;
+			Message.Name = GetName(InNameId);
 			Message.Time = InNotifyTime; 
 			Message.Duration = InNotifyDuration;
 			Message.NotifyEventType = InNotifyEventType;
@@ -955,6 +923,7 @@ void FAnimationProvider::AppendNotify(uint64 InAnimInstanceId, double InTime, ui
 		Message.AssetId = InAssetId;
 		Message.NotifyId = InNotifyId;
 		Message.NameId = InNameId;
+		Message.Name = GetName(InNameId);
 		Message.Time = InNotifyTime; 
 		Message.Duration = InNotifyDuration;
 		Message.NotifyEventType = InNotifyEventType;
@@ -965,7 +934,7 @@ void FAnimationProvider::AppendNotify(uint64 InAnimInstanceId, double InTime, ui
 	Session.UpdateDurationSeconds(InTime);
 }
 
-void FAnimationProvider::AppendMontage(uint64 InAnimInstanceId, double InTime, uint64 InMontageId, uint32 InCurrentSectionNameId, uint32 InNextSectionNameId, float InWeight, float InDesiredWeight)
+void FAnimationProvider::AppendMontage(uint64 InAnimInstanceId, double InTime, uint64 InMontageId, uint32 InCurrentSectionNameId, uint32 InNextSectionNameId, float InWeight, float InDesiredWeight, uint16 InFrameCounter)
 {
 	Session.WriteAccessCheck();
 
@@ -993,6 +962,7 @@ void FAnimationProvider::AppendMontage(uint64 InAnimInstanceId, double InTime, u
 	Message.NextSectionNameId = InNextSectionNameId;
 	Message.Weight = InWeight;
 	Message.DesiredWeight = InDesiredWeight;
+	Message.FrameCounter = InFrameCounter;
 
 	TimelineStorage->Timeline->AppendEvent(InTime, Message);
 
diff --git a/Engine/Plugins/Animation/GameplayInsights/Source/GameplayInsights/Private/AnimationProvider.h b/Engine/Plugins/Animation/GameplayInsights/Source/GameplayInsights/Private/AnimationProvider.h
--- a/Engine/Plugins/Animation/GameplayInsights/Source/GameplayInsights/Private/AnimationProvider.h
+++ b/Engine/Plugins/Animation/GameplayInsights/Source/GameplayInsights/Private/AnimationProvider.h
@@ -26,8 +26,8 @@ public:
 	virtual void GetSkeletalMeshComponentSpacePose(const FSkeletalMeshPoseMessage& InMessage, const FSkeletalMeshInfo& InMeshInfo, FTransform& OutComponentToWorld, TArray<FTransform>& OutTransforms) const override;
 	virtual void EnumerateSkeletalMeshCurveIds(uint64 InObjectId, TFunctionRef<void(uint32)> Callback) const override;
 	virtual void EnumerateSkeletalMeshCurves(const FSkeletalMeshPoseMessage& InMessage, TFunctionRef<void(const FSkeletalMeshNamedCurve&)> Callback) const override;
-	virtual void EnumerateTickRecordTimelines(uint64 InObjectId, TFunctionRef<void(uint64, int32, const TickRecordTimeline&)> Callback) const override;
-	virtual bool ReadTickRecordTimeline(uint64 InObjectId, uint64 InAssetId, int32 InNodeId, TFunctionRef<void(const TickRecordTimeline&)> Callback) const override;
+	virtual bool ReadTickRecordTimeline(uint64 InObjectId, TFunctionRef<void(const TickRecordTimeline&)> Callback) const override;
+	virtual void EnumerateTickRecordIds(uint64 InObjectId, TFunctionRef<void(uint64, int32)> Callback) const override;
 	virtual bool ReadAnimGraphTimeline(uint64 InObjectId, TFunctionRef<void(const AnimGraphTimeline&)> Callback) const override;
 	virtual bool ReadAnimNodesTimeline(uint64 InObjectId, TFunctionRef<void(const AnimNodesTimeline&)> Callback) const override;
 	virtual bool ReadAnimNodeValuesTimeline(uint64 InObjectId, TFunctionRef<void(const AnimNodeValuesTimeline&)> Callback) const override;
@@ -51,7 +51,7 @@ public:
 
 	/** Add a skeletal mesh pose/curves etc. */
 	void AppendSkeletalMeshComponent(uint64 InObjectId, uint64 InMeshId, double InTime, uint16 InLodIndex, uint16 InFrameCounter, const TArrayView<const FTransform>& InPose, const TArrayView<const FSkeletalMeshNamedCurve>& InCurves);
-	void AppendSkeletalMeshComponent(uint64 InObjectId, uint64 InMeshId, double InTime, uint16 InLodIndex, uint16 InFrameCounter, const FTransform& InComponentToworld, const TArrayView<const FTransform>& InPose, const TArrayView<const uint32>& InCurveIds, const TArrayView<const float>& InCurveValues);
+	void AppendSkeletalMeshComponent(uint64 InObjectId, uint64 InMeshId, double InTime, uint16 InLodIndex, uint16 InFrameCounter, const FTransform& InComponentToWorld, const TArrayView<const FTransform>& InPose, const TArrayView<const uint32>& InCurveIds, const TArrayView<const float>& InCurveValues);
 
 	/** Get a skeletal mesh for a specified path. If the mesh no longer exists, make a fake one */
 	USkeletalMesh* GetSkeletalMesh(const TCHAR* InPath);
@@ -79,7 +79,7 @@ public:
 	void AppendAnimNodeValueClass(uint64 InAnimInstanceId, double InTime, uint16 InFrameCounter, int32 InNodeId, const TCHAR* InKey, uint64 InValue);
 
 	/** Add a sequence player record */
-	void AppendAnimSequencePlayer(uint64 InAnimInstanceId, double InTime, int32 InNodeId, float InPosition, float InLength, int32 InFrameCount);
+	void AppendAnimSequencePlayer(uint64 InAnimInstanceId, double InTime, int32 InNodeId, float InPosition, float InLength, uint16 InFrameCounter);
 
 	/** Add a blend space player record */
 	void AppendBlendSpacePlayer(uint64 InAnimInstanceId, double InTime, int32 InNodeId, uint64 InBlendSpaceId, float InPositionX, float InPositionY, float InPositionZ);
@@ -91,7 +91,7 @@ public:
 	void AppendNotify(uint64 InAnimInstanceId, double InTime, uint64 InAssetId, uint64 InNotifyId, uint32 InNameId, float InNotifyTime, float InNotifyDuration, EAnimNotifyMessageType InNotifyEventType);
 
 	/** Append montage data */
-	void AppendMontage(uint64 InAnimInstanceId, double InTime, uint64 InMontageId, uint32 InCurrentSectionNameId, uint32 InNextSectionNameId, float InWeight, float InDesiredWeight);
+	void AppendMontage(uint64 InAnimInstanceId, double InTime, uint64 InMontageId, uint32 InCurrentSectionNameId, uint32 InNextSectionNameId, float InWeight, float InDesiredWeight, uint16 InFrameCounter);
 
 private:
 	/** Add anim node values helper */
@@ -129,8 +129,8 @@ private:
 	
 	struct FTickRecordTimelineStorage
 	{
-		TMap<TTuple<uint64, int32>, uint32> AssetIdAndPlayerToTickRecordTimeline;
-		TArray<TSharedRef<Trace::TPointTimeline<FTickRecordMessage>>> Timelines;
+		TSet<TTuple<uint64, int32>> AssetIdAndPlayers;
+		TSharedPtr<Trace::TPointTimeline<FTickRecordMessage>> Timeline;
 	};
 
 	struct FSkeletalMeshTimelineStorage
diff --git a/Engine/Plugins/Animation/GameplayInsights/Source/GameplayInsights/Private/AnimationSharedData.cpp b/Engine/Plugins/Animation/GameplayInsights/Source/GameplayInsights/Private/AnimationSharedData.cpp
--- a/Engine/Plugins/Animation/GameplayInsights/Source/GameplayInsights/Private/AnimationSharedData.cpp
+++ b/Engine/Plugins/Animation/GameplayInsights/Source/GameplayInsights/Private/AnimationSharedData.cpp
@@ -38,6 +38,7 @@ FAnimationSharedData::FAnimationSharedData(FGameplaySharedData& InGameplayShared
 	, AnalysisSession(nullptr)
 	, MarkerTime(0.0)
 	, bTimeMarkerValid(false)
+	, bMarkerFrameValid(false)
 	, bSkeletalMeshPoseTracksEnabled(true)
 	, bSkeletalMeshCurveTracksEnabled(true)
 	, bTickRecordTracksEnabled(true)
@@ -129,25 +130,17 @@ void FAnimationSharedData::Tick(Insights::ITimingViewSession& InTimingViewSessio
 						}
 					});
 
-					AnimationProvider->EnumerateTickRecordTimelines(InObjectInfo.Id, [this, &InObjectInfo, &ObjectEventsTrack, &InTimingViewSession, &GameplayProvider](uint64 InAssetId, int32 InNodeId, const IAnimationProvider::TickRecordTimeline& InTimeline)
+					AnimationProvider->ReadTickRecordTimeline(InObjectInfo.Id, [this, &InObjectInfo, &ObjectEventsTrack, &InTimingViewSession, &GameplayProvider](const IAnimationProvider::TickRecordTimeline& InTimeline)
 					{
-						auto FindTickRecordTrackWithAssetId = [&InAssetId](const FBaseTimingTrack& InTrack)
+						auto FindTickRecordTrackWithAssetId = [](const FBaseTimingTrack& InTrack)
 						{
-							if (InTrack.Is<FAnimationTickRecordsTrack>())
-							{
-								const FAnimationTickRecordsTrack& AnimationTickRecordsTrack = InTrack.As<FAnimationTickRecordsTrack>();
-								return AnimationTickRecordsTrack.GetAssetId() == InAssetId;
-							}
-
-							return false;
+							return InTrack.Is<FAnimationTickRecordsTrack>();
 						};
 
 						TSharedPtr<FAnimationTickRecordsTrack> ExistingAnimationTickRecordsTrack = StaticCastSharedPtr<FAnimationTickRecordsTrack>(ObjectEventsTrack->GetGameplayTrack().FindChildTrack(InObjectInfo.Id, FindTickRecordTrackWithAssetId));
 						if(!ExistingAnimationTickRecordsTrack.IsValid())
 						{
-							const FObjectInfo* AssetObjectInfo = GameplayProvider->FindObjectInfo(InAssetId);
-							FString AssetName = AssetObjectInfo ? AssetObjectInfo->Name : LOCTEXT("UnknownAsset", "Unknown").ToString();
-							TSharedPtr<FAnimationTickRecordsTrack> AnimationTickRecordsTrack = MakeShared<FAnimationTickRecordsTrack>(*this, InObjectInfo.Id, InAssetId, InNodeId, *AssetName);
+							TSharedPtr<FAnimationTickRecordsTrack> AnimationTickRecordsTrack = MakeShared<FAnimationTickRecordsTrack>(*this, InObjectInfo.Id, InObjectInfo.Name);
 							AnimationTickRecordsTrack->SetVisibilityFlag(bTickRecordTracksEnabled);
 							AnimationTickRecordsTracks.Add(AnimationTickRecordsTrack.ToSharedRef());
 
@@ -474,10 +467,23 @@ void FAnimationSharedData::OnTimeMarkerChanged(Insights::ETimeChangedFlags InFla
 	bTimeMarkerValid = InTimeMarker != std::numeric_limits<double>::infinity();
 	MarkerTime = InTimeMarker;
 
+	if(bTimeMarkerValid)
+	{
+		Trace::FAnalysisSessionReadScope SessionReadScope(*AnalysisSession);
+
+		const Trace::IFrameProvider& FramesProvider = Trace::ReadFrameProvider(*AnalysisSession);
+		bMarkerFrameValid = FramesProvider.GetFrameFromTime(ETraceFrameType::TraceFrameType_Game, MarkerTime, MarkerFrame);
+	}
+	else
+	{
+		MarkerFrame = Trace::FFrame();
+		bMarkerFrameValid = false;
+	}
+
 #if WITH_EDITOR
 	for(const TSharedRef<FAnimNodesTrack> AnimNodesTrack : AnimNodesTracks)
 	{
-		AnimNodesTrack->UpdateDebugData(InTimeMarker);
+		AnimNodesTrack->UpdateDebugData(MarkerFrame);
 	}
 
 	InvalidateViewports();
@@ -488,7 +494,7 @@ void FAnimationSharedData::OnTimeMarkerChanged(Insights::ETimeChangedFlags InFla
 	{
 		if(bTimeMarkerValid && PoseTrack->IsPotentiallyDebugged())
 		{
-			PoseTrack->DrawPoses(WorldToUse, MarkerTime);
+			PoseTrack->DrawPoses(WorldToUse, MarkerTime, MarkerFrame.StartTime, MarkerFrame.EndTime);
 		}
 	}
 #endif
@@ -565,9 +571,9 @@ void FAnimationSharedData::DrawPoses(UWorld* InWorld)
 	{
 		if(PoseTrack->IsVisible())
 		{
-			if(bTimeMarkerValid && PoseTrack->ShouldDrawPose())
+			if(bTimeMarkerValid && bMarkerFrameValid && PoseTrack->ShouldDrawPose())
 			{
-				PoseTrack->DrawPoses(InWorld, MarkerTime);
+				PoseTrack->DrawPoses(InWorld, MarkerTime, MarkerFrame.StartTime, MarkerFrame.EndTime);
 			}
 		}
 	}
diff --git a/Engine/Plugins/Animation/GameplayInsights/Source/GameplayInsights/Private/AnimationSharedData.h b/Engine/Plugins/Animation/GameplayInsights/Source/GameplayInsights/Private/AnimationSharedData.h
--- a/Engine/Plugins/Animation/GameplayInsights/Source/GameplayInsights/Private/AnimationSharedData.h
+++ b/Engine/Plugins/Animation/GameplayInsights/Source/GameplayInsights/Private/AnimationSharedData.h
@@ -3,6 +3,7 @@
 #pragma once
 
 #include "Insights/ViewModels/TimingEventsTrack.h"
+#include "TraceServices/Model/Frames.h"
 
 class FGameplaySharedData;
 namespace Trace { class IAnalysisSession; }
@@ -76,6 +77,9 @@ public:
 	/** Open an anim graph tab to see its schematic view */
 	void OpenAnimGraphTab(uint64 InAnimInstanceId) const;
 
+	/** Get the current marker frame, if any. @return true if the frame is valid */
+	bool GetCurrentMarkerFrame(Trace::FFrame& OutFrame) const { OutFrame = MarkerFrame; return bMarkerFrameValid; }
+
 private:
 	// UI handlers
 	void ToggleAnimationTracks();
@@ -113,9 +117,11 @@ private:
 
 	/** Various times and ranges */
 	double MarkerTime;
+	Trace::FFrame MarkerFrame;
 
 	/** Validity flags for pose times/ranges */
 	bool bTimeMarkerValid;
+	bool bMarkerFrameValid;
 
 	// Flags controlling whether check type of our animation tracks are enabled
 	bool bSkeletalMeshPoseTracksEnabled;
diff --git a/Engine/Plugins/Animation/GameplayInsights/Source/GameplayInsights/Private/AnimationTickRecordsTrack.cpp b/Engine/Plugins/Animation/GameplayInsights/Source/GameplayInsights/Private/AnimationTickRecordsTrack.cpp
--- a/Engine/Plugins/Animation/GameplayInsights/Source/GameplayInsights/Private/AnimationTickRecordsTrack.cpp
+++ b/Engine/Plugins/Animation/GameplayInsights/Source/GameplayInsights/Private/AnimationTickRecordsTrack.cpp
@@ -59,15 +59,12 @@ static FLinearColor MakeSeriesColor(FTickRecordSeries::ESeriesType InSeed, bool
 	return MakeSeriesColor((uint32)InSeed, bInLine);
 }
 
-FAnimationTickRecordsTrack::FAnimationTickRecordsTrack(const FAnimationSharedData& InSharedData, uint64 InObjectID, uint64 InAssetId, int32 InNodeId, const TCHAR* InName)
-	: FGameplayGraphTrack(InSharedData.GetGameplaySharedData(), InObjectID, MakeTrackName(InSharedData.GetGameplaySharedData(), InAssetId, InName))
+FAnimationTickRecordsTrack::FAnimationTickRecordsTrack(const FAnimationSharedData& InSharedData, uint64 InObjectID, const TCHAR* InName)
+	: FGameplayGraphTrack(InSharedData.GetGameplaySharedData(), InObjectID, FText::Format(LOCTEXT("AnimationTickRecordsTrackName", "Blend Weights - {0}"), FText::FromString(InName)))
 	, SharedData(InSharedData)
-	, AssetId(InAssetId)
-	, NodeId(InNodeId)
 {
-	uint32 Hash = GetTypeHash(GetName());
-	MainSeriesLineColor = MakeSeriesColor(Hash, true);
-	MainSeriesFillColor = MakeSeriesColor(Hash);
+	bDrawLabels = true;
+	Layout = EGameplayGraphLayout::Stack;
 
 #if WITH_EDITOR
 	const FGameplayProvider* GameplayProvider = SharedData.GetAnalysisSession().ReadProvider<FGameplayProvider>(FGameplayProvider::ProviderName);
@@ -89,109 +86,112 @@ FAnimationTickRecordsTrack::FAnimationTickRecordsTrack(const FAnimationSharedDat
 
 void FAnimationTickRecordsTrack::AddAllSeries()
 {
-	if(AllSeries.Num() > 0)
-	{
-		return;
-	}
-
 	struct FSeriesDescription
 	{
 		FText Name;
 		FText Description;
-		FLinearColor LineColor;
-		FLinearColor FillColor;
 		FTickRecordSeries::ESeriesType Type;
 		bool bEnabled;
 	};
 
-	const FSeriesDescription SeriesDescriptions[] =
+	const FGameplayProvider* GameplayProvider = SharedData.GetAnalysisSession().ReadProvider<FGameplayProvider>(FGameplayProvider::ProviderName);
+	const FAnimationProvider* AnimationProvider = SharedData.GetAnalysisSession().ReadProvider<FAnimationProvider>(FAnimationProvider::ProviderName);
+	if(GameplayProvider && AnimationProvider)
 	{
-		{ 
-			LOCTEXT("SeriesNameBlendWeight", "Blend Weight"), 
-			LOCTEXT("SeriesDescBlendWeight", "The final effective weight that this animation sequence was played at"), 
-			MainSeriesLineColor, 
-			MainSeriesFillColor, 
-			FTickRecordSeries::ESeriesType::BlendWeight,
-			true
-		},
-		{ 
-			LOCTEXT("SeriesNamePlaybackTime", "Playback Time"), 
-			LOCTEXT("SeriesDescPlaybackTime", "The playback time of this animation sequence"), 
-			MakeSeriesColor(FTickRecordSeries::ESeriesType::PlaybackTime, true), 
-			MakeSeriesColor(FTickRecordSeries::ESeriesType::PlaybackTime),
-			FTickRecordSeries::ESeriesType::PlaybackTime,
-			false
-		},
-		{ 
-			LOCTEXT("SeriesNameRootMotionWeight", "Root Motion Weight"),
-			LOCTEXT("SeriesDescRootMotionWeight", "The final effective root motion weight that this animation sequence was played at"), 
-			MakeSeriesColor(FTickRecordSeries::ESeriesType::RootMotionWeight, true), 
-			MakeSeriesColor(FTickRecordSeries::ESeriesType::RootMotionWeight), 
-			FTickRecordSeries::ESeriesType::RootMotionWeight, 
-			false
-		},
-		{ 
-			LOCTEXT("SeriesNamePlayRate", "Play Rate"), 
-			LOCTEXT("SeriesDescPlayRate", "The play rate/speed of this animation sequence"), 
-			MakeSeriesColor(FTickRecordSeries::ESeriesType::PlayRate, true),
-			MakeSeriesColor(FTickRecordSeries::ESeriesType::PlayRate), 
-			FTickRecordSeries::ESeriesType::PlayRate,
-			false
-		},
-	};
+		Trace::FAnalysisSessionReadScope SessionReadScope(SharedData.GetAnalysisSession());
 
-	const FSeriesDescription BlendSpaceSeriesDescriptions[] =
-	{
-		{ 
-			LOCTEXT("SeriesNameBlendSpacePositionX", "BlendSpace Position X"), 
-			LOCTEXT("SeriesDescBlendSpacePositionX", "The X value used to sample this blend space"), 
-			MakeSeriesColor(FTickRecordSeries::ESeriesType::BlendSpacePositionX, true),
-			MakeSeriesColor(FTickRecordSeries::ESeriesType::BlendSpacePositionX),
-			FTickRecordSeries::ESeriesType::BlendSpacePositionX, 
-			false
-		},
+		AnimationProvider->EnumerateTickRecordIds(GetGameplayTrack().GetObjectId(), [this, &GameplayProvider](uint64 InAssetId, int32 InNodeId)
 		{
-			LOCTEXT("SeriesNameBlendSpacePositionY", "BlendSpace Position Y"), 
-			LOCTEXT("SeriesDescBlendSpacePositionY", "The Y value used to sample this blend space"), 
-			MakeSeriesColor(FTickRecordSeries::ESeriesType::BlendSpacePositionY, true), 
-			MakeSeriesColor(FTickRecordSeries::ESeriesType::BlendSpacePositionY), 
-			FTickRecordSeries::ESeriesType::BlendSpacePositionY, 
-			false
-		},
-	};
+			auto CheckExistingSeries = [InAssetId, InNodeId](const TSharedPtr<FGraphSeries>& InSeries)
+			{
+				TSharedPtr<FTickRecordSeries> TickRecordSeries = StaticCastSharedPtr<FTickRecordSeries>(InSeries);
+				return TickRecordSeries->AssetId == InAssetId && TickRecordSeries->NodeId == InNodeId;
+			};
 
-	auto AddSeries = [this](const FSeriesDescription& InSeriesDescription)
-	{
-		TSharedRef<FTickRecordSeries> Series = MakeShared<FTickRecordSeries>();
-		Series->SetName(InSeriesDescription.Name.ToString());
-		Series->SetDescription(InSeriesDescription.Description.ToString());
-		Series->SetColor(InSeriesDescription.LineColor, InSeriesDescription.LineColor, InSeriesDescription.FillColor);
-		Series->Type = InSeriesDescription.Type;
-		Series->SetVisibility(InSeriesDescription.bEnabled);
-		Series->SetBaselineY(25.0f);
-		Series->SetScaleY(20.0f);
-		Series->EnableAutoZoom();
-		AllSeries.Add(Series);	
-	};
+			if (!AllSeries.ContainsByPredicate(CheckExistingSeries))
+			{
+				static const FSeriesDescription SeriesDescriptions[] =
+				{
+					{
+						LOCTEXT("SeriesNameBlendWeight", "Blend Weight"),
+						LOCTEXT("SeriesDescBlendWeight", "The final effective weight that this animation sequence was played at"),
+						FTickRecordSeries::ESeriesType::BlendWeight,
+						true
+					},
+					{
+						LOCTEXT("SeriesNamePlaybackTime", "Playback Time"),
+						LOCTEXT("SeriesDescPlaybackTime", "The playback time of this animation sequence"),
+						FTickRecordSeries::ESeriesType::PlaybackTime,
+						false
+					},
+					{
+						LOCTEXT("SeriesNameRootMotionWeight", "Root Motion Weight"),
+						LOCTEXT("SeriesDescRootMotionWeight", "The final effective root motion weight that this animation sequence was played at"),
+						FTickRecordSeries::ESeriesType::RootMotionWeight,
+						false
+					},
+					{
+						LOCTEXT("SeriesNamePlayRate", "Play Rate"),
+						LOCTEXT("SeriesDescPlayRate", "The play rate/speed of this animation sequence"),
+						FTickRecordSeries::ESeriesType::PlayRate,
+						false
+					},
+				};
+
+				static const FSeriesDescription BlendSpaceSeriesDescriptions[] =
+				{
+					{
+						LOCTEXT("SeriesNameBlendSpacePositionX", "BlendSpace Position X"),
+						LOCTEXT("SeriesDescBlendSpacePositionX", "The X value used to sample this blend space"),
+						FTickRecordSeries::ESeriesType::BlendSpacePositionX,
+						false
+					},
+					{
+						LOCTEXT("SeriesNameBlendSpacePositionY", "BlendSpace Position Y"),
+						LOCTEXT("SeriesDescBlendSpacePositionY", "The Y value used to sample this blend space"),
+						FTickRecordSeries::ESeriesType::BlendSpacePositionY,
+						false
+					},
+				};
 
-	for(const FSeriesDescription& SeriesDescription : SeriesDescriptions)
-	{
-		AddSeries(SeriesDescription);
-	}
+				const FClassInfo& ClassInfo = GameplayProvider->GetClassInfoFromObject(InAssetId);
+				const FObjectInfo& ObjectInfo = GameplayProvider->GetObjectInfo(InAssetId);
 
-	const FGameplayProvider* GameplayProvider = SharedData.GetAnalysisSession().ReadProvider<FGameplayProvider>(FGameplayProvider::ProviderName);
-	if(GameplayProvider)
-	{
-		Trace::FAnalysisSessionReadScope SessionReadScope(SharedData.GetAnalysisSession());
+				auto AddSeries = [this, GameplayProvider](const FSeriesDescription& InSeriesDescription, const TCHAR* InAssetName, uint64 InAssetId, int32 InNodeId)
+				{
+					FString Name = FText::Format(LOCTEXT("SeriesNameFormat", "{0} - {1}"), FText::FromString(InAssetName), InSeriesDescription.Name).ToString();
+					uint32 Hash = GetTypeHash(Name);
+					FLinearColor LineColor = MakeSeriesColor(Hash, true);
+					FLinearColor FillColor = MakeSeriesColor(Hash);
+
+					TSharedRef<FTickRecordSeries> Series = MakeShared<FTickRecordSeries>();
+					Series->SetName(Name);
+					Series->SetDescription(InSeriesDescription.Description.ToString());
+					Series->SetColor(LineColor, LineColor, FillColor);
+					Series->Type = InSeriesDescription.Type;
+					Series->SetVisibility(InSeriesDescription.bEnabled);
+					Series->SetBaselineY(25.0f);
+					Series->SetScaleY(20.0f);
+					Series->EnableAutoZoom();
+					Series->AssetId = InAssetId;
+					Series->NodeId = InNodeId;
+					AllSeries.Add(Series);
+				};
+
+				for (const FSeriesDescription& SeriesDescription : SeriesDescriptions)
+				{
+					AddSeries(SeriesDescription, ObjectInfo.Name, InAssetId, InNodeId);
+				}
 
-		const FClassInfo& ClassInfo = GameplayProvider->GetClassInfoFromObject(AssetId);
-		if(FCString::Stristr(ClassInfo.Name, TEXT("BlendSpace")) != nullptr)
-		{
-			for(const FSeriesDescription& BlendSpaceSeriesDescription : BlendSpaceSeriesDescriptions)
-			{
-				AddSeries(BlendSpaceSeriesDescription);
+				if (FCString::Stristr(ClassInfo.Name, TEXT("BlendSpace")) != nullptr)
+				{
+					for (const FSeriesDescription& BlendSpaceSeriesDescription : BlendSpaceSeriesDescriptions)
+					{
+						AddSeries(BlendSpaceSeriesDescription, ObjectInfo.Name, InAssetId, InNodeId);
+					}
+				}
 			}
-		}
+		});
 	}
 }
 
@@ -209,14 +209,17 @@ bool FAnimationTickRecordsTrack::UpdateSeriesBoundsHelper(FTickRecordSeries& InS
 		InSeries.CurrentMin = 0.0;
 		InSeries.CurrentMax = 0.0;
 
-		AnimationProvider->ReadTickRecordTimeline(GetGameplayTrack().GetObjectId(), GetAssetId(), GetNodeId(), [&bFoundEvents, &InViewport, &InSeries, &Projection](const FAnimationProvider::TickRecordTimeline& InTimeline)
+		AnimationProvider->ReadTickRecordTimeline(GetGameplayTrack().GetObjectId(), [&bFoundEvents, &InViewport, &InSeries, &Projection](const FAnimationProvider::TickRecordTimeline& InTimeline)
 		{
 			InTimeline.EnumerateEvents(InViewport.GetStartTime(), InViewport.GetEndTime(), [&bFoundEvents, &InSeries, &Projection](double InStartTime, double InEndTime, uint32 InDepth, const FTickRecordMessage& InMessage)
 			{
-				const float Value = Invoke(Projection, InMessage);
-				InSeries.CurrentMin = FMath::Min(InSeries.CurrentMin, Value);
-				InSeries.CurrentMax = FMath::Max(InSeries.CurrentMax, Value);
-				bFoundEvents = true;
+				if (InMessage.AssetId == InSeries.AssetId && InMessage.NodeId == InSeries.NodeId)
+				{
+					const float Value = Invoke(Projection, InMessage);
+					InSeries.CurrentMin = FMath::Min(InSeries.CurrentMin, Value);
+					InSeries.CurrentMax = FMath::Max(InSeries.CurrentMax, Value);
+					bFoundEvents = true;
+				}
 			});
 		});
 	}
@@ -235,11 +238,21 @@ void FAnimationTickRecordsTrack::UpdateSeriesHelper(FTickRecordSeries& InSeries,
 
 		FGraphTrackBuilder Builder(*this, InSeries, InViewport);
 
-		AnimationProvider->ReadTickRecordTimeline(GetGameplayTrack().GetObjectId(), GetAssetId(), GetNodeId(), [this, &AnimationProvider, &Builder, &InViewport, &Projection](const FAnimationProvider::TickRecordTimeline& InTimeline)
+		AnimationProvider->ReadTickRecordTimeline(GetGameplayTrack().GetObjectId(), [this, &AnimationProvider, &Builder, &InViewport, &InSeries, &Projection](const FAnimationProvider::TickRecordTimeline& InTimeline)
 		{
-			InTimeline.EnumerateEvents(InViewport.GetStartTime(), InViewport.GetEndTime(), [this, &Builder, &Projection](double InStartTime, double InEndTime, uint32 InDepth, const FTickRecordMessage& InMessage)
+			uint16 FrameCounter = 0;
+			uint16 LastFrameWithTickRecord = 0;
+
+			InTimeline.EnumerateEvents(InViewport.GetStartTime(), InViewport.GetEndTime(), [this, &FrameCounter, &LastFrameWithTickRecord, &Builder, &InSeries, &Projection](double InStartTime, double InEndTime, uint32 InDepth, const FTickRecordMessage& InMessage)
 			{
-				Builder.AddEvent(InStartTime, InEndTime - InStartTime, Invoke(Projection, InMessage), InMessage.bContinuous);
+				FrameCounter = InMessage.FrameCounter;
+
+				if (InMessage.AssetId == InSeries.AssetId && InMessage.NodeId == InSeries.NodeId)
+				{
+					Builder.AddEvent(InStartTime, InEndTime - InStartTime, Invoke(Projection, InMessage), LastFrameWithTickRecord == FrameCounter - 1);
+
+					LastFrameWithTickRecord = InMessage.FrameCounter;
+				}
 			});
 		});
 	}
@@ -293,6 +306,20 @@ void FAnimationTickRecordsTrack::UpdateSeries(FGameplayGraphSeries& InSeries, co
 	}
 }
 
+void FAnimationTickRecordsTrack::PostUpdate(const ITimingTrackUpdateContext& Context)
+{
+	FGameplayGraphTrack::PostUpdate(Context);
+
+	if(Context.GetHoveredEvent().IsValid() && Context.GetHoveredEvent()->Is<FGraphTrackEvent>())
+	{
+		TSharedPtr<const FGraphTrackEvent> HoveredEvent = StaticCastSharedPtr<const FGraphTrackEvent>(Context.GetHoveredEvent());
+		if(HoveredEvent->GetTrack() == SharedThis(this))
+		{
+			CachedHoveredEvent = HoveredEvent;
+		}
+	}
+}
+
 void FAnimationTickRecordsTrack::InitTooltip(FTooltipDrawState& Tooltip, const ITimingEvent& HoveredTimingEvent) const
 {
 	const FGraphTrackEvent& GraphTrackEvent = *static_cast<const FGraphTrackEvent*>(&HoveredTimingEvent);
@@ -307,6 +334,7 @@ void FAnimationTickRecordsTrack::InitTooltip(FTooltipDrawState& Tooltip, const I
 
 		Tooltip.AddNameValueTextLine(LOCTEXT("EventTime", "Time").ToString(), FText::AsNumber(InFoundStartTime).ToString());
 		Tooltip.AddNameValueTextLine(GraphTrackEvent.GetSeries()->GetName().ToString(), FText::AsNumber(GraphTrackEvent.GetValue()).ToString());
+		Tooltip.AddNameValueTextLine(LOCTEXT("EventWorld", "World").ToString(), GetGameplayTrack().GetWorldName(SharedData.GetAnalysisSession()).ToString());
 
 		Tooltip.UpdateLayout();
 	});
@@ -337,7 +365,7 @@ void FAnimationTickRecordsTrack::FindTickRecordMessage(const FTimingEventSearchP
 			{
 				Trace::FAnalysisSessionReadScope SessionReadScope(SharedData.GetAnalysisSession());
 
-				AnimationProvider->ReadTickRecordTimeline(GetGameplayTrack().GetObjectId(), AssetId, NodeId, [this, &InContext](const FAnimationProvider::TickRecordTimeline& InTimeline)
+				AnimationProvider->ReadTickRecordTimeline(GetGameplayTrack().GetObjectId(), [this, &InContext](const FAnimationProvider::TickRecordTimeline& InTimeline)
 				{
 					InTimeline.EnumerateEvents(InContext.GetParameters().StartTime, InContext.GetParameters().EndTime, [&InContext](double InEventStartTime, double InEventEndTime, uint32 InDepth, const FTickRecordMessage& InMessage)
 					{
@@ -364,60 +392,51 @@ void FAnimationTickRecordsTrack::FindTickRecordMessage(const FTimingEventSearchP
 void FAnimationTickRecordsTrack::BuildContextMenu(FMenuBuilder& MenuBuilder)
 {
 #if WITH_EDITOR
-	MenuBuilder.BeginSection("TrackActions", LOCTEXT("TrackActionsMenuHeader", "Track Actions"));
+	if(CachedHoveredEvent.IsValid())
 	{
-		MenuBuilder.AddMenuEntry
-		(
-			LOCTEXT("FindAssetPlayerNode", "Find Asset Player Node"),
-			LOCTEXT("FindAssetPlayerNode_Tooltip", "Open the animation blueprint that this animation was played from."),
-			FSlateIcon(),
-			FUIAction(
-				FExecuteAction::CreateLambda([this]()
-				{
-					if(InstanceClass.LoadSynchronous())
+		int32 NodeId = StaticCastSharedRef<const FTickRecordSeries>(CachedHoveredEvent.Pin()->GetSeries())->NodeId;
+
+		MenuBuilder.BeginSection("TrackActions", LOCTEXT("TrackActionsMenuHeader", "Track Actions"));
+		{
+			MenuBuilder.AddMenuEntry
+			(
+				LOCTEXT("FindAssetPlayerNode", "Find Asset Player Node"),
+				LOCTEXT("FindAssetPlayerNode_Tooltip", "Open the animation blueprint that this animation was played from."),
+				FSlateIcon(),
+				FUIAction(
+					FExecuteAction::CreateLambda([this, NodeId]()
 					{
-						if(UAnimBlueprint* AnimBlueprint = Cast<UAnimBlueprint>(InstanceClass.Get()->ClassGeneratedBy))
+						if(InstanceClass.LoadSynchronous())
 						{
-							GEditor->GetEditorSubsystem<UAssetEditorSubsystem>()->OpenEditorForAsset(AnimBlueprint);
-
-							if(IAnimationBlueprintEditor* AnimBlueprintEditor = static_cast<IAnimationBlueprintEditor*>(GEditor->GetEditorSubsystem<UAssetEditorSubsystem>()->FindEditorForAsset(AnimBlueprint, true)))
+							if(UAnimBlueprint* AnimBlueprint = Cast<UAnimBlueprint>(InstanceClass.Get()->ClassGeneratedBy))
 							{
-								int32 AnimNodeIndex = InstanceClass.Get()->AnimNodeProperties.Num() - NodeId - 1;
-								TWeakObjectPtr<const UEdGraphNode>* GraphNode = InstanceClass.Get()->AnimBlueprintDebugData.NodePropertyIndexToNodeMap.Find(AnimNodeIndex);
-								if(GraphNode != nullptr && GraphNode->Get())
+								GEditor->GetEditorSubsystem<UAssetEditorSubsystem>()->OpenEditorForAsset(AnimBlueprint);
+
+								if(IAnimationBlueprintEditor* AnimBlueprintEditor = static_cast<IAnimationBlueprintEditor*>(GEditor->GetEditorSubsystem<UAssetEditorSubsystem>()->FindEditorForAsset(AnimBlueprint, true)))
 								{
-									AnimBlueprintEditor->JumpToHyperlink(GraphNode->Get());
+									int32 AnimNodeIndex = InstanceClass.Get()->AnimNodeProperties.Num() - NodeId - 1;
+									TWeakObjectPtr<const UEdGraphNode>* GraphNode = InstanceClass.Get()->AnimBlueprintDebugData.NodePropertyIndexToNodeMap.Find(AnimNodeIndex);
+									if(GraphNode != nullptr && GraphNode->Get())
+									{
+										AnimBlueprintEditor->JumpToHyperlink(GraphNode->Get());
+									}
 								}
 							}
 						}
-					}
-				})
-			),
-			NAME_None,
-			EUserInterfaceActionType::Button
-		);
+					})
+				),
+				NAME_None,
+				EUserInterfaceActionType::Button
+			);
+		}
+		MenuBuilder.EndSection();
 	}
-	MenuBuilder.EndSection();
 #endif
 
 	FGameplayGraphTrack::BuildContextMenu(MenuBuilder);
 }
 
-FText FAnimationTickRecordsTrack::MakeTrackName(const FGameplaySharedData& InSharedData, uint64 InAssetID, const TCHAR* InName) const
-{
-	FText AssetTypeName = LOCTEXT("UnknownAsset", "Unknown");
-
-	const FGameplayProvider* GameplayProvider = InSharedData.GetAnalysisSession().ReadProvider<FGameplayProvider>(FGameplayProvider::ProviderName);
-	if(GameplayProvider)
-	{
-		const FClassInfo& ClassInfo = GameplayProvider->GetClassInfoFromObject(InAssetID);
-		AssetTypeName = FText::FromString(ClassInfo.Name);
-	}
-
-	return FText::Format(LOCTEXT("AnimationTickRecordsTrackName", "{0} - {1}"), AssetTypeName, FText::FromString(InName));
-}
-
-void FAnimationTickRecordsTrack::GetVariantsAtTime(double InTime, TArray<TSharedRef<FVariantTreeNode>>& OutVariants) const
+void FAnimationTickRecordsTrack::GetVariantsAtFrame(const Trace::FFrame& InFrame, TArray<TSharedRef<FVariantTreeNode>>& OutVariants) const
 {
 	const FGameplayProvider* GameplayProvider = SharedData.GetAnalysisSession().ReadProvider<FGameplayProvider>(FGameplayProvider::ProviderName);
 	const FAnimationProvider* AnimationProvider = SharedData.GetAnalysisSession().ReadProvider<FAnimationProvider>(FAnimationProvider::ProviderName);
@@ -425,33 +444,26 @@ void FAnimationTickRecordsTrack::GetVariantsAtTime(double InTime, TArray<TShared
 	{
 		Trace::FAnalysisSessionReadScope SessionReadScope(SharedData.GetAnalysisSession());
 
-		const FClassInfo& ClassInfo = GameplayProvider->GetClassInfoFromObject(AssetId);
-		TSharedRef<FVariantTreeNode> Header = OutVariants.Add_GetRef(FVariantTreeNode::MakeObject(FText::FromString(ClassInfo.Name), AssetId));
-
-		const Trace::IFrameProvider& FramesProvider = Trace::ReadFrameProvider(SharedData.GetAnalysisSession());
-
-		AnimationProvider->ReadTickRecordTimeline(GetGameplayTrack().GetObjectId(), AssetId, NodeId, [&Header, &FramesProvider, &InTime](const FAnimationProvider::TickRecordTimeline& InTimeline)
+		AnimationProvider->ReadTickRecordTimeline(GetGameplayTrack().GetObjectId(), [&GameplayProvider, &OutVariants, &InFrame](const FAnimationProvider::TickRecordTimeline& InTimeline)
 		{
-			// round to nearest frame boundary
-			Trace::FFrame Frame;
-			if(FramesProvider.GetFrameFromTime(ETraceFrameType::TraceFrameType_Game, InTime, Frame))
+			InTimeline.EnumerateEvents(InFrame.StartTime, InFrame.EndTime, [&GameplayProvider, &OutVariants, &InFrame](double InStartTime, double InEndTime, uint32 InDepth, const FTickRecordMessage& InMessage)
 			{
-				InTimeline.EnumerateEvents(Frame.StartTime, Frame.EndTime, [&Header](double InEventStartTime, double InEventEndTime, uint32 InDepth, const FTickRecordMessage& InMessage)
+				if(InStartTime >= InFrame.StartTime && InEndTime <= InFrame.EndTime)
 				{
-					if(Header->GetChildren().Num() == 0)
+					const FClassInfo& ClassInfo = GameplayProvider->GetClassInfoFromObject(InMessage.AssetId);
+					TSharedRef<FVariantTreeNode> Header = OutVariants.Add_GetRef(FVariantTreeNode::MakeObject(FText::FromString(ClassInfo.Name), InMessage.AssetId));
+
+					Header->AddChild(FVariantTreeNode::MakeFloat(LOCTEXT("BlendWeight", "Blend Weight"), InMessage.BlendWeight));
+					Header->AddChild(FVariantTreeNode::MakeFloat(LOCTEXT("PlaybackTime", "Playback Time"), InMessage.PlaybackTime));
+					Header->AddChild(FVariantTreeNode::MakeFloat(LOCTEXT("RootMotionWeight", "Root Motion Weight"), InMessage.RootMotionWeight));
+					Header->AddChild(FVariantTreeNode::MakeFloat(LOCTEXT("PlayRate", "Play Rate"), InMessage.PlayRate));
+					if(InMessage.bIsBlendSpace)
 					{
-						Header->AddChild(FVariantTreeNode::MakeFloat(LOCTEXT("BlendWeight", "Blend Weight"), InMessage.BlendWeight));
-						Header->AddChild(FVariantTreeNode::MakeFloat(LOCTEXT("PlaybackTime", "Playback Time"), InMessage.PlaybackTime));
-						Header->AddChild(FVariantTreeNode::MakeFloat(LOCTEXT("RootMotionWeight", "Root Motion Weight"), InMessage.RootMotionWeight));
-						Header->AddChild(FVariantTreeNode::MakeFloat(LOCTEXT("PlayRate", "Play Rate"), InMessage.PlayRate));
-						if(InMessage.bIsBlendSpace)
-						{
-							Header->AddChild(FVariantTreeNode::MakeFloat(LOCTEXT("BlendSpacePositionX", "Blend Space Position X"), InMessage.BlendSpacePositionX));
-							Header->AddChild(FVariantTreeNode::MakeFloat(LOCTEXT("BlendSpacePositionY", "Blend Space Position Y"), InMessage.BlendSpacePositionY));
-						}
+						Header->AddChild(FVariantTreeNode::MakeFloat(LOCTEXT("BlendSpacePositionX", "Blend Space Position X"), InMessage.BlendSpacePositionX));
+						Header->AddChild(FVariantTreeNode::MakeFloat(LOCTEXT("BlendSpacePositionY", "Blend Space Position Y"), InMessage.BlendSpacePositionY));
 					}
-				});
-			}
+				}
+			});
 		});
 	}
 }
diff --git a/Engine/Plugins/Animation/GameplayInsights/Source/GameplayInsights/Private/AnimationTickRecordsTrack.h b/Engine/Plugins/Animation/GameplayInsights/Source/GameplayInsights/Private/AnimationTickRecordsTrack.h
--- a/Engine/Plugins/Animation/GameplayInsights/Source/GameplayInsights/Private/AnimationTickRecordsTrack.h
+++ b/Engine/Plugins/Animation/GameplayInsights/Source/GameplayInsights/Private/AnimationTickRecordsTrack.h
@@ -11,6 +11,7 @@ class FGameplaySharedData;
 class UAnimBlueprintGeneratedClass;
 class FAnimationTickRecordsTrack;
 class FTimingTrackViewport;
+class FGraphTrackEvent;
 
 class FTickRecordSeries : public FGameplayGraphSeries
 {
@@ -30,6 +31,12 @@ public:
 	virtual FString FormatValue(double Value) const override;
 
 public:
+	/** The asset ID that this series uses */
+	uint64 AssetId;
+
+	/** The node ID that this series' data comes from */
+	int32 NodeId;
+
 	ESeriesType Type;
 };
 
@@ -38,29 +45,21 @@ class FAnimationTickRecordsTrack : public FGameplayGraphTrack
 	INSIGHTS_DECLARE_RTTI(FAnimationTickRecordsTrack, FGameplayGraphTrack)
 
 public:
-	FAnimationTickRecordsTrack(const FAnimationSharedData& InSharedData, uint64 InObjectId, uint64 InAssetId, int32 InNodeId, const TCHAR* InName);
+	FAnimationTickRecordsTrack(const FAnimationSharedData& InSharedData, uint64 InObjectId, const TCHAR* InName);
 
+	virtual void PostUpdate(const ITimingTrackUpdateContext& Context) override;
 	virtual void InitTooltip(FTooltipDrawState& Tooltip, const ITimingEvent& HoveredTimingEvent) const override;
 	virtual const TSharedPtr<const ITimingEvent> SearchEvent(const FTimingEventSearchParameters& InSearchParameters) const override;
 	virtual void BuildContextMenu(FMenuBuilder& MenuBuilder) override;
 	virtual bool UpdateSeriesBounds(FGameplayGraphSeries& InSeries, const FTimingTrackViewport& InViewport) override;
 	virtual void UpdateSeries(FGameplayGraphSeries& InSeries, const FTimingTrackViewport& InViewport) override;
 	virtual void AddAllSeries() override;
-	virtual void GetVariantsAtTime(double InTime, TArray<TSharedRef<FVariantTreeNode>>& OutVariants) const override;
-
-	/** Get the asset ID that this track uses */
-	uint64 GetAssetId() const { return AssetId; }
-
-	/** Get the node ID that this track uses */
-	int32 GetNodeId() const { return NodeId; }
+	virtual void GetVariantsAtFrame(const Trace::FFrame& InFrame, TArray<TSharedRef<FVariantTreeNode>>& OutVariants) const override;
 
 private:
 	// Helper function used to find a tick record
 	void FindTickRecordMessage(const FTimingEventSearchParameters& InParameters, TFunctionRef<void(double, double, uint32, const FTickRecordMessage&)> InFoundPredicate) const;
 
-	// Helper used to build track name
-	FText MakeTrackName(const FGameplaySharedData& InSharedData, uint64 InAssetID, const TCHAR* InName) const;
-
 	// Helper function for series bounds update
 	template<typename ProjectionType>
 	bool UpdateSeriesBoundsHelper(FTickRecordSeries& InSeries, const FTimingTrackViewport& InViewport, ProjectionType Projection);
@@ -78,13 +77,6 @@ private:
 	TSoftObjectPtr<UAnimBlueprintGeneratedClass> InstanceClass;
 #endif
 
-	/** Colors for line drawing */
-	FLinearColor MainSeriesLineColor;
-	FLinearColor MainSeriesFillColor;
-
-	/** The asset ID that this track uses */
-	uint64 AssetId;
-
-	/** The node ID that this track's data comes from */
-	int32 NodeId;
+	/** Cached hovered event */
+	TWeakPtr<const FGraphTrackEvent> CachedHoveredEvent;
 };
\ No newline at end of file
diff --git a/Engine/Plugins/Animation/GameplayInsights/Source/GameplayInsights/Private/GameplayAnalyzer.cpp b/Engine/Plugins/Animation/GameplayInsights/Source/GameplayInsights/Private/GameplayAnalyzer.cpp
--- a/Engine/Plugins/Animation/GameplayInsights/Source/GameplayInsights/Private/GameplayAnalyzer.cpp
+++ b/Engine/Plugins/Animation/GameplayInsights/Source/GameplayInsights/Private/GameplayAnalyzer.cpp
@@ -17,6 +17,7 @@ void FGameplayAnalyzer::OnAnalysisBegin(const FOnAnalysisContext& Context)
 	Builder.RouteEvent(RouteId_Class, "Object", "Class");
 	Builder.RouteEvent(RouteId_Object, "Object", "Object");
 	Builder.RouteEvent(RouteId_ObjectEvent, "Object", "ObjectEvent");
+	Builder.RouteEvent(RouteId_World, "Object", "World");
 }
 
 bool FGameplayAnalyzer::OnEvent(uint16 RouteId, const FOnEventContext& Context)
@@ -26,6 +27,16 @@ bool FGameplayAnalyzer::OnEvent(uint16 RouteId, const FOnEventContext& Context)
 	const auto& EventData = Context.EventData;
 	switch (RouteId)
 	{
+	case RouteId_World:
+	{
+		uint64 Id = EventData.GetValue<uint64>("Id");
+		int32 PIEInstanceId = EventData.GetValue<int32>("PIEInstanceId");
+		uint8 Type = EventData.GetValue<uint8>("Type");
+		uint8 NetMode = EventData.GetValue<uint8>("NetMode");
+		bool bIsSimulating = EventData.GetValue<bool>("IsSimulating");
+		GameplayProvider.AppendWorld(Id, PIEInstanceId, Type, NetMode, bIsSimulating);
+		break;
+	}
 	case RouteId_Class:
 	{
 		const TCHAR* ClassNameAndPathName = reinterpret_cast<const TCHAR*>(EventData.GetAttachment());
diff --git a/Engine/Plugins/Animation/GameplayInsights/Source/GameplayInsights/Private/GameplayAnalyzer.h b/Engine/Plugins/Animation/GameplayInsights/Source/GameplayInsights/Private/GameplayAnalyzer.h
--- a/Engine/Plugins/Animation/GameplayInsights/Source/GameplayInsights/Private/GameplayAnalyzer.h
+++ b/Engine/Plugins/Animation/GameplayInsights/Source/GameplayInsights/Private/GameplayAnalyzer.h
@@ -19,6 +19,7 @@ public:
 private:
 	enum : uint16
 	{
+		RouteId_World,
 		RouteId_Class,
 		RouteId_Object,
 		RouteId_ObjectEvent,
diff --git a/Engine/Plugins/Animation/GameplayInsights/Source/GameplayInsights/Private/GameplayGraphTrack.cpp b/Engine/Plugins/Animation/GameplayInsights/Source/GameplayInsights/Private/GameplayGraphTrack.cpp
--- a/Engine/Plugins/Animation/GameplayInsights/Source/GameplayInsights/Private/GameplayGraphTrack.cpp
+++ b/Engine/Plugins/Animation/GameplayInsights/Source/GameplayInsights/Private/GameplayGraphTrack.cpp
@@ -14,6 +14,7 @@
 #include "Insights/Common/PaintUtils.h"
 #include "Insights/ViewModels/ITimingViewDrawHelper.h"
 #include "Fonts/FontMeasure.h"
+#include "TraceServices/Model/Frames.h"
 
 #define LOCTEXT_NAMESPACE "GameplayGraphTrack"
 
@@ -80,34 +81,8 @@ void FGameplayGraphTrack::UpdateTrackHeight(const ITimingTrackUpdateContext& Con
 	const float TimelineDY2 = 2.0f * Viewport.GetLayout().TimelineDY;
 	const float DesiredTrackHeight = FMath::Max(0.0f, ((Viewport.GetLayout().ComputeTrackHeight(NumLanes) - TimelineDY2) * RequestedTrackSizeScale) + TimelineDY2);
 
-	if (CurrentTrackHeight < DesiredTrackHeight)
-	{
-		float NewTrackHeight;
-		if (Viewport.IsDirty(ETimingTrackViewportDirtyFlags::VLayoutChanged))
-		{
-			NewTrackHeight = DesiredTrackHeight;
-		}
-		else
-		{
-			NewTrackHeight = FMath::CeilToFloat(CurrentTrackHeight * 0.9f + DesiredTrackHeight * 0.1f);
-		}
-
-		SetHeight(NewTrackHeight);
-	}
-	else if (CurrentTrackHeight > DesiredTrackHeight)
-	{
-		float NewTrackHeight;
-		if (Viewport.IsDirty(ETimingTrackViewportDirtyFlags::VLayoutChanged))
-		{
-			NewTrackHeight = DesiredTrackHeight;
-		}
-		else
-		{
-			NewTrackHeight = FMath::FloorToFloat(CurrentTrackHeight * 0.9f + DesiredTrackHeight * 0.1f);
-		}
-
-		SetHeight(NewTrackHeight);
-	}
+	// Dont interpolate here, it is too expensive to continually invalidate the track right now
+	SetHeight(DesiredTrackHeight);
 }
 
 void FGameplayGraphTrack::UpdateSeriesInternal(FGameplayGraphSeries& InSeries, const FTimingTrackViewport& InViewport, int32 InActiveSeriesIndex)
@@ -326,4 +301,9 @@ void FGameplayGraphTrack::BuildContextMenu(FMenuBuilder& MenuBuilder)
 	MenuBuilder.EndSection();
 }
 
+void FGameplayGraphTrack::GetVariantsAtFrame(const Trace::FFrame& InFrame, TArray<TSharedRef<FVariantTreeNode>>& OutVariants) const
+{ 
+	GetVariantsAtTime(InFrame.StartTime, OutVariants); 
+}
+
 #undef LOCTEXT_NAMESPACE
diff --git a/Engine/Plugins/Animation/GameplayInsights/Source/GameplayInsights/Private/GameplayGraphTrack.h b/Engine/Plugins/Animation/GameplayInsights/Source/GameplayInsights/Private/GameplayGraphTrack.h
--- a/Engine/Plugins/Animation/GameplayInsights/Source/GameplayInsights/Private/GameplayGraphTrack.h
+++ b/Engine/Plugins/Animation/GameplayInsights/Source/GameplayInsights/Private/GameplayGraphTrack.h
@@ -9,6 +9,7 @@
 class FTimingEventSearchParameters;
 class FGameplayGraphTrack;
 struct FVariantTreeNode;
+namespace Trace { struct FFrame; }
 
 /** The various layouts that we display series with */
 enum class EGameplayGraphLayout : int32
@@ -50,8 +51,9 @@ public:
 	virtual void PreUpdate(const ITimingTrackUpdateContext& Context) override;
 	virtual void BuildContextMenu(FMenuBuilder& MenuBuilder) override;
 
-	// Get all variants at the specified time
+	// Get all variants at the specified time/frame
 	virtual void GetVariantsAtTime(double InTime, TArray<TSharedRef<FVariantTreeNode>>& OutVariants) const {}
+	virtual void GetVariantsAtFrame(const Trace::FFrame& InFrame, TArray<TSharedRef<FVariantTreeNode>>& OutVariants) const;
 
 	// Get the requested track size scale
 	float GetRequestedTrackSizeScale() const { return RequestedTrackSizeScale; }
diff --git a/Engine/Plugins/Animation/GameplayInsights/Source/GameplayInsights/Private/GameplayProvider.cpp b/Engine/Plugins/Animation/GameplayInsights/Source/GameplayInsights/Private/GameplayProvider.cpp
--- a/Engine/Plugins/Animation/GameplayInsights/Source/GameplayInsights/Private/GameplayProvider.cpp
+++ b/Engine/Plugins/Animation/GameplayInsights/Source/GameplayInsights/Private/GameplayProvider.cpp
@@ -65,6 +65,19 @@ const FClassInfo* FGameplayProvider::FindClassInfo(uint64 InClassId) const
 	return nullptr;
 }
 
+const FClassInfo* FGameplayProvider::FindClassInfo(const TCHAR* InClassPath) const
+{
+	Session.ReadAccessCheck();
+
+	const int32* ClassIndex = ClassPathNameToIndexMap.Find(InClassPath);
+	if (ClassIndex != nullptr)
+	{
+		return &ClassInfos[*ClassIndex];
+	}
+
+	return nullptr;
+}
+
 const FObjectInfo* FGameplayProvider::FindObjectInfo(uint64 InObjectId) const
 {
 	Session.ReadAccessCheck();
@@ -78,6 +91,52 @@ const FObjectInfo* FGameplayProvider::FindObjectInfo(uint64 InObjectId) const
 	return nullptr;
 }
 
+const FWorldInfo* FGameplayProvider::FindWorldInfo(uint64 InObjectId) const
+{
+	Session.ReadAccessCheck();
+
+	const int32* WorldIndex = WorldIdToIndexMap.Find(InObjectId);
+	if (WorldIndex != nullptr)
+	{
+		return &WorldInfos[*WorldIndex];
+	}
+
+	return nullptr;
+}
+
+const FWorldInfo* FGameplayProvider::FindWorldInfoFromObject(uint64 InObjectId) const
+{
+	const FClassInfo* WorldClass = FindClassInfo(TEXT("/Script/Engine.World"));
+	if(WorldClass)
+	{
+		// Traverse outer chain until we find a world
+		const FObjectInfo* ObjectInfo = FindObjectInfo(InObjectId);
+		while (ObjectInfo != nullptr)
+		{
+			if (ObjectInfo->ClassId == WorldClass->Id)
+			{
+				return FindWorldInfo(ObjectInfo->Id);
+			}
+
+			ObjectInfo = FindObjectInfo(ObjectInfo->OuterId);
+		}
+	}
+
+	return nullptr;
+}
+
+bool FGameplayProvider::IsWorld(uint64 InObjectId) const
+{
+	const FClassInfo* WorldClass = FindClassInfo(TEXT("/Script/Engine.World"));
+	if (WorldClass)
+	{
+		const FObjectInfo* ObjectInfo = FindObjectInfo(InObjectId);
+		return ObjectInfo->ClassId == WorldClass->Id;
+	}
+
+	return false;
+}
+
 const FClassInfo& FGameplayProvider::GetClassInfo(uint64 InClassId) const
 {
 	const FClassInfo* ClassInfo = FindClassInfo(InClassId);
@@ -138,6 +197,7 @@ void FGameplayProvider::AppendClass(uint64 InClassId, uint64 InSuperId, const TC
 
 		int32 NewClassInfoIndex = ClassInfos.Add(NewClassInfo);
 		ClassIdToIndexMap.Add(InClassId, NewClassInfoIndex);
+		ClassPathNameToIndexMap.Add(NewClassPathName, NewClassInfoIndex);
 	}
 }
 
@@ -202,4 +262,22 @@ void FGameplayProvider::AppendObjectEvent(uint64 InObjectId, double InTime, cons
 	Session.UpdateDurationSeconds(InTime);
 }
 
+void FGameplayProvider::AppendWorld(uint64 InObjectId, int32 InPIEInstanceId, uint8 InType, uint8 InNetMode, bool bInIsSimulating)
+{
+	Session.WriteAccessCheck();
+
+	if (WorldIdToIndexMap.Find(InObjectId) == nullptr)
+	{
+		FWorldInfo NewWorldInfo;
+		NewWorldInfo.Id = InObjectId;
+		NewWorldInfo.PIEInstanceId = InPIEInstanceId;
+		NewWorldInfo.Type = (FWorldInfo::EType)InType;
+		NewWorldInfo.NetMode = (FWorldInfo::ENetMode)InNetMode;
+		NewWorldInfo.bIsSimulating = bInIsSimulating;
+
+		int32 NewWorldInfoIndex = WorldInfos.Add(NewWorldInfo);
+		WorldIdToIndexMap.Add(InObjectId, NewWorldInfoIndex);
+	}
+}
+
 #undef LOCTEXT_NAMESPACE
\ No newline at end of file
diff --git a/Engine/Plugins/Animation/GameplayInsights/Source/GameplayInsights/Private/GameplayProvider.h b/Engine/Plugins/Animation/GameplayInsights/Source/GameplayInsights/Private/GameplayProvider.h
--- a/Engine/Plugins/Animation/GameplayInsights/Source/GameplayInsights/Private/GameplayProvider.h
+++ b/Engine/Plugins/Animation/GameplayInsights/Source/GameplayInsights/Private/GameplayProvider.h
@@ -21,7 +21,11 @@ public:
 	virtual bool ReadObjectEvent(uint64 InObjectId, uint64 InMessageId, TFunctionRef<void(const FObjectEventMessage&)> Callback) const override;
 	virtual void EnumerateObjects(TFunctionRef<void(const FObjectInfo&)> Callback) const override;
 	virtual const FClassInfo* FindClassInfo(uint64 InClassId) const override;
+	virtual const FClassInfo* FindClassInfo(const TCHAR* InClassPath) const override;
 	virtual const FObjectInfo* FindObjectInfo(uint64 InObjectId) const override;
+	virtual const FWorldInfo* FindWorldInfo(uint64 InObjectId) const override;
+	virtual const FWorldInfo* FindWorldInfoFromObject(uint64 InObjectId) const override;
+	virtual bool IsWorld(uint64 InObjectId) const override;
 	virtual const FClassInfo& GetClassInfo(uint64 InClassId) const override;
 	virtual const FClassInfo& GetClassInfoFromObject(uint64 InObjectId) const override;
 	virtual const FObjectInfo& GetObjectInfo(uint64 InObjectId) const override;
@@ -36,6 +40,9 @@ public:
 	/** Add an object event message */
 	void AppendObjectEvent(uint64 InObjectId, double InTime, const TCHAR* InEvent);
 
+	/** Add a world message */
+	void AppendWorld(uint64 InObjectId, int32 InPIEInstanceId, uint8 InType, uint8 InNetMode, bool bInIsSimulating);
+
 private:
 	Trace::IAnalysisSession& Session;
 
@@ -45,18 +52,27 @@ private:
 	/** All object info, grow only for stable indices */
 	TArray<FObjectInfo> ObjectInfos;
 
+	/** All world info, grow only for stable indices */
+	TArray<FWorldInfo> WorldInfos;
+
 	/** Classes that are in use. Map from Id to ClassInfo index */
 	TMap<uint64, int32> ClassIdToIndexMap;
 
 	/** Objects that are in use. Map from Id to ObjectInfo index */
 	TMap<uint64, int32> ObjectIdToIndexMap;
 
+	/** Worlds that are in use. Map from Id to WorldInfo index */
+	TMap<uint64, int32> WorldIdToIndexMap;
+
 	/** Map from object Id to timeline index */
 	TMap<uint64, uint32> ObjectIdToEventTimelines;
 
 	/** Message storage */
 	TArray<TSharedRef<Trace::TPointTimeline<FObjectEventMessage>>> EventTimelines;
 
+	/** Map of class path name to ClassInfo index */
+	TMap<FStringView, int32> ClassPathNameToIndexMap;
+
 	/** EndPlay event text */
 	const TCHAR* EndPlayEvent;
 
diff --git a/Engine/Plugins/Animation/GameplayInsights/Source/GameplayInsights/Private/GameplayTrack.cpp b/Engine/Plugins/Animation/GameplayInsights/Source/GameplayInsights/Private/GameplayTrack.cpp
--- a/Engine/Plugins/Animation/GameplayInsights/Source/GameplayInsights/Private/GameplayTrack.cpp
+++ b/Engine/Plugins/Animation/GameplayInsights/Source/GameplayInsights/Private/GameplayTrack.cpp
@@ -10,6 +10,10 @@
 #include "Insights/ViewModels/ITimingViewDrawHelper.h"
 #include "GameplaySharedData.h"
 #include "Framework/MultiBox/MultiBoxBuilder.h"
+#include "GameplayProvider.h"
+#include "TraceServices/Model/AnalysisSession.h"
+#include "IGameplayProvider.h"
+#include "TraceServices/Model/Frames.h"
 
 INSIGHTS_IMPLEMENT_RTTI(FGameplayTimingEventsTrack);
 
@@ -91,6 +95,62 @@ void FGameplayTrack::DrawHeaderForTimingTrack(const ITimingTrackDrawContext& InC
 	}
 }
 
+FText FGameplayTrack::GetWorldPostFix(const FWorldInfo& InWorldInfo)
+{
+	switch (InWorldInfo.Type)
+	{
+	default:
+		return LOCTEXT("Unknown", "Unknown");
+	case FWorldInfo::EType::None:
+		return LOCTEXT("None", "None");
+	case FWorldInfo::EType::PIE:
+	{
+		switch (InWorldInfo.NetMode)
+		{
+		case FWorldInfo::ENetMode::Client:
+			return FText::Format(LOCTEXT("ClientIndexFormat", "Client {0}"), InWorldInfo.PIEInstanceId - 1);
+		case FWorldInfo::ENetMode::DedicatedServer:
+			return LOCTEXT("ServerPostfix", "Server");
+		case FWorldInfo::ENetMode::ListenServer:
+			return LOCTEXT("ListenServerPostfix", "Listen Server");
+		case FWorldInfo::ENetMode::Standalone:
+			return InWorldInfo.bIsSimulating ? LOCTEXT("SimulateInEditorPostfix", "Simulate") : LOCTEXT("PlayInEditorPostfix", "PIE");
+		}
+	}
+	case FWorldInfo::EType::Editor:
+	case FWorldInfo::EType::EditorPreview:
+		return LOCTEXT("EditorPostfix", "Editor");
+	case FWorldInfo::EType::Game:
+	case FWorldInfo::EType::GamePreview:
+	case FWorldInfo::EType::GameRPC:
+		return LOCTEXT("GamePostfix", "Game");
+	case FWorldInfo::EType::Inactive:
+		return LOCTEXT("InactivePostfix", "Inactive");
+	}
+}
+
+FText FGameplayTrack::GetWorldName(const Trace::IAnalysisSession& InAnalysisSession) const
+{
+	const FGameplayProvider* GameplayProvider = InAnalysisSession.ReadProvider<FGameplayProvider>(FGameplayProvider::ProviderName);
+	if (GameplayProvider)
+	{
+		Trace::FAnalysisSessionReadScope SessionReadScope(InAnalysisSession);
+
+		if(const FWorldInfo* WorldInfo = GameplayProvider->FindWorldInfoFromObject(ObjectId))
+		{
+			const FObjectInfo& WorldObjectInfo = GameplayProvider->GetObjectInfo(WorldInfo->Id);
+			return FText::Format(LOCTEXT("WorldNameFormat", "{0} ({1})"), FText::FromString(WorldObjectInfo.Name), GetWorldPostFix(*WorldInfo));
+		}
+	}
+
+	return LOCTEXT("UnknownWorld", "Unknown");
+}
+
+void FGameplayTimingEventsTrack::GetVariantsAtFrame(const Trace::FFrame& InFrame, TArray<TSharedRef<FVariantTreeNode>>& OutVariants) const
+{ 
+	GetVariantsAtTime(InFrame.StartTime, OutVariants);
+}
+
 void FGameplayTimingEventsTrack::BuildContextMenu(FMenuBuilder& MenuBuilder)
 {
 	MenuBuilder.BeginSection("View", LOCTEXT("ViewHeader", "View"));
diff --git a/Engine/Plugins/Animation/GameplayInsights/Source/GameplayInsights/Private/GameplayTrack.h b/Engine/Plugins/Animation/GameplayInsights/Source/GameplayInsights/Private/GameplayTrack.h
--- a/Engine/Plugins/Animation/GameplayInsights/Source/GameplayInsights/Private/GameplayTrack.h
+++ b/Engine/Plugins/Animation/GameplayInsights/Source/GameplayInsights/Private/GameplayTrack.h
@@ -11,6 +11,9 @@ namespace GameplayTrackConstants
 
 struct FVariantTreeNode;
 class FGameplaySharedData;
+struct FWorldInfo;
+namespace Trace { class IAnalysisSession; }
+namespace Trace { struct FFrame; }
 
 // Provides parent/child hierarchy structure and the owning object Id
 // Designed as a compositional member of outer timing tracks
@@ -52,6 +55,12 @@ public:
 	/** Access indent */
 	uint32 GetIndent() const { return Indent; }
 
+	/** Get the world name for this object */
+	FText GetWorldName(const Trace::IAnalysisSession& InAnalysisSession) const;
+
+	/** Helper function to find a postfix for a world (e.g. "Client 0", "Server" etc.) */
+	static FText GetWorldPostFix(const FWorldInfo& InWorldInfo);
+
 private:
 	/** Outer timing track */
 	FBaseTimingTrack& TimingTrack;
@@ -104,8 +113,9 @@ public:
 
 	virtual void BuildContextMenu(FMenuBuilder& MenuBuilder) override;
 
-	/** Get all variants at the specified time */
+	/** Get all variants at the specified time/frame */
 	virtual void GetVariantsAtTime(double InTime, TArray<TSharedRef<FVariantTreeNode>>& OutVariants) const {}
+	virtual void GetVariantsAtFrame(const Trace::FFrame& InFrame, TArray<TSharedRef<FVariantTreeNode>>& OutVariants) const;
 
 protected:
 	const FGameplaySharedData& GameplaySharedData;
diff --git a/Engine/Plugins/Animation/GameplayInsights/Source/GameplayInsights/Private/MontageTrack.cpp b/Engine/Plugins/Animation/GameplayInsights/Source/GameplayInsights/Private/MontageTrack.cpp
--- a/Engine/Plugins/Animation/GameplayInsights/Source/GameplayInsights/Private/MontageTrack.cpp
+++ b/Engine/Plugins/Animation/GameplayInsights/Source/GameplayInsights/Private/MontageTrack.cpp
@@ -115,12 +115,12 @@ void FMontageTrack::UpdateSeries(FGameplayGraphSeries& InSeries, const FTimingTr
 
 		AnimationProvider->ReadMontageTimeline(GetGameplayTrack().GetObjectId(), [this, &Builder, &InViewport, &MontageSeries, &AnimationProvider](const FAnimationProvider::AnimMontageTimeline& InTimeline)
 		{
-			int32 FrameCounter = 0;
-			int32 LastFrameWithMontage = 0;
+			uint16 FrameCounter = 0;
+			uint16 LastFrameWithMontage = 0;
 
 			InTimeline.EnumerateEvents(InViewport.GetStartTime(), InViewport.GetEndTime(), [this, &FrameCounter, &LastFrameWithMontage, &Builder, &MontageSeries, &AnimationProvider](double InStartTime, double InEndTime, uint32 InDepth, const FAnimMontageMessage& InMessage)
 			{
-				FrameCounter++;
+				FrameCounter = InMessage.FrameCounter;
 
 				if(InMessage.MontageId == MontageSeries.MontageId)
 				{
@@ -169,6 +169,8 @@ void FMontageTrack::InitTooltip(FTooltipDrawState& Tooltip, const ITimingEvent&
 			}
 		}
 
+		Tooltip.AddNameValueTextLine(LOCTEXT("EventWorld", "World").ToString(), GetGameplayTrack().GetWorldName(SharedData.GetAnalysisSession()).ToString());
+
 		Tooltip.UpdateLayout();
 	});
 }
@@ -208,13 +210,21 @@ void FMontageTrack::FindMontageMessage(const FTimingEventSearchParameters& InPar
 			}
 		},
 
+		[&InParameters](double InFoundStartTime, double InFoundEndTime, uint32 InFoundDepth, const FAnimMontageMessage& InEvent)
+		{
+			// Match the start time exactly here
+			return InFoundStartTime == InParameters.StartTime;
+		},
+
 		[&InFoundPredicate](double InFoundStartTime, double InFoundEndTime, uint32 InFoundDepth, const FAnimMontageMessage& InEvent)
 		{
 			InFoundPredicate(InFoundStartTime, InFoundEndTime, InFoundDepth, InEvent);
-		});
+		},
+		
+		TTimingEventSearch<FAnimMontageMessage>::NoMatch);
 }
 
-void FMontageTrack::GetVariantsAtTime(double InTime, TArray<TSharedRef<FVariantTreeNode>>& OutVariants) const 
+void FMontageTrack::GetVariantsAtFrame(const Trace::FFrame& InFrame, TArray<TSharedRef<FVariantTreeNode>>& OutVariants) const 
 {
 	TSharedRef<FVariantTreeNode> Header = OutVariants.Add_GetRef(FVariantTreeNode::MakeHeader(LOCTEXT("MontagesHeader", "Montages")));
 
@@ -226,29 +236,24 @@ void FMontageTrack::GetVariantsAtTime(double InTime, TArray<TSharedRef<FVariantT
 	{
 		Trace::FAnalysisSessionReadScope SessionReadScope(SharedData.GetAnalysisSession());
 
-		AnimationProvider->ReadMontageTimeline(GetGameplayTrack().GetObjectId(), [this, &InTime, &AnimationProvider, &FramesProvider, &GameplayProvider, &Header](const FAnimationProvider::AnimMontageTimeline& InTimeline)
+		AnimationProvider->ReadMontageTimeline(GetGameplayTrack().GetObjectId(), [this, &InFrame, &AnimationProvider, &GameplayProvider, &Header](const FAnimationProvider::AnimMontageTimeline& InTimeline)
 		{
-			// round to nearest frame boundary
-			Trace::FFrame Frame;
-			if(FramesProvider.GetFrameFromTime(ETraceFrameType::TraceFrameType_Game, InTime, Frame))
+			InTimeline.EnumerateEvents(InFrame.StartTime, InFrame.EndTime, [this, &AnimationProvider, &GameplayProvider, &Header, &InFrame](double InStartTime, double InEndTime, uint32 InDepth, const FAnimMontageMessage& InMessage)
 			{
-				InTimeline.EnumerateEvents(Frame.StartTime, Frame.EndTime, [this, &AnimationProvider, &GameplayProvider, &Header, &Frame](double InStartTime, double InEndTime, uint32 InDepth, const FAnimMontageMessage& InMessage)
+				if(InStartTime >= InFrame.StartTime && InEndTime <= InFrame.EndTime)
 				{
-					if(InStartTime >= Frame.StartTime && InEndTime <= Frame.EndTime)
-					{
-						const FObjectInfo& MontageInfo = GameplayProvider->GetObjectInfo(InMessage.MontageId);
-						TSharedRef<FVariantTreeNode> MontageHeader = Header->AddChild(FVariantTreeNode::MakeObject(FText::FromString(MontageInfo.Name), InMessage.MontageId));
+					const FObjectInfo& MontageInfo = GameplayProvider->GetObjectInfo(InMessage.MontageId);
+					TSharedRef<FVariantTreeNode> MontageHeader = Header->AddChild(FVariantTreeNode::MakeObject(FText::FromString(MontageInfo.Name), InMessage.MontageId));
 
-						MontageHeader->AddChild(FVariantTreeNode::MakeFloat(LOCTEXT("EventWeight", "Weight"), InMessage.Weight));
-						MontageHeader->AddChild(FVariantTreeNode::MakeFloat(LOCTEXT("EventDesiredWeight", "Desired Weight"), InMessage.DesiredWeight));
+					MontageHeader->AddChild(FVariantTreeNode::MakeFloat(LOCTEXT("EventWeight", "Weight"), InMessage.Weight));
+					MontageHeader->AddChild(FVariantTreeNode::MakeFloat(LOCTEXT("EventDesiredWeight", "Desired Weight"), InMessage.DesiredWeight));
 
-						const TCHAR* CurrentSectionName = AnimationProvider->GetName(InMessage.CurrentSectionNameId);
-						MontageHeader->AddChild(FVariantTreeNode::MakeString(LOCTEXT("CurrentSectionName", "Current Section"), CurrentSectionName));
-						const TCHAR* NextSectionName = AnimationProvider->GetName(InMessage.NextSectionNameId);
-						MontageHeader->AddChild(FVariantTreeNode::MakeString(LOCTEXT("NextSectionName", "Next Section"), NextSectionName));
-					}
-				});
-			}
+					const TCHAR* CurrentSectionName = AnimationProvider->GetName(InMessage.CurrentSectionNameId);
+					MontageHeader->AddChild(FVariantTreeNode::MakeString(LOCTEXT("CurrentSectionName", "Current Section"), CurrentSectionName));
+					const TCHAR* NextSectionName = AnimationProvider->GetName(InMessage.NextSectionNameId);
+					MontageHeader->AddChild(FVariantTreeNode::MakeString(LOCTEXT("NextSectionName", "Next Section"), NextSectionName));
+				}
+			});
 		});
 	}
 }
diff --git a/Engine/Plugins/Animation/GameplayInsights/Source/GameplayInsights/Private/MontageTrack.h b/Engine/Plugins/Animation/GameplayInsights/Source/GameplayInsights/Private/MontageTrack.h
--- a/Engine/Plugins/Animation/GameplayInsights/Source/GameplayInsights/Private/MontageTrack.h
+++ b/Engine/Plugins/Animation/GameplayInsights/Source/GameplayInsights/Private/MontageTrack.h
@@ -27,7 +27,7 @@ public:
 	virtual void AddAllSeries() override;
 	virtual bool UpdateSeriesBounds(FGameplayGraphSeries& InSeries, const FTimingTrackViewport& InViewport) override;
 	virtual void UpdateSeries(FGameplayGraphSeries& InSeries, const FTimingTrackViewport& InViewport) override;
-	virtual void GetVariantsAtTime(double InTime, TArray<TSharedRef<FVariantTreeNode>>& OutVariants) const override;
+	virtual void GetVariantsAtFrame(const Trace::FFrame& InFrame, TArray<TSharedRef<FVariantTreeNode>>& OutVariants) const override;
 
 private:
 	// Helper function used to find a montage message
diff --git a/Engine/Plugins/Animation/GameplayInsights/Source/GameplayInsights/Private/ObjectEventsTrack.cpp b/Engine/Plugins/Animation/GameplayInsights/Source/GameplayInsights/Private/ObjectEventsTrack.cpp
--- a/Engine/Plugins/Animation/GameplayInsights/Source/GameplayInsights/Private/ObjectEventsTrack.cpp
+++ b/Engine/Plugins/Animation/GameplayInsights/Source/GameplayInsights/Private/ObjectEventsTrack.cpp
@@ -16,9 +16,10 @@
 INSIGHTS_IMPLEMENT_RTTI(FObjectEventsTrack)
 
 FObjectEventsTrack::FObjectEventsTrack(const FGameplaySharedData& InSharedData, uint64 InObjectID, const TCHAR* InName)
-	: FGameplayTimingEventsTrack(InSharedData, InObjectID, MakeTrackName(InSharedData, InObjectID, InName))
+	: FGameplayTimingEventsTrack(InSharedData, InObjectID, FText::FromString(InName))
 	, SharedData(InSharedData)
 {
+	SetName(MakeTrackName(InSharedData, InObjectID, InName).ToString());
 }
 
 void FObjectEventsTrack::BuildDrawState(ITimingEventsTrackDrawStateBuilder& Builder, const ITimingTrackUpdateContext& Context)
@@ -40,6 +41,7 @@ void FObjectEventsTrack::BuildDrawState(ITimingEventsTrackDrawStateBuilder& Buil
 	}
 }
 
+
 void FObjectEventsTrack::Draw(const ITimingTrackDrawContext& Context) const
 {
 	DrawEvents(Context);
@@ -56,6 +58,7 @@ void FObjectEventsTrack::InitTooltip(FTooltipDrawState& Tooltip, const ITimingEv
 
 		Tooltip.AddTitle(FText::FromString(FString(InMessage.Name)).ToString());
 		Tooltip.AddNameValueTextLine(LOCTEXT("EventTime", "Time").ToString(), FText::AsNumber(HoveredTimingEvent.GetStartTime()).ToString());
+		Tooltip.AddNameValueTextLine(LOCTEXT("EventWorld", "World").ToString(), GetGameplayTrack().GetWorldName(SharedData.GetAnalysisSession()).ToString());
 
 		Tooltip.UpdateLayout();
 	});
@@ -104,13 +107,16 @@ void FObjectEventsTrack::FindObjectEvent(const FTimingEventSearchParameters& InP
 		});
 }
 
+
 FText FObjectEventsTrack::MakeTrackName(const FGameplaySharedData& InSharedData, uint64 InObjectID, const TCHAR* InName) const
 {
 	FText ClassName = LOCTEXT("UnknownClass", "Unknown");
 
 	const FGameplayProvider* GameplayProvider = InSharedData.GetAnalysisSession().ReadProvider<FGameplayProvider>(FGameplayProvider::ProviderName);
-	if(GameplayProvider)	
+	if(GameplayProvider)
 	{
+		Trace::FAnalysisSessionReadScope SessionReadScope(InSharedData.GetAnalysisSession());
+
 		if(const FObjectInfo* ObjectInfo = GameplayProvider->FindObjectInfo(InObjectID))
 		{
 			if(const FClassInfo* ClassInfo = GameplayProvider->FindClassInfo(ObjectInfo->ClassId))
@@ -118,34 +124,39 @@ FText FObjectEventsTrack::MakeTrackName(const FGameplaySharedData& InSharedData,
 				ClassName = FText::FromString(ClassInfo->Name);
 			}
 		}
+
+		FText ObjectName;
+		if (GameplayProvider->IsWorld(InObjectID))
+		{
+			ObjectName = GetGameplayTrack().GetWorldName(InSharedData.GetAnalysisSession());
+		}
+		else
+		{
+			ObjectName = FText::FromString(InName);
+		}
+
+		return FText::Format(LOCTEXT("ObjectEventsTrackName", "{0} - {1}"), ClassName, ObjectName);
 	}
 
-	return FText::Format(LOCTEXT("ObjectEventsTrackName", "{0} - {1}"), ClassName, FText::FromString(InName));
+	return ClassName;
 }
 
-void FObjectEventsTrack::GetVariantsAtTime(double InTime, TArray<TSharedRef<FVariantTreeNode>>& OutVariants) const
+void FObjectEventsTrack::GetVariantsAtFrame(const Trace::FFrame& InFrame, TArray<TSharedRef<FVariantTreeNode>>& OutVariants) const
 {
 	const FGameplayProvider* GameplayProvider = SharedData.GetAnalysisSession().ReadProvider<FGameplayProvider>(FGameplayProvider::ProviderName);
 	if(GameplayProvider)
 	{
 		Trace::FAnalysisSessionReadScope SessionReadScope(SharedData.GetAnalysisSession());
 
-		const Trace::IFrameProvider& FramesProvider = Trace::ReadFrameProvider(SharedData.GetAnalysisSession());
-
 		TSharedRef<FVariantTreeNode> Header = OutVariants.Add_GetRef(FVariantTreeNode::MakeHeader(FText::FromString(GetName())));
 
 		// object events
-		GameplayProvider->ReadObjectEventsTimeline(GetGameplayTrack().GetObjectId(), [&InTime, &FramesProvider, &Header](const FGameplayProvider::ObjectEventsTimeline& InTimeline)
+		GameplayProvider->ReadObjectEventsTimeline(GetGameplayTrack().GetObjectId(), [&InFrame, &Header](const FGameplayProvider::ObjectEventsTimeline& InTimeline)
 		{
-			// round to nearest frame boundary
-			Trace::FFrame Frame;
-			if(FramesProvider.GetFrameFromTime(ETraceFrameType::TraceFrameType_Game, InTime, Frame))
+			InTimeline.EnumerateEvents(InFrame.StartTime, InFrame.EndTime, [&Header](double InStartTime, double InEndTime, uint32 InDepth, const FObjectEventMessage& InMessage)
 			{
-				InTimeline.EnumerateEvents(Frame.StartTime, Frame.EndTime, [&Header](double InStartTime, double InEndTime, uint32 InDepth, const FObjectEventMessage& InMessage)
-				{
-					Header->AddChild(FVariantTreeNode::MakeFloat(FText::FromString(InMessage.Name), InStartTime));
-				});
-			}
+				Header->AddChild(FVariantTreeNode::MakeFloat(FText::FromString(InMessage.Name), InStartTime));
+			});
 		});
 	}
 }
diff --git a/Engine/Plugins/Animation/GameplayInsights/Source/GameplayInsights/Private/ObjectEventsTrack.h b/Engine/Plugins/Animation/GameplayInsights/Source/GameplayInsights/Private/ObjectEventsTrack.h
--- a/Engine/Plugins/Animation/GameplayInsights/Source/GameplayInsights/Private/ObjectEventsTrack.h
+++ b/Engine/Plugins/Animation/GameplayInsights/Source/GameplayInsights/Private/ObjectEventsTrack.h
@@ -8,6 +8,8 @@
 class FGameplaySharedData;
 struct FObjectEventMessage;
 class FTimingEventSearchParameters;
+struct FWorldInfo;
+namespace Trace { struct FFrame; }
 
 class FObjectEventsTrack : public FGameplayTimingEventsTrack
 {
@@ -20,7 +22,7 @@ public:
 	virtual void Draw(const ITimingTrackDrawContext& Context) const override;
 	virtual void InitTooltip(FTooltipDrawState& Tooltip, const ITimingEvent& HoveredTimingEvent) const override;
 	virtual const TSharedPtr<const ITimingEvent> SearchEvent(const FTimingEventSearchParameters& InSearchParameters) const override;
-	virtual void GetVariantsAtTime(double InTime, TArray<TSharedRef<FVariantTreeNode>>& OutVariants) const override;
+	virtual void GetVariantsAtFrame(const Trace::FFrame& InFrame, TArray<TSharedRef<FVariantTreeNode>>& OutVariants) const override;
 
 private:
 	// Helper function used to find an object event
diff --git a/Engine/Plugins/Animation/GameplayInsights/Source/GameplayInsights/Private/SGameplayInsightsTransportControls.cpp b/Engine/Plugins/Animation/GameplayInsights/Source/GameplayInsights/Private/SGameplayInsightsTransportControls.cpp
--- a/Engine/Plugins/Animation/GameplayInsights/Source/GameplayInsights/Private/SGameplayInsightsTransportControls.cpp
+++ b/Engine/Plugins/Animation/GameplayInsights/Source/GameplayInsights/Private/SGameplayInsightsTransportControls.cpp
@@ -92,7 +92,7 @@ FReply SGameplayInsightsTransportControls::OnClick_Forward_Step()
 		const Trace::FFrame* FirstFrame = FramesProvider.GetFrame(ETraceFrameType::TraceFrameType_Game, 0);
 		if(FirstFrame)
 		{
-			CurrentTime = FirstFrame->StartTime + (double)KINDA_SMALL_NUMBER;
+			CurrentTime = FirstFrame->StartTime + (double)SMALL_NUMBER;
 		}
 	}
 	
@@ -104,7 +104,7 @@ FReply SGameplayInsightsTransportControls::OnClick_Forward_Step()
 			const Trace::FFrame* NextFrame = FramesProvider.GetFrame(ETraceFrameType::TraceFrameType_Game, Frame.Index + 1);
 			if(NextFrame)
 			{
-				SetTimeMarker(NextFrame->StartTime + (double)KINDA_SMALL_NUMBER, false);
+				SetTimeMarker(NextFrame->StartTime + (double)SMALL_NUMBER, false);
 			}
 		}
 	}
@@ -125,7 +125,7 @@ FReply SGameplayInsightsTransportControls::OnClick_Forward_End()
 		const Trace::FFrame* LastFrame = FramesProvider.GetFrame(ETraceFrameType::TraceFrameType_Game, FramesProvider.GetFrameCount(ETraceFrameType::TraceFrameType_Game) - 1);
 		if(LastFrame)
 		{
-			SetTimeMarker(LastFrame->StartTime + (double)KINDA_SMALL_NUMBER, true);
+			SetTimeMarker(LastFrame->StartTime + (double)SMALL_NUMBER, true);
 		}
 	}
 
@@ -146,7 +146,7 @@ FReply SGameplayInsightsTransportControls::OnClick_Backward_Step()
 		const Trace::FFrame* LastFrame = FramesProvider.GetFrame(ETraceFrameType::TraceFrameType_Game, FramesProvider.GetFrameCount(TraceFrameType_Game) - 1);
 		if(LastFrame)
 		{
-			CurrentTime = LastFrame->StartTime + (double)KINDA_SMALL_NUMBER;
+			CurrentTime = LastFrame->StartTime + (double)SMALL_NUMBER;
 		}
 	}
 
@@ -158,7 +158,7 @@ FReply SGameplayInsightsTransportControls::OnClick_Backward_Step()
 			const Trace::FFrame* PrevFrame = FramesProvider.GetFrame(ETraceFrameType::TraceFrameType_Game, Frame.Index - 1);
 			if(PrevFrame)
 			{
-				SetTimeMarker(PrevFrame->StartTime + (double)KINDA_SMALL_NUMBER, false);
+				SetTimeMarker(PrevFrame->StartTime + (double)SMALL_NUMBER, false);
 			}
 		}
 	}
@@ -177,7 +177,7 @@ FReply SGameplayInsightsTransportControls::OnClick_Backward_End()
 	const Trace::FFrame* FirstFrame = FramesProvider.GetFrame(ETraceFrameType::TraceFrameType_Game, 0);
 	if(FirstFrame)
 	{
-		SetTimeMarker(FirstFrame->StartTime + (double)KINDA_SMALL_NUMBER, true);
+		SetTimeMarker(FirstFrame->StartTime + (double)SMALL_NUMBER, true);
 	}
 
 	bPlaying = false;
@@ -197,7 +197,7 @@ FReply SGameplayInsightsTransportControls::OnClick_Forward()
 		const Trace::FFrame* FirstFrame = FramesProvider.GetFrame(ETraceFrameType::TraceFrameType_Game, 0);
 		if(FirstFrame)
 		{
-			SetTimeMarker(FirstFrame->StartTime + (double)KINDA_SMALL_NUMBER, false);
+			SetTimeMarker(FirstFrame->StartTime + (double)SMALL_NUMBER, false);
 		}
 	}
 
@@ -217,7 +217,7 @@ FReply SGameplayInsightsTransportControls::OnClick_Backward()
 		const Trace::FFrame* LastFrame = FramesProvider.GetFrame(ETraceFrameType::TraceFrameType_Game, FramesProvider.GetFrameCount(TraceFrameType_Game) - 1);
 		if(LastFrame)
 		{
-			SetTimeMarker(LastFrame->StartTime + (double)KINDA_SMALL_NUMBER, false);
+			SetTimeMarker(LastFrame->StartTime + (double)SMALL_NUMBER, false);
 		}
 	}
 
diff --git a/Engine/Plugins/Animation/GameplayInsights/Source/GameplayInsights/Private/STrackVariantValueView.cpp b/Engine/Plugins/Animation/GameplayInsights/Source/GameplayInsights/Private/STrackVariantValueView.cpp
--- a/Engine/Plugins/Animation/GameplayInsights/Source/GameplayInsights/Private/STrackVariantValueView.cpp
+++ b/Engine/Plugins/Animation/GameplayInsights/Source/GameplayInsights/Private/STrackVariantValueView.cpp
@@ -10,33 +10,48 @@
 void STrackVariantValueView::Construct(const FArguments& InArgs, const TSharedRef<FBaseTimingTrack>& InTimingTrack, Insights::ITimingViewSession& InTimingViewSession, const Trace::IAnalysisSession& InAnalysisSession)
 {
 	TimingTrack = InTimingTrack;
+	AnalysisSession = &InAnalysisSession;
 
 	InTimingViewSession.OnTimeMarkerChanged().AddSP(this, &STrackVariantValueView::HandleTimeMarkerChanged);
 
 	ChildSlot
 	[
 		SAssignNew(VariantValueView, SVariantValueView, InAnalysisSession)
-		.OnGetVariantValues_Lambda([this](double InTime, TArray<TSharedRef<FVariantTreeNode>>& OutValues)
+		.OnGetVariantValues_Lambda([this](const Trace::FFrame& InFrame, TArray<TSharedRef<FVariantTreeNode>>& OutValues)
 		{
 			TSharedPtr<FBaseTimingTrack> PinnedTrack = TimingTrack.Pin();
 			if(PinnedTrack.IsValid())
 			{
 				if(PinnedTrack->Is<FGameplayTimingEventsTrack>())
 				{
-					StaticCastSharedPtr<FGameplayTimingEventsTrack>(PinnedTrack)->GetVariantsAtTime(InTime, OutValues);
+					StaticCastSharedPtr<FGameplayTimingEventsTrack>(PinnedTrack)->GetVariantsAtFrame(InFrame, OutValues);
 				}
 				else if(PinnedTrack->Is<FGameplayGraphTrack>())
 				{
-					StaticCastSharedPtr<FGameplayGraphTrack>(PinnedTrack)->GetVariantsAtTime(InTime, OutValues);
+					StaticCastSharedPtr<FGameplayGraphTrack>(PinnedTrack)->GetVariantsAtFrame(InFrame, OutValues);
 				}
 			}
 		})
 	];
 
-	VariantValueView->RequestRefresh(InTimingViewSession.GetTimeMarker());
+	Trace::FAnalysisSessionReadScope SessionReadScope(InAnalysisSession);
+
+	const Trace::IFrameProvider& FramesProvider = Trace::ReadFrameProvider(InAnalysisSession);
+	Trace::FFrame MarkerFrame;
+	if(FramesProvider.GetFrameFromTime(ETraceFrameType::TraceFrameType_Game, InTimingViewSession.GetTimeMarker(), MarkerFrame))
+	{
+		VariantValueView->RequestRefresh(MarkerFrame);
+	}
 }
 
 void STrackVariantValueView::HandleTimeMarkerChanged(Insights::ETimeChangedFlags InFlags, double InTimeMarker)
 {
-	VariantValueView->RequestRefresh(InTimeMarker);
+	Trace::FAnalysisSessionReadScope SessionReadScope(*AnalysisSession);
+
+	const Trace::IFrameProvider& FramesProvider = Trace::ReadFrameProvider(*AnalysisSession);
+	Trace::FFrame MarkerFrame;
+	if(FramesProvider.GetFrameFromTime(ETraceFrameType::TraceFrameType_Game, InTimeMarker, MarkerFrame))
+	{
+		VariantValueView->RequestRefresh(MarkerFrame);
+	}
 }
diff --git a/Engine/Plugins/Animation/GameplayInsights/Source/GameplayInsights/Private/STrackVariantValueView.h b/Engine/Plugins/Animation/GameplayInsights/Source/GameplayInsights/Private/STrackVariantValueView.h
--- a/Engine/Plugins/Animation/GameplayInsights/Source/GameplayInsights/Private/STrackVariantValueView.h
+++ b/Engine/Plugins/Animation/GameplayInsights/Source/GameplayInsights/Private/STrackVariantValueView.h
@@ -32,4 +32,6 @@ private:
 	TWeakPtr<FBaseTimingTrack> TimingTrack;
 
 	TSharedPtr<SVariantValueView> VariantValueView;
+
+	const Trace::IAnalysisSession* AnalysisSession;
 };
\ No newline at end of file
diff --git a/Engine/Plugins/Animation/GameplayInsights/Source/GameplayInsights/Private/SVariantValueView.cpp b/Engine/Plugins/Animation/GameplayInsights/Source/GameplayInsights/Private/SVariantValueView.cpp
--- a/Engine/Plugins/Animation/GameplayInsights/Source/GameplayInsights/Private/SVariantValueView.cpp
+++ b/Engine/Plugins/Animation/GameplayInsights/Source/GameplayInsights/Private/SVariantValueView.cpp
@@ -330,7 +330,7 @@ void SVariantValueView::RefreshNodes()
 {
 	VariantTreeNodes.Reset();
 
-	OnGetVariantValues.ExecuteIfBound(TimeMarker, VariantTreeNodes);
+	OnGetVariantValues.ExecuteIfBound(Frame, VariantTreeNodes);
 
 	if(VariantTreeNodes.Num() > 0)
 	{
diff --git a/Engine/Plugins/Animation/GameplayInsights/Source/GameplayInsights/Private/SVariantValueView.h b/Engine/Plugins/Animation/GameplayInsights/Source/GameplayInsights/Private/SVariantValueView.h
--- a/Engine/Plugins/Animation/GameplayInsights/Source/GameplayInsights/Private/SVariantValueView.h
+++ b/Engine/Plugins/Animation/GameplayInsights/Source/GameplayInsights/Private/SVariantValueView.h
@@ -5,12 +5,13 @@
 #include "CoreMinimal.h"
 #include "Widgets/SCompoundWidget.h"
 #include "Widgets/Views/STreeView.h"
+#include "TraceServices/Model/Frames.h"
 
 namespace Trace { class IAnalysisSession; }
 struct FVariantTreeNode;
 
 // Delegate called to get variant values to display
-DECLARE_DELEGATE_TwoParams(FOnGetVariantValues, double /*InTime*/, TArray<TSharedRef<FVariantTreeNode>>& /*OutValues*/);
+DECLARE_DELEGATE_TwoParams(FOnGetVariantValues, const Trace::FFrame& /*InFrame*/, TArray<TSharedRef<FVariantTreeNode>>& /*OutValues*/);
 
 class SVariantValueView : public SCompoundWidget
 {
@@ -24,7 +25,7 @@ public:
 	void Construct(const FArguments& InArgs, const Trace::IAnalysisSession& InAnalysisSession);
 
 	/** Refresh the displayed variants. */
-	void RequestRefresh(double InTimeMarker) { TimeMarker = InTimeMarker; bNeedsRefresh = true; }
+	void RequestRefresh(const Trace::FFrame& InFrame) { Frame = InFrame; bNeedsRefresh = true; }
 
 private:
 	// Generate a row widget for a property item
@@ -43,7 +44,7 @@ private:
 
 	TArray<TSharedRef<FVariantTreeNode>> VariantTreeNodes;
 
-	double TimeMarker;
+	Trace::FFrame Frame;
 
 	FOnGetVariantValues OnGetVariantValues;
 
diff --git a/Engine/Plugins/Animation/GameplayInsights/Source/GameplayInsights/Private/SkeletalMeshCurvesTrack.cpp b/Engine/Plugins/Animation/GameplayInsights/Source/GameplayInsights/Private/SkeletalMeshCurvesTrack.cpp
--- a/Engine/Plugins/Animation/GameplayInsights/Source/GameplayInsights/Private/SkeletalMeshCurvesTrack.cpp
+++ b/Engine/Plugins/Animation/GameplayInsights/Source/GameplayInsights/Private/SkeletalMeshCurvesTrack.cpp
@@ -211,7 +211,7 @@ void FSkeletalMeshCurvesTrack::FindSkeletalMeshPoseMessage(const FTimingEventSea
 		});
 }
 
-void FSkeletalMeshCurvesTrack::GetVariantsAtTime(double InTime, TArray<TSharedRef<FVariantTreeNode>>& OutVariants) const
+void FSkeletalMeshCurvesTrack::GetVariantsAtFrame(const Trace::FFrame& InFrame, TArray<TSharedRef<FVariantTreeNode>>& OutVariants) const
 {
 	const FGameplayProvider* GameplayProvider = SharedData.GetAnalysisSession().ReadProvider<FGameplayProvider>(FGameplayProvider::ProviderName);
 	const FAnimationProvider* AnimationProvider = SharedData.GetAnalysisSession().ReadProvider<FAnimationProvider>(FAnimationProvider::ProviderName);
@@ -221,26 +221,19 @@ void FSkeletalMeshCurvesTrack::GetVariantsAtTime(double InTime, TArray<TSharedRe
 
 		TSharedRef<FVariantTreeNode> Header = OutVariants.Add_GetRef(FVariantTreeNode::MakeHeader(FText::FromString(GetName())));
 
-		const Trace::IFrameProvider& FramesProvider = Trace::ReadFrameProvider(SharedData.GetAnalysisSession());
-
-		AnimationProvider->ReadSkeletalMeshPoseTimeline(GetGameplayTrack().GetObjectId(), [&Header, &FramesProvider, &AnimationProvider, &InTime](const FAnimationProvider::SkeletalMeshPoseTimeline& InTimeline, bool bInHasCurves)
+		AnimationProvider->ReadSkeletalMeshPoseTimeline(GetGameplayTrack().GetObjectId(), [&Header, &AnimationProvider, &InFrame](const FAnimationProvider::SkeletalMeshPoseTimeline& InTimeline, bool bInHasCurves)
 		{
-			// round to nearest frame boundary
-			Trace::FFrame Frame;
-			if(FramesProvider.GetFrameFromTime(ETraceFrameType::TraceFrameType_Game, InTime, Frame))
+			InTimeline.EnumerateEvents(InFrame.StartTime, InFrame.EndTime, [&Header, &AnimationProvider, &InFrame](double InStartTime, double InEndTime, uint32 InDepth, const FSkeletalMeshPoseMessage& InMessage)
 			{
-				InTimeline.EnumerateEvents(Frame.StartTime, Frame.EndTime, [&Header, &AnimationProvider, &Frame](double InStartTime, double InEndTime, uint32 InDepth, const FSkeletalMeshPoseMessage& InMessage)
+				if(InStartTime >= InFrame.StartTime && InStartTime <= InFrame.EndTime)
 				{
-					if(Frame.StartTime <= InStartTime && Frame.EndTime >= InEndTime)
+					AnimationProvider->EnumerateSkeletalMeshCurves(InMessage, [&Header, &AnimationProvider](const FSkeletalMeshNamedCurve& InCurve)
 					{
-						AnimationProvider->EnumerateSkeletalMeshCurves(InMessage, [&Header, &AnimationProvider](const FSkeletalMeshNamedCurve& InCurve)
-						{
-							const TCHAR* CurveName = AnimationProvider->GetName(InCurve.Id);
-							Header->AddChild(FVariantTreeNode::MakeFloat(FText::FromString(CurveName), InCurve.Value));
-						});
-					}
-				});
-			}
+						const TCHAR* CurveName = AnimationProvider->GetName(InCurve.Id);
+						Header->AddChild(FVariantTreeNode::MakeFloat(FText::FromString(CurveName), InCurve.Value));
+					});
+				}
+			});
 		});
 	}
 }
diff --git a/Engine/Plugins/Animation/GameplayInsights/Source/GameplayInsights/Private/SkeletalMeshCurvesTrack.h b/Engine/Plugins/Animation/GameplayInsights/Source/GameplayInsights/Private/SkeletalMeshCurvesTrack.h
--- a/Engine/Plugins/Animation/GameplayInsights/Source/GameplayInsights/Private/SkeletalMeshCurvesTrack.h
+++ b/Engine/Plugins/Animation/GameplayInsights/Source/GameplayInsights/Private/SkeletalMeshCurvesTrack.h
@@ -27,7 +27,7 @@ public:
 	virtual void AddAllSeries() override;
 	virtual bool UpdateSeriesBounds(FGameplayGraphSeries& InSeries, const FTimingTrackViewport& InViewport) override;
 	virtual void UpdateSeries(FGameplayGraphSeries& InSeries, const FTimingTrackViewport& InViewport) override;
-	virtual void GetVariantsAtTime(double InTime, TArray<TSharedRef<FVariantTreeNode>>& OutVariants) const override;
+	virtual void GetVariantsAtFrame(const Trace::FFrame& InFrame, TArray<TSharedRef<FVariantTreeNode>>& OutVariants) const override;
 
 private:
 	// Helper function used to find a skeletal mesh pose
diff --git a/Engine/Plugins/Animation/GameplayInsights/Source/GameplayInsights/Private/SkeletalMeshPoseTrack.cpp b/Engine/Plugins/Animation/GameplayInsights/Source/GameplayInsights/Private/SkeletalMeshPoseTrack.cpp
--- a/Engine/Plugins/Animation/GameplayInsights/Source/GameplayInsights/Private/SkeletalMeshPoseTrack.cpp
+++ b/Engine/Plugins/Animation/GameplayInsights/Source/GameplayInsights/Private/SkeletalMeshPoseTrack.cpp
@@ -77,9 +77,7 @@ void FSkeletalMeshPoseTrack::BuildDrawState(ITimingEventsTrackDrawStateBuilder&
 		{
 			InTimeline.EnumerateEvents(Context.GetViewport().GetStartTime(), Context.GetViewport().GetEndTime(), [&Builder](double InStartTime, double InEndTime, uint32 InDepth, const FSkeletalMeshPoseMessage& InMessage)
 			{
-				static TCHAR Buffer[256];
-				FCString::Snprintf(Buffer, 256, TEXT("%d Bones"), InMessage.NumTransforms);
-				Builder.AddEvent(InStartTime, InEndTime, 0, Buffer);
+				Builder.AddEvent(InStartTime, InEndTime, 0, InMessage.MeshName);
 			});
 		});
 	}
@@ -122,6 +120,8 @@ void FSkeletalMeshPoseTrack::InitTooltip(FTooltipDrawState& Tooltip, const ITimi
 
 		Tooltip.AddNameValueTextLine(LOCTEXT("BoneCount", "Bone Count").ToString(), FText::AsNumber(InMessage.NumTransforms).ToString());
 		Tooltip.AddNameValueTextLine(LOCTEXT("CurveCount", "Curve Count").ToString(), FText::AsNumber(InMessage.NumCurves).ToString());
+		Tooltip.AddNameValueTextLine(LOCTEXT("LODIndex", "LOD").ToString(), FText::AsNumber(InMessage.LodIndex).ToString());
+		Tooltip.AddNameValueTextLine(LOCTEXT("EventWorld", "World").ToString(), GetGameplayTrack().GetWorldName(SharedData.GetAnalysisSession()).ToString());
 
 		Tooltip.UpdateLayout();
 	});
@@ -365,7 +365,7 @@ USkeletalMeshComponent* FSkeletalMeshPoseTrack::GetComponent(UWorld* InWorld)
 	return nullptr;
 }
 
-void FSkeletalMeshPoseTrack::DrawPoses(UWorld* InWorld, double InTime)
+void FSkeletalMeshPoseTrack::DrawPoses(UWorld* InWorld, double InTime, double InFrameStartTime, double InFrameEndTime)
 {
 	if(SharedData.IsAnalysisSessionValid())
 	{
@@ -382,11 +382,11 @@ void FSkeletalMeshPoseTrack::DrawPoses(UWorld* InWorld, double InTime)
 				CacheForWorld.Component->SetVisibility(false);
 			}
 
-			AnimationProvider->ReadSkeletalMeshPoseTimeline(GetGameplayTrack().GetObjectId(), [this, &CacheForWorld, &AnimationProvider, &GameplayProvider, &InTime](const FAnimationProvider::SkeletalMeshPoseTimeline& InTimeline, bool bInHasCurves)
+			AnimationProvider->ReadSkeletalMeshPoseTimeline(GetGameplayTrack().GetObjectId(), [this, &CacheForWorld, &AnimationProvider, &GameplayProvider, &InFrameStartTime, &InFrameEndTime](const FAnimationProvider::SkeletalMeshPoseTimeline& InTimeline, bool bInHasCurves)
 			{
-				InTimeline.EnumerateEvents(InTime, InTime, [this, &CacheForWorld, &AnimationProvider, &GameplayProvider, &InTime](double InStartTime, double InEndTime, uint32 InDepth, const FSkeletalMeshPoseMessage& InMessage)
+				InTimeline.EnumerateEvents(InFrameStartTime, InFrameEndTime, [this, &CacheForWorld, &AnimationProvider, &GameplayProvider, &InFrameStartTime, &InFrameEndTime](double InStartTime, double InEndTime, uint32 InDepth, const FSkeletalMeshPoseMessage& InMessage)
 				{
-					if((InStartTime <= InTime && InEndTime > InTime))
+					if(InStartTime >= InFrameStartTime && InStartTime <= InFrameEndTime)
 					{
 						const FSkeletalMeshInfo* SkeletalMeshInfo = AnimationProvider->FindSkeletalMeshInfo(InMessage.MeshId);
 						const FObjectInfo* SkeletalMeshObjectInfo = GameplayProvider->FindObjectInfo(InMessage.MeshId);
@@ -395,7 +395,7 @@ void FSkeletalMeshPoseTrack::DrawPoses(UWorld* InWorld, double InTime)
 							UInsightsSkeletalMeshComponent* Component = CacheForWorld.GetComponent();
 							Component->SetVisibility(bDrawPose);
 
-							if(CacheForWorld.Time != InTime)
+							if(CacheForWorld.Time != InFrameStartTime)
 							{
 								USkeletalMesh* SkeletalMesh = TSoftObjectPtr<USkeletalMesh>(FSoftObjectPath(SkeletalMeshObjectInfo->PathName)).LoadSynchronous();
 								if(SkeletalMesh)
@@ -405,7 +405,7 @@ void FSkeletalMeshPoseTrack::DrawPoses(UWorld* InWorld, double InTime)
 
 								Component->SetPoseFromProvider(*AnimationProvider, InMessage, *SkeletalMeshInfo);
 
-								CacheForWorld.Time = InTime;
+								CacheForWorld.Time = InFrameStartTime;
 							}
 
 							Component->SetDrawDebugSkeleton(bDrawSkeleton);
diff --git a/Engine/Plugins/Animation/GameplayInsights/Source/GameplayInsights/Private/SkeletalMeshPoseTrack.h b/Engine/Plugins/Animation/GameplayInsights/Source/GameplayInsights/Private/SkeletalMeshPoseTrack.h
--- a/Engine/Plugins/Animation/GameplayInsights/Source/GameplayInsights/Private/SkeletalMeshPoseTrack.h
+++ b/Engine/Plugins/Animation/GameplayInsights/Source/GameplayInsights/Private/SkeletalMeshPoseTrack.h
@@ -55,7 +55,7 @@ public:
 	void RemoveWorld(UWorld* InWorld);
 
 	// Draw poses at the specified time
-	void DrawPoses(UWorld* InWorld, double InTime);
+	void DrawPoses(UWorld* InWorld, double InTime, double InFrameStartTime, double InFrameEndTime);
 
 	// FGCObject interface
 	virtual void AddReferencedObjects(FReferenceCollector& Collector) override;
diff --git a/Engine/Plugins/Animation/GameplayInsights/Source/GameplayInsights/Public/IAnimationProvider.h b/Engine/Plugins/Animation/GameplayInsights/Source/GameplayInsights/Public/IAnimationProvider.h
--- a/Engine/Plugins/Animation/GameplayInsights/Source/GameplayInsights/Public/IAnimationProvider.h
+++ b/Engine/Plugins/Animation/GameplayInsights/Source/GameplayInsights/Public/IAnimationProvider.h
@@ -27,7 +27,8 @@ struct FSkeletalMeshPoseMessage
 	uint64 TransformStartIndex = 0;
 	uint64 CurveStartIndex = 0;
 	uint64 ComponentId = 0;	
-	uint64 MeshId = 0;	
+	uint64 MeshId = 0;
+	const TCHAR* MeshName = nullptr;
 	uint16 NumTransforms = 0;
 	uint16 NumCurves = 0;
 	uint16 FrameCounter = 0;
@@ -55,7 +56,6 @@ struct FTickRecordMessage
 	uint16 FrameCounter = 0;
 	bool bLooping = false;
 	bool bIsBlendSpace = false;
-	bool bContinuous = true;
 };
 
 enum class EAnimGraphPhase : uint8
@@ -172,7 +172,7 @@ struct FAnimSequencePlayerMessage
 	int32 NodeId = -1;
 	float Position = 0.0f;
 	float Length = 0.0f;
-	int32 FrameCount = 0;
+	uint16 FrameCounter = 0;
 };
 
 struct FBlendSpacePlayerMessage
@@ -209,6 +209,7 @@ struct FAnimNotifyMessage
 	uint64 AnimInstanceId = 0;
 	uint64 AssetId = 0;
 	uint64 NotifyId = 0;
+	const TCHAR* Name = nullptr;
 	uint32 NameId = 0;
 	float Time = 0.0f; 
 	float Duration = 0.0f;
@@ -223,6 +224,7 @@ struct FAnimMontageMessage
 	uint32 NextSectionNameId = 0;
 	float Weight = 0.0f;
 	float DesiredWeight = 0.0f;
+	uint16 FrameCounter = 0;
 };
 
 class IAnimationProvider : public Trace::IProvider
@@ -243,8 +245,8 @@ public:
 	virtual void GetSkeletalMeshComponentSpacePose(const FSkeletalMeshPoseMessage& InMessage, const FSkeletalMeshInfo& InMeshInfo, FTransform& OutComponentToWorld, TArray<FTransform>& OutTransforms) const = 0;
 	virtual void EnumerateSkeletalMeshCurveIds(uint64 InObjectId, TFunctionRef<void(uint32)> Callback) const = 0;
 	virtual void EnumerateSkeletalMeshCurves(const FSkeletalMeshPoseMessage& InMessage, TFunctionRef<void(const FSkeletalMeshNamedCurve&)> Callback) const = 0;
-	virtual void EnumerateTickRecordTimelines(uint64 InObjectId, TFunctionRef<void(uint64, int32, const TickRecordTimeline&)> Callback) const = 0;
-	virtual bool ReadTickRecordTimeline(uint64 InObjectId, uint64 InAssetId, int32 InNodeId, TFunctionRef<void(const TickRecordTimeline&)> Callback) const = 0;
+	virtual bool ReadTickRecordTimeline(uint64 InObjectId, TFunctionRef<void(const TickRecordTimeline&)> Callback) const = 0;
+	virtual void EnumerateTickRecordIds(uint64 InObjectId, TFunctionRef<void(uint64, int32)> Callback) const = 0;
 	virtual bool ReadAnimGraphTimeline(uint64 InObjectId, TFunctionRef<void(const AnimGraphTimeline&)> Callback) const = 0;
 	virtual bool ReadAnimNodesTimeline(uint64 InObjectId, TFunctionRef<void(const AnimNodesTimeline&)> Callback) const = 0;
 	virtual bool ReadAnimNodeValuesTimeline(uint64 InObjectId, TFunctionRef<void(const AnimNodeValuesTimeline&)> Callback) const = 0;
diff --git a/Engine/Plugins/Animation/GameplayInsights/Source/GameplayInsights/Public/IGameplayProvider.h b/Engine/Plugins/Animation/GameplayInsights/Source/GameplayInsights/Public/IGameplayProvider.h
--- a/Engine/Plugins/Animation/GameplayInsights/Source/GameplayInsights/Public/IGameplayProvider.h
+++ b/Engine/Plugins/Animation/GameplayInsights/Source/GameplayInsights/Public/IGameplayProvider.h
@@ -29,6 +29,40 @@ struct FObjectEventMessage
 	const TCHAR* Name = nullptr;
 };
 
+
+struct FWorldInfo
+{
+	/** Types of worlds that we know about - synced with EngineTypes.h */
+	enum class EType : uint8
+	{
+		None,
+		Game,
+		Editor,
+		PIE,
+		EditorPreview,
+		GamePreview,
+		GameRPC,
+		Inactive
+	};
+
+	/** Types of net modes that we know about - synced with EngineBaseTypes.h */
+	enum class ENetMode : uint8
+	{
+		Standalone,
+		DedicatedServer,
+		ListenServer,
+		Client,
+
+		MAX,
+	};
+
+	uint64 Id = 0;
+	int32 PIEInstanceId = 0;
+	EType Type = EType::None;
+	ENetMode NetMode = ENetMode::Standalone;
+	bool bIsSimulating = false;
+};
+
 // Delegate fired when an object receives an end play event
 DECLARE_MULTICAST_DELEGATE_ThreeParams(FOnObjectEndPlay, uint64 /*ObjectId*/, double /*Time*/, const FObjectInfo& /*ObjectInfo*/);
 
@@ -41,7 +75,11 @@ public:
 	virtual bool ReadObjectEvent(uint64 InObjectId, uint64 InMessageId, TFunctionRef<void(const FObjectEventMessage&)> Callback) const = 0;
 	virtual void EnumerateObjects(TFunctionRef<void(const FObjectInfo&)> Callback) const = 0;
 	virtual const FClassInfo* FindClassInfo(uint64 InClassId) const = 0;
+	virtual const FClassInfo* FindClassInfo(const TCHAR* InClassPath) const = 0;
 	virtual const FObjectInfo* FindObjectInfo(uint64 InObjectId) const = 0;
+	virtual const FWorldInfo* FindWorldInfo(uint64 InObjectId) const = 0;
+	virtual const FWorldInfo* FindWorldInfoFromObject(uint64 InObjectId) const = 0;
+	virtual bool IsWorld(uint64 InObjectId) const = 0;
 	virtual const FClassInfo& GetClassInfo(uint64 InClassId) const = 0;
 	virtual const FClassInfo& GetClassInfoFromObject(uint64 InObjectId) const = 0;
 	virtual const FObjectInfo& GetObjectInfo(uint64 InObjectId) const = 0;
diff --git a/Engine/Source/Developer/TraceServices/Public/Model/PointTimeline.h b/Engine/Source/Developer/TraceServices/Public/Model/PointTimeline.h
--- a/Engine/Source/Developer/TraceServices/Public/Model/PointTimeline.h
+++ b/Engine/Source/Developer/TraceServices/Public/Model/PointTimeline.h
@@ -68,7 +68,7 @@ public:
 			{
 				if (!(Event->Time < IntervalStart || IntervalEnd < Event->Time))
 				{
-					if(bEnumerateOutsideRange && PreviousEvent != Event)
+					if(!bIntersected && bEnumerateOutsideRange && PreviousEvent != Event)
 					{
 						Callback(PreviousEvent->Time, PreviousEvent->Time, 0, PreviousEvent->Event);
 					}
diff --git a/Engine/Source/Runtime/Engine/Private/Animation/AnimTrace.cpp b/Engine/Source/Runtime/Engine/Private/Animation/AnimTrace.cpp
--- a/Engine/Source/Runtime/Engine/Private/Animation/AnimTrace.cpp
+++ b/Engine/Source/Runtime/Engine/Private/Animation/AnimTrace.cpp
@@ -166,7 +166,7 @@ UE_TRACE_EVENT_BEGIN(Animation, AnimSequencePlayer)
 	UE_TRACE_EVENT_FIELD(int32, NodeId)
 	UE_TRACE_EVENT_FIELD(float, Position)
 	UE_TRACE_EVENT_FIELD(float, Length)
-	UE_TRACE_EVENT_FIELD(int32, FrameCount)
+	UE_TRACE_EVENT_FIELD(uint16, FrameCounter)
 UE_TRACE_EVENT_END()
 
 UE_TRACE_EVENT_BEGIN(Animation, BlendSpacePlayer)
@@ -218,6 +218,7 @@ UE_TRACE_EVENT_BEGIN(Animation, Montage)
 	UE_TRACE_EVENT_FIELD(uint32, NextSectionNameId)
 	UE_TRACE_EVENT_FIELD(float, Weight)
 	UE_TRACE_EVENT_FIELD(float, DesiredWeight)
+	UE_TRACE_EVENT_FIELD(uint16, FrameCounter)
 UE_TRACE_EVENT_END()
 
 // Object annotations used for tracing
@@ -366,7 +367,8 @@ void FAnimTrace::OutputAnimTickRecord(const FAnimationBaseContext& InContext, co
 		return;
 	}
 
-	TRACE_OBJECT(InContext.AnimInstanceProxy->GetAnimInstanceObject());
+	UObject* AnimInstance = InContext.AnimInstanceProxy->GetAnimInstanceObject();
+	TRACE_OBJECT(AnimInstance);
 
 	if(InTickRecord.SourceAsset)
 	{
@@ -389,7 +391,7 @@ void FAnimTrace::OutputAnimTickRecord(const FAnimationBaseContext& InContext, co
 
 		UE_TRACE_LOG(Animation, TickRecord, AnimationChannel)
 			<< TickRecord.Cycle(FPlatformTime::Cycles64())
-			<< TickRecord.AnimInstanceId(FObjectTrace::GetObjectId(InContext.AnimInstanceProxy->GetAnimInstanceObject()))
+			<< TickRecord.AnimInstanceId(FObjectTrace::GetObjectId(AnimInstance))
 			<< TickRecord.AssetId(FObjectTrace::GetObjectId(InTickRecord.SourceAsset))
 			<< TickRecord.NodeId(InContext.GetCurrentNodeId())
 			<< TickRecord.BlendWeight(InTickRecord.EffectiveBlendWeight)
@@ -398,7 +400,7 @@ void FAnimTrace::OutputAnimTickRecord(const FAnimationBaseContext& InContext, co
 			<< TickRecord.PlayRate(InTickRecord.PlayRateMultiplier)
 			<< TickRecord.BlendSpacePositionX(BlendSpacePositionX)
 			<< TickRecord.BlendSpacePositionY(BlendSpacePositionY)
-			<< TickRecord.FrameCounter((uint16)(GFrameCounter % 0xffff))
+			<< TickRecord.FrameCounter(FObjectTrace::GetObjectWorldTickCounter(AnimInstance))
 			<< TickRecord.Looping(InTickRecord.bLooping)
 			<< TickRecord.IsBlendSpace(bIsBlendSpace);
 	}
@@ -536,7 +538,7 @@ void FAnimTrace::OutputSkeletalMeshComponent(const USkeletalMeshComponent* InCom
 			<< SkeletalMeshComponent2.CurveIds(CurveIds.GetData(), CurveIds.Num())
 			<< SkeletalMeshComponent2.CurveValues(CurveValues.GetData(), CurveValues.Num())
 			<< SkeletalMeshComponent2.LodIndex((uint16)InComponent->PredictedLODLevel)
-			<< SkeletalMeshComponent2.FrameCounter((uint16)(GFrameCounter % 0xffff));
+			<< SkeletalMeshComponent2.FrameCounter(FObjectTrace::GetObjectWorldTickCounter(InComponent));
 	}
 }
 
@@ -546,7 +548,7 @@ void FAnimTrace::OutputSkeletalMeshFrame(const USkeletalMeshComponent* InCompone
 	UE_TRACE_LOG(Animation, SkeletalMeshFrame, AnimationChannel)
 		<< SkeletalMeshFrame.Cycle(FPlatformTime::Cycles64())
 		<< SkeletalMeshFrame.ComponentId(FObjectTrace::GetObjectId(InComponent))
-		<< SkeletalMeshFrame.FrameCounter((uint16)(GFrameCounter % 0xffff));
+		<< SkeletalMeshFrame.FrameCounter(FObjectTrace::GetObjectWorldTickCounter(InComponent));
 }
 
 void FAnimTrace::OutputAnimGraph(const FAnimationBaseContext& InContext, uint64 InStartCycle, uint64 InEndCycle, uint8 InPhase)
@@ -574,7 +576,7 @@ void FAnimTrace::OutputAnimGraph(const FAnimationBaseContext& InContext, uint64
 		<< AnimGraph.EndCycle(InEndCycle)
 		<< AnimGraph.AnimInstanceId(FObjectTrace::GetObjectId(AnimInstance))
 		<< AnimGraph.NodeCount(BPClass ? BPClass->AnimNodeProperties.Num() : 0)
-		<< AnimGraph.FrameCounter((uint16)(GFrameCounter % 0xffff))
+		<< AnimGraph.FrameCounter(FObjectTrace::GetObjectWorldTickCounter(AnimInstance))
 		<< AnimGraph.Phase(InPhase);
 }
 
@@ -598,7 +600,8 @@ void FAnimTrace::OutputAnimNodeStart(const FAnimationBaseContext& InContext, uin
 
 	check(InContext.AnimInstanceProxy);
 
-	TRACE_OBJECT(InContext.AnimInstanceProxy->GetAnimInstanceObject());
+	UObject* AnimInstance = InContext.AnimInstanceProxy->GetAnimInstanceObject();
+	TRACE_OBJECT(AnimInstance);
 
 	FString DisplayNameString;
 	IAnimClassInterface* AnimBlueprintClass = InContext.GetAnimClass();
@@ -626,11 +629,11 @@ void FAnimTrace::OutputAnimNodeStart(const FAnimationBaseContext& InContext, uin
 
 	UE_TRACE_LOG(Animation, AnimNodeStart, AnimationChannel, (DisplayNameString.Len() + 1) * sizeof(TCHAR))
 		<< AnimNodeStart.StartCycle(InStartCycle)
-		<< AnimNodeStart.AnimInstanceId(FObjectTrace::GetObjectId(InContext.AnimInstanceProxy->GetAnimInstanceObject()))
+		<< AnimNodeStart.AnimInstanceId(FObjectTrace::GetObjectId(AnimInstance))
 		<< AnimNodeStart.PreviousNodeId(InPreviousNodeId)
 		<< AnimNodeStart.NodeId(InNodeId)
 		<< AnimNodeStart.Weight(InBlendWeight)
-		<< AnimNodeStart.FrameCounter((uint16)(GFrameCounter % 0xffff))
+		<< AnimNodeStart.FrameCounter(FObjectTrace::GetObjectWorldTickCounter(AnimInstance))
 		<< AnimNodeStart.Phase(InPhase)
 		<< AnimNodeStart.Attachment(*DisplayNameString, (DisplayNameString.Len() + 1) * sizeof(TCHAR));
 }
@@ -650,11 +653,12 @@ void FAnimTrace::OutputAnimNodeEnd(const FAnimationBaseContext& InContext, uint6
 		return;
 	}
 
-	TRACE_OBJECT(InContext.AnimInstanceProxy->GetAnimInstanceObject());
+	UObject* AnimInstance = InContext.AnimInstanceProxy->GetAnimInstanceObject();
+	TRACE_OBJECT(AnimInstance);
 
 	UE_TRACE_LOG(Animation, AnimNodeEnd, AnimationChannel)
 		<< AnimNodeEnd.EndCycle(InEndCycle)
-		<< AnimNodeEnd.AnimInstanceId(FObjectTrace::GetObjectId(InContext.AnimInstanceProxy->GetAnimInstanceObject()));
+		<< AnimNodeEnd.AnimInstanceId(FObjectTrace::GetObjectId(AnimInstance));
 }
 
 void FAnimTrace::OutputAnimNodeValue(const FAnimationBaseContext& InContext, const TCHAR* InKey, bool InValue)
@@ -672,17 +676,18 @@ void FAnimTrace::OutputAnimNodeValue(const FAnimationBaseContext& InContext, con
 		return;
 	}
 
-	TRACE_OBJECT(InContext.AnimInstanceProxy->GetAnimInstanceObject());
+	UObject* AnimInstance = InContext.AnimInstanceProxy->GetAnimInstanceObject();
+	TRACE_OBJECT(AnimInstance);
 
 	int32 KeyLength = FCString::Strlen(InKey) + 1;
 
 	UE_TRACE_LOG(Animation, AnimNodeValueBool, AnimationChannel, KeyLength * sizeof(TCHAR))
 		<< AnimNodeValueBool.Cycle(FPlatformTime::Cycles64())
-		<< AnimNodeValueBool.AnimInstanceId(FObjectTrace::GetObjectId(InContext.AnimInstanceProxy->GetAnimInstanceObject()))
+		<< AnimNodeValueBool.AnimInstanceId(FObjectTrace::GetObjectId(AnimInstance))
 		<< AnimNodeValueBool.NodeId(InContext.GetCurrentNodeId())
 		<< AnimNodeValueBool.KeyLength(KeyLength)
 		<< AnimNodeValueBool.Value(InValue)
-		<< AnimNodeValueBool.FrameCounter((uint16)(GFrameCounter % 0xffff))
+		<< AnimNodeValueBool.FrameCounter(FObjectTrace::GetObjectWorldTickCounter(AnimInstance))
 		<< AnimNodeValueBool.Attachment(InKey, KeyLength * sizeof(TCHAR));
 }
 
@@ -701,17 +706,18 @@ void FAnimTrace::OutputAnimNodeValue(const FAnimationBaseContext& InContext, con
 		return;
 	}
 
-	TRACE_OBJECT(InContext.AnimInstanceProxy->GetAnimInstanceObject());
+	UObject* AnimInstance = InContext.AnimInstanceProxy->GetAnimInstanceObject();
+	TRACE_OBJECT(AnimInstance);
 
 	int32 KeyLength = FCString::Strlen(InKey) + 1;
 
 	UE_TRACE_LOG(Animation, AnimNodeValueInt, AnimationChannel, KeyLength * sizeof(TCHAR))
 		<< AnimNodeValueInt.Cycle(FPlatformTime::Cycles64())
-		<< AnimNodeValueInt.AnimInstanceId(FObjectTrace::GetObjectId(InContext.AnimInstanceProxy->GetAnimInstanceObject()))
+		<< AnimNodeValueInt.AnimInstanceId(FObjectTrace::GetObjectId(AnimInstance))
 		<< AnimNodeValueInt.NodeId(InContext.GetCurrentNodeId())
 		<< AnimNodeValueInt.KeyLength(KeyLength)
 		<< AnimNodeValueInt.Value(InValue)
-		<< AnimNodeValueInt.FrameCounter((uint16)(GFrameCounter % 0xffff))
+		<< AnimNodeValueInt.FrameCounter(FObjectTrace::GetObjectWorldTickCounter(AnimInstance))
 		<< AnimNodeValueInt.Attachment(InKey, KeyLength * sizeof(TCHAR));
 }
 
@@ -730,17 +736,18 @@ void FAnimTrace::OutputAnimNodeValue(const FAnimationBaseContext& InContext, con
 		return;
 	}
 
-	TRACE_OBJECT(InContext.AnimInstanceProxy->GetAnimInstanceObject());
+	UObject* AnimInstance = InContext.AnimInstanceProxy->GetAnimInstanceObject();
+	TRACE_OBJECT(AnimInstance);
 
 	int32 KeyLength = FCString::Strlen(InKey) + 1;
 
 	UE_TRACE_LOG(Animation, AnimNodeValueFloat, AnimationChannel, KeyLength * sizeof(TCHAR))
 		<< AnimNodeValueFloat.Cycle(FPlatformTime::Cycles64())
-		<< AnimNodeValueFloat.AnimInstanceId(FObjectTrace::GetObjectId(InContext.AnimInstanceProxy->GetAnimInstanceObject()))
+		<< AnimNodeValueFloat.AnimInstanceId(FObjectTrace::GetObjectId(AnimInstance))
 		<< AnimNodeValueFloat.NodeId(InContext.GetCurrentNodeId())
 		<< AnimNodeValueFloat.KeyLength(KeyLength)
 		<< AnimNodeValueFloat.Value(InValue)
-		<< AnimNodeValueFloat.FrameCounter((uint16)(GFrameCounter % 0xffff))
+		<< AnimNodeValueFloat.FrameCounter(FObjectTrace::GetObjectWorldTickCounter(AnimInstance))
 		<< AnimNodeValueFloat.Attachment(InKey, KeyLength * sizeof(TCHAR));
 }
 
@@ -754,18 +761,19 @@ void FAnimTrace::OutputAnimNodeValue(const FAnimationBaseContext& InContext, con
 
 	check(InContext.AnimInstanceProxy);
 
-	TRACE_OBJECT(InContext.AnimInstanceProxy->GetAnimInstanceObject());
+	UObject* AnimInstance = InContext.AnimInstanceProxy->GetAnimInstanceObject();
+	TRACE_OBJECT(AnimInstance);
 
 	int32 KeyLength = FCString::Strlen(InKey) + 1;
 
 	UE_TRACE_LOG(Animation, AnimNodeValueVector2D, AnimationChannel, KeyLength * sizeof(TCHAR))
 		<< AnimNodeValueVector2D.Cycle(FPlatformTime::Cycles64())
-		<< AnimNodeValueVector2D.AnimInstanceId(FObjectTrace::GetObjectId(InContext.AnimInstanceProxy->GetAnimInstanceObject()))
+		<< AnimNodeValueVector2D.AnimInstanceId(FObjectTrace::GetObjectId(AnimInstance))
 		<< AnimNodeValueVector2D.NodeId(InContext.GetCurrentNodeId())
 		<< AnimNodeValueVector2D.KeyLength(KeyLength)
 		<< AnimNodeValueVector2D.ValueX(InValue.X)
 		<< AnimNodeValueVector2D.ValueY(InValue.Y)
-		<< AnimNodeValueVector2D.FrameCounter((uint16)(GFrameCounter % 0xffff))
+		<< AnimNodeValueVector2D.FrameCounter(FObjectTrace::GetObjectWorldTickCounter(AnimInstance))
 		<< AnimNodeValueVector2D.Attachment(InKey, KeyLength * sizeof(TCHAR));
 }
 
@@ -791,7 +799,8 @@ void FAnimTrace::OutputAnimNodeValue(const FAnimationBaseContext& InContext, con
 		return;
 	}
 
-	TRACE_OBJECT(InContext.AnimInstanceProxy->GetAnimInstanceObject());
+	UObject* AnimInstance = InContext.AnimInstanceProxy->GetAnimInstanceObject();
+	TRACE_OBJECT(AnimInstance);
 
 	int32 KeyLength = FCString::Strlen(InKey) + 1;
 
@@ -803,7 +812,7 @@ void FAnimTrace::OutputAnimNodeValue(const FAnimationBaseContext& InContext, con
 		<< AnimNodeValueVector.ValueX(InValue.X)
 		<< AnimNodeValueVector.ValueY(InValue.Y)
 		<< AnimNodeValueVector.ValueZ(InValue.Z)
-		<< AnimNodeValueVector.FrameCounter((uint16)(GFrameCounter % 0xffff))
+		<< AnimNodeValueVector.FrameCounter(FObjectTrace::GetObjectWorldTickCounter(AnimInstance))
 		<< AnimNodeValueVector.Attachment(InKey, KeyLength * sizeof(TCHAR));
 }
 
@@ -822,7 +831,8 @@ void FAnimTrace::OutputAnimNodeValue(const FAnimationBaseContext& InContext, con
 		return;
 	}
 
-	TRACE_OBJECT(InContext.AnimInstanceProxy->GetAnimInstanceObject());
+	UObject* AnimInstance = InContext.AnimInstanceProxy->GetAnimInstanceObject();
+	TRACE_OBJECT(AnimInstance);
 
 	int32 KeyLength = FCString::Strlen(InKey) + 1;
 	int32 ValueLength = InValue.GetStringLength() + 1;
@@ -835,10 +845,10 @@ void FAnimTrace::OutputAnimNodeValue(const FAnimationBaseContext& InContext, con
 
 	UE_TRACE_LOG(Animation, AnimNodeValueString, AnimationChannel, (KeyLength + ValueLength) * sizeof(TCHAR))
 		<< AnimNodeValueString.Cycle(FPlatformTime::Cycles64())
-		<< AnimNodeValueString.AnimInstanceId(FObjectTrace::GetObjectId(InContext.AnimInstanceProxy->GetAnimInstanceObject()))
+		<< AnimNodeValueString.AnimInstanceId(FObjectTrace::GetObjectId(AnimInstance))
 		<< AnimNodeValueString.NodeId(InContext.GetCurrentNodeId())
 		<< AnimNodeValueString.KeyLength(KeyLength)
-		<< AnimNodeValueString.FrameCounter((uint16)(GFrameCounter % 0xffff))
+		<< AnimNodeValueString.FrameCounter(FObjectTrace::GetObjectWorldTickCounter(AnimInstance))
 		<< AnimNodeValueString.Attachment(StringCopyFunc);
 }
 
@@ -857,7 +867,8 @@ void FAnimTrace::OutputAnimNodeValue(const FAnimationBaseContext& InContext, con
 		return;
 	}
 
-	TRACE_OBJECT(InContext.AnimInstanceProxy->GetAnimInstanceObject());
+	UObject* AnimInstance = InContext.AnimInstanceProxy->GetAnimInstanceObject();
+	TRACE_OBJECT(AnimInstance);
 
 	int32 KeyLength = FCString::Strlen(InKey) + 1;
 	int32 ValueLength = FCString::Strlen(InValue) + 1;
@@ -870,10 +881,10 @@ void FAnimTrace::OutputAnimNodeValue(const FAnimationBaseContext& InContext, con
 
 	UE_TRACE_LOG(Animation, AnimNodeValueString, AnimationChannel, (KeyLength + ValueLength) * sizeof(TCHAR))
 		<< AnimNodeValueString.Cycle(FPlatformTime::Cycles64())
-		<< AnimNodeValueString.AnimInstanceId(FObjectTrace::GetObjectId(InContext.AnimInstanceProxy->GetAnimInstanceObject()))
+		<< AnimNodeValueString.AnimInstanceId(FObjectTrace::GetObjectId(AnimInstance))
 		<< AnimNodeValueString.NodeId(InContext.GetCurrentNodeId())
 		<< AnimNodeValueString.KeyLength(KeyLength)
-		<< AnimNodeValueString.FrameCounter((uint16)(GFrameCounter % 0xffff))
+		<< AnimNodeValueString.FrameCounter(FObjectTrace::GetObjectWorldTickCounter(AnimInstance))
 		<< AnimNodeValueString.Attachment(StringCopyFunc);
 }
 
@@ -892,18 +903,19 @@ void FAnimTrace::OutputAnimNodeValue(const FAnimationBaseContext& InContext, con
 		return;
 	}
 
-	TRACE_OBJECT(InContext.AnimInstanceProxy->GetAnimInstanceObject());
+	UObject* AnimInstance = InContext.AnimInstanceProxy->GetAnimInstanceObject();
+	TRACE_OBJECT(AnimInstance);
 	TRACE_OBJECT(InValue);
 
 	int32 KeyLength = FCString::Strlen(InKey) + 1;
 
 	UE_TRACE_LOG(Animation, AnimNodeValueObject, AnimationChannel, KeyLength * sizeof(TCHAR))
 		<< AnimNodeValueObject.Cycle(FPlatformTime::Cycles64())
-		<< AnimNodeValueObject.AnimInstanceId(FObjectTrace::GetObjectId(InContext.AnimInstanceProxy->GetAnimInstanceObject()))
+		<< AnimNodeValueObject.AnimInstanceId(FObjectTrace::GetObjectId(AnimInstance))
 		<< AnimNodeValueObject.NodeId(InContext.GetCurrentNodeId())
 		<< AnimNodeValueObject.Value(FObjectTrace::GetObjectId(InValue))
 		<< AnimNodeValueObject.KeyLength(KeyLength)
-		<< AnimNodeValueObject.FrameCounter((uint16)(GFrameCounter % 0xffff))
+		<< AnimNodeValueObject.FrameCounter(FObjectTrace::GetObjectWorldTickCounter(AnimInstance))
 		<< AnimNodeValueObject.Attachment(InKey, KeyLength * sizeof(TCHAR));
 }
 
@@ -922,18 +934,19 @@ void FAnimTrace::OutputAnimNodeValue(const FAnimationBaseContext& InContext, con
 		return;
 	}
 
-	TRACE_OBJECT(InContext.AnimInstanceProxy->GetAnimInstanceObject());
+	UObject* AnimInstance = InContext.AnimInstanceProxy->GetAnimInstanceObject();
+	TRACE_OBJECT(AnimInstance);
 	TRACE_CLASS(InValue);
 
 	int32 KeyLength = FCString::Strlen(InKey) + 1;
 
 	UE_TRACE_LOG(Animation, AnimNodeValueClass, AnimationChannel, KeyLength * sizeof(TCHAR))
 		<< AnimNodeValueClass.Cycle(FPlatformTime::Cycles64())
-		<< AnimNodeValueClass.AnimInstanceId(FObjectTrace::GetObjectId(InContext.AnimInstanceProxy->GetAnimInstanceObject()))
+		<< AnimNodeValueClass.AnimInstanceId(FObjectTrace::GetObjectId(AnimInstance))
 		<< AnimNodeValueClass.NodeId(InContext.GetCurrentNodeId())
 		<< AnimNodeValueClass.Value(FObjectTrace::GetObjectId(InValue))
 		<< AnimNodeValueClass.KeyLength(KeyLength)
-		<< AnimNodeValueClass.FrameCounter((uint16)(GFrameCounter % 0xffff))
+		<< AnimNodeValueClass.FrameCounter(FObjectTrace::GetObjectWorldTickCounter(AnimInstance))
 		<< AnimNodeValueClass.Attachment(InKey, KeyLength * sizeof(TCHAR));
 }
 
@@ -952,15 +965,16 @@ void FAnimTrace::OutputAnimSequencePlayer(const FAnimationBaseContext& InContext
 		return;
 	}
 
-	TRACE_OBJECT(InContext.AnimInstanceProxy->GetAnimInstanceObject());
+	UObject* AnimInstance = InContext.AnimInstanceProxy->GetAnimInstanceObject();
+	TRACE_OBJECT(AnimInstance);
 
 	UE_TRACE_LOG(Animation, AnimSequencePlayer, AnimationChannel)
 		<< AnimSequencePlayer.Cycle(FPlatformTime::Cycles64())
-		<< AnimSequencePlayer.AnimInstanceId(FObjectTrace::GetObjectId(InContext.AnimInstanceProxy->GetAnimInstanceObject()))
+		<< AnimSequencePlayer.AnimInstanceId(FObjectTrace::GetObjectId(AnimInstance))
 		<< AnimSequencePlayer.NodeId(InContext.GetCurrentNodeId())
 		<< AnimSequencePlayer.Position(InNode.GetAccumulatedTime())
 		<< AnimSequencePlayer.Length(InNode.Sequence ? InNode.Sequence->SequenceLength : 0.0f)
-		<< AnimSequencePlayer.FrameCount(InNode.Sequence ? InNode.Sequence->GetNumberOfFrames() : 0);
+		<< AnimSequencePlayer.FrameCounter(InNode.Sequence ? InNode.Sequence->GetNumberOfFrames() : 0);
 }
 
 void FAnimTrace::OutputStateMachineState(const FAnimationBaseContext& InContext, int32 InStateMachineIndex, int32 InStateIndex, float InStateWeight, float InElapsedTime)
@@ -978,11 +992,12 @@ void FAnimTrace::OutputStateMachineState(const FAnimationBaseContext& InContext,
 		return;
 	}
 
-	TRACE_OBJECT(InContext.AnimInstanceProxy->GetAnimInstanceObject());
+	UObject* AnimInstance = InContext.AnimInstanceProxy->GetAnimInstanceObject();
+	TRACE_OBJECT(AnimInstance);
 
 	UE_TRACE_LOG(Animation, StateMachineState, AnimationChannel)
 		<< StateMachineState.Cycle(FPlatformTime::Cycles64())
-		<< StateMachineState.AnimInstanceId(FObjectTrace::GetObjectId(InContext.AnimInstanceProxy->GetAnimInstanceObject()))
+		<< StateMachineState.AnimInstanceId(FObjectTrace::GetObjectId(AnimInstance))
 		<< StateMachineState.NodeId(InContext.GetCurrentNodeId())
 		<< StateMachineState.StateMachineIndex(InStateMachineIndex)
 		<< StateMachineState.StateIndex(InStateIndex)
@@ -1085,7 +1100,8 @@ void FAnimTrace::OutputMontage(UAnimInstance* InAnimInstance, const FAnimMontage
 			<< Montage.CurrentSectionNameId(CurrentSectionNameId)
 			<< Montage.NextSectionNameId(NextSectionNameId)
 			<< Montage.Weight(InMontageInstance.GetWeight())
-			<< Montage.DesiredWeight(InMontageInstance.GetDesiredWeight());
+			<< Montage.DesiredWeight(InMontageInstance.GetDesiredWeight())
+			<< Montage.FrameCounter(FObjectTrace::GetObjectWorldTickCounter(InAnimInstance));
 	}
 }
 
diff --git a/Engine/Source/Runtime/Engine/Private/ObjectTrace.cpp b/Engine/Source/Runtime/Engine/Private/ObjectTrace.cpp
--- a/Engine/Source/Runtime/Engine/Private/ObjectTrace.cpp
+++ b/Engine/Source/Runtime/Engine/Private/ObjectTrace.cpp
@@ -15,6 +15,9 @@
 #include "HAL/IConsoleManager.h"
 #include "Engine/World.h"
 #include "TraceFilter.h"
+#if WITH_EDITOR
+#include "Editor.h"
+#endif
 
 UE_TRACE_CHANNEL(ObjectChannel)
 
@@ -37,6 +40,14 @@ UE_TRACE_EVENT_BEGIN(Object, ObjectEvent)
 	UE_TRACE_EVENT_FIELD(uint8, Event)
 UE_TRACE_EVENT_END()
 
+UE_TRACE_EVENT_BEGIN(Object, World, Important)
+	UE_TRACE_EVENT_FIELD(uint64, Id)
+	UE_TRACE_EVENT_FIELD(int32, PIEInstanceId)
+	UE_TRACE_EVENT_FIELD(uint8, Type)
+	UE_TRACE_EVENT_FIELD(uint8, NetMode)
+	UE_TRACE_EVENT_FIELD(bool, IsSimulating)
+UE_TRACE_EVENT_END()
+
 // Object annotations used for tracing
 struct FTracedObjectAnnotation
 {
@@ -61,12 +72,26 @@ struct FTracedObjectAnnotation
 // Object annotations used for tracing
 FUObjectAnnotationSparse<FTracedObjectAnnotation, true> GObjectTraceAnnotations;
 
+// Handle used to hook to world tick
+static FDelegateHandle WorldTickStartHandle;
+
 void FObjectTrace::Init()
 {
-	if (FParse::Param(FCommandLine::Get(), TEXT("objecttrace")))
+	WorldTickStartHandle = FWorldDelegates::OnWorldTickStart.AddLambda([](UWorld* InWorld, ELevelTick InTickType, float InDeltaSeconds)
 	{
-		Trace::ToggleChannel(ObjectChannel, true);
-	}
+		if(InTickType == LEVELTICK_All)
+		{
+			if(UObjectTraceWorldSubsystem* Subsystem = UWorld::GetSubsystem<UObjectTraceWorldSubsystem>(InWorld))
+			{
+				Subsystem->FrameIndex++;
+			}
+		}
+	});
+}
+
+void FObjectTrace::Destroy()
+{
+	FWorldDelegates::OnWorldTickStart.Remove(WorldTickStartHandle);
 }
 
 uint64 FObjectTrace::GetObjectId(const UObject* InObject)
@@ -104,6 +129,22 @@ uint64 FObjectTrace::GetObjectId(const UObject* InObject)
 	return Id | (OuterId << 32);
 }
 
+uint16 FObjectTrace::GetObjectWorldTickCounter(const UObject* InObject)
+{
+	if(InObject != nullptr)
+	{
+		if(UWorld* World = InObject->GetWorld())
+		{
+			if(UObjectTraceWorldSubsystem* WorldSubsystem = UWorld::GetSubsystem<UObjectTraceWorldSubsystem>(World))
+			{
+				return WorldSubsystem->FrameIndex;
+			}
+		}
+	}
+
+	return 0;
+}
+
 void FObjectTrace::OutputClass(const UClass* InClass)
 {
 	bool bChannelEnabled = UE_TRACE_CHANNELEXPR_IS_ENABLED(ObjectChannel);
@@ -214,4 +255,34 @@ void FObjectTrace::OutputObjectEvent(const UObject* InObject, const TCHAR* InEve
 		<< ObjectEvent.Attachment(InEvent, StringBufferSize);
 }
 
+void FObjectTrace::OutputWorld(const UWorld* InWorld)
+{
+	bool bChannelEnabled = UE_TRACE_CHANNELEXPR_IS_ENABLED(ObjectChannel);
+	if (!bChannelEnabled || InWorld == nullptr)
+	{
+		return;
+	}
+
+	if (CANNOT_TRACE_OBJECT(InWorld))
+	{
+		return;
+	}
+
+#if WITH_EDITOR
+	bool bIsSimulating = GEditor->bIsSimulatingInEditor;
+#else
+	bool bIsSimulating = false;
+#endif
+
+	UE_TRACE_LOG(Object, World, ObjectChannel)
+		<< World.Id(GetObjectId(InWorld))
+		<< World.PIEInstanceId(InWorld->GetOutermost()->PIEInstanceID)
+		<< World.Type((uint8)InWorld->WorldType)
+		<< World.NetMode((uint8)InWorld->GetNetMode())
+		<< World.IsSimulating(bIsSimulating);
+
+	// Trace object AFTER world info so we dont risk world info not being present in the trace
+	TRACE_OBJECT(InWorld);
+}
+
 #endif
\ No newline at end of file
diff --git a/Engine/Source/Runtime/Engine/Private/UnrealEngine.cpp b/Engine/Source/Runtime/Engine/Private/UnrealEngine.cpp
--- a/Engine/Source/Runtime/Engine/Private/UnrealEngine.cpp
+++ b/Engine/Source/Runtime/Engine/Private/UnrealEngine.cpp
@@ -231,7 +231,6 @@ UnrealEngine.cpp: Implements the UEngine class and helpers.
 #include "Components/SkinnedMeshComponent.h"
 #include "Net/Core/Trace/NetTrace.h"
 #include "ObjectTrace.h"
-#include "Animation/AnimTrace.h"
 #include "StudioAnalytics.h"
 #include "TraceFilter.h"
 
@@ -271,6 +270,10 @@ void FEngineModule::StartupModule()
 	USkinnedMeshComponent::BindWorldDelegates();
 #endif
 
+#if OBJECT_TRACE_ENABLED
+	FObjectTrace::Init();
+#endif
+
 #if TRACE_FILTERING_ENABLED
 	FTraceFilter::Init();
 #endif
@@ -282,6 +285,10 @@ void FEngineModule::ShutdownModule()
 	FTraceFilter::Destroy();
 #endif
 
+#if OBJECT_TRACE_ENABLED
+	FObjectTrace::Destroy();
+#endif
+
 	FParticleSystemWorldManager::OnShutdown();
 }
 
diff --git a/Engine/Source/Runtime/Engine/Private/World.cpp b/Engine/Source/Runtime/Engine/Private/World.cpp
--- a/Engine/Source/Runtime/Engine/Private/World.cpp
+++ b/Engine/Source/Runtime/Engine/Private/World.cpp
@@ -3985,6 +3985,7 @@ bool UWorld::SetGameMode(const FURL& InURL)
 
 void UWorld::InitializeActorsForPlay(const FURL& InURL, bool bResetTime, FRegisterComponentContext* Context)
 {
+	TRACE_WORLD(this);
 	TRACE_OBJECT_EVENT(this, InitializeActorsForPlay);
 
 	check(bIsWorldInitialized);
diff --git a/Engine/Source/Runtime/Engine/Public/ObjectTrace.h b/Engine/Source/Runtime/Engine/Public/ObjectTrace.h
--- a/Engine/Source/Runtime/Engine/Public/ObjectTrace.h
+++ b/Engine/Source/Runtime/Engine/Public/ObjectTrace.h
@@ -4,6 +4,9 @@
 
 #include "CoreTypes.h"
 #include "Trace/Config.h"
+#include "Subsystems/WorldSubsystem.h"
+
+#include "ObjectTrace.generated.h"
 
 #if UE_TRACE_ENABLED && !IS_PROGRAM && !UE_BUILD_SHIPPING
 #define OBJECT_TRACE_ENABLED 1
@@ -11,6 +14,21 @@
 #define OBJECT_TRACE_ENABLED 0
 #endif
 
+// World subsystem used to track world info
+UCLASS()
+class UObjectTraceWorldSubsystem : public UWorldSubsystem
+{
+	GENERATED_BODY()
+
+	UObjectTraceWorldSubsystem()
+		: FrameIndex(0)
+	{}
+
+public:
+	// The frame index incremented each tick
+	uint16 FrameIndex;
+};
+
 #if OBJECT_TRACE_ENABLED
 
 class UClass;
@@ -21,6 +39,9 @@ struct FObjectTrace
 	/** Initialize object tracing */
 	ENGINE_API static void Init();
 
+	/** Shut down object tracing */
+	ENGINE_API static void Destroy();
+
 	/** Helper function to output an object */
 	ENGINE_API static void OutputClass(const UClass* InClass);
 
@@ -32,6 +53,12 @@ struct FObjectTrace
 
 	/** Helper function to get an object ID from a UObject */
 	ENGINE_API static uint64 GetObjectId(const UObject* InObject);
+
+	/** Helper function to get an object's world's tick counter */
+	ENGINE_API static uint16 GetObjectWorldTickCounter(const UObject* InObject);
+
+	/** Helper function to output a world */
+	ENGINE_API static void OutputWorld(const UWorld* InWorld);
 };
 
 #define TRACE_CLASS(Class) \
@@ -43,10 +70,14 @@ struct FObjectTrace
 #define TRACE_OBJECT_EVENT(Object, Event) \
 	FObjectTrace::OutputObjectEvent(Object, TEXT(#Event));
 
+#define TRACE_WORLD(World) \
+	FObjectTrace::OutputWorld(World);
+
 #else
 
 #define TRACE_CLASS(Class)
 #define TRACE_OBJECT(Object)
 #define TRACE_OBJECT_EVENT(Object, Event)
+#define TRACE_WORLD(World)
 
 #endif
\ No newline at end of file
