Fix non-unity/non-pch builds

diff --git a/Engine/Plugins/Animation/GameplayInsights/Source/GameplayInsights/Private/GameplayProvider.h b/Engine/Plugins/Animation/GameplayInsights/Source/GameplayInsights/Private/GameplayProvider.h
--- a/Engine/Plugins/Animation/GameplayInsights/Source/GameplayInsights/Private/GameplayProvider.h
+++ b/Engine/Plugins/Animation/GameplayInsights/Source/GameplayInsights/Private/GameplayProvider.h
@@ -6,6 +6,7 @@
 #include "Model/PointTimeline.h"
 #include "Containers/Map.h"
 #include "Templates/SharedPointer.h"
+#include "Containers/StringView.h"
 
 namespace Trace { class IAnalysisSession; }
 
