Gameplay Insights is now Animation Insights

String changes only.

diff --git a/Engine/Plugins/Animation/GameplayInsights/GameplayInsights.uplugin b/Engine/Plugins/Animation/GameplayInsights/GameplayInsights.uplugin
--- a/Engine/Plugins/Animation/GameplayInsights/GameplayInsights.uplugin
+++ b/Engine/Plugins/Animation/GameplayInsights/GameplayInsights.uplugin
@@ -2,8 +2,8 @@
 	"FileVersion": 3,
 	"Version": 1,
 	"VersionName": "1.0",
-	"FriendlyName": "Gameplay Insights",
-	"Description": "Allows debugging of gameplay and animation systems via Unreal Insights",
+	"FriendlyName": "Animation Insights",
+	"Description": "Allows debugging of animation systems via Unreal Insights",
 	"Category": "Insights",
 	"CreatedBy": "Epic Games, Inc.",
 	"CreatedByURL": "http://epicgames.com",
diff --git a/Engine/Plugins/Animation/GameplayInsights/Source/GameplayInsights/Private/GameplayInsightsModule.cpp b/Engine/Plugins/Animation/GameplayInsights/Source/GameplayInsights/Private/GameplayInsightsModule.cpp
--- a/Engine/Plugins/Animation/GameplayInsights/Source/GameplayInsights/Private/GameplayInsightsModule.cpp
+++ b/Engine/Plugins/Animation/GameplayInsights/Source/GameplayInsights/Private/GameplayInsightsModule.cpp
@@ -87,8 +87,8 @@ void FGameplayInsightsModule::StartupModule()
 		);
 
 		FInsightsMajorTabConfig TimingProfilerConfig;
-		TimingProfilerConfig.TabLabel = LOCTEXT("GameplayInsightsTabName", "Gameplay Insights");
-		TimingProfilerConfig.TabTooltip = LOCTEXT("GameplayInsightsTabTooltip", "Open the Gameplay Insights tab.");
+		TimingProfilerConfig.TabLabel = LOCTEXT("AnimationInsightsTabName", "Animation Insights");
+		TimingProfilerConfig.TabTooltip = LOCTEXT("AnimationInsightsTabTooltip", "Open the Animation Insights tab.");
 		TimingProfilerConfig.Layout = FTabManager::NewLayout("GameplayInsightsTimingLayout_v1.2")
 		->AddArea
 		(
