Fix crash when interacting with the animation insights transport controls when session is tracing to external application

Also fix scrubbing timeline ruler in the same circumstances

diff --git a/Engine/Plugins/Animation/GameplayInsights/Source/GameplayInsights/Private/AnimationSharedData.cpp b/Engine/Plugins/Animation/GameplayInsights/Source/GameplayInsights/Private/AnimationSharedData.cpp
--- a/Engine/Plugins/Animation/GameplayInsights/Source/GameplayInsights/Private/AnimationSharedData.cpp
+++ b/Engine/Plugins/Animation/GameplayInsights/Source/GameplayInsights/Private/AnimationSharedData.cpp
@@ -467,7 +467,7 @@ void FAnimationSharedData::OnTimeMarkerChanged(Insights::ETimeChangedFlags InFla
 	bTimeMarkerValid = InTimeMarker != std::numeric_limits<double>::infinity();
 	MarkerTime = InTimeMarker;
 
-	if(bTimeMarkerValid)
+	if(bTimeMarkerValid && AnalysisSession != nullptr)
 	{
 		Trace::FAnalysisSessionReadScope SessionReadScope(*AnalysisSession);
 
diff --git a/Engine/Plugins/Animation/GameplayInsights/Source/GameplayInsights/Private/GameplaySharedData.h b/Engine/Plugins/Animation/GameplayInsights/Source/GameplayInsights/Private/GameplaySharedData.h
--- a/Engine/Plugins/Animation/GameplayInsights/Source/GameplayInsights/Private/GameplaySharedData.h
+++ b/Engine/Plugins/Animation/GameplayInsights/Source/GameplayInsights/Private/GameplaySharedData.h
@@ -38,6 +38,9 @@ public:
 	// Get the timing view session
 	Insights::ITimingViewSession& GetTimingViewSession() { return *TimingViewSession; }
 
+	// Check the validity of the analysis session
+	bool IsAnalysisSessionValid() const { return AnalysisSession != nullptr; }
+
 	// Enumerate object tracks
 	void EnumerateObjectTracks(TFunctionRef<void(const TSharedRef<FObjectEventsTrack>&)> InCallback) const;
 
diff --git a/Engine/Plugins/Animation/GameplayInsights/Source/GameplayInsights/Private/GameplayTimingViewExtender.cpp b/Engine/Plugins/Animation/GameplayInsights/Source/GameplayInsights/Private/GameplayTimingViewExtender.cpp
--- a/Engine/Plugins/Animation/GameplayInsights/Source/GameplayInsights/Private/GameplayTimingViewExtender.cpp
+++ b/Engine/Plugins/Animation/GameplayInsights/Source/GameplayInsights/Private/GameplayTimingViewExtender.cpp
@@ -34,7 +34,7 @@ void FGameplayTimingViewExtender::OnBeginSession(Insights::ITimingViewSession& I
 #if WITH_EDITOR
 		InSession.AddOverlayWidget(
 			SNew(SOverlay)
-			.Visibility(EVisibility::SelfHitTestInvisible)
+			.Visibility_Lambda([PerSessionData](){ return PerSessionData->GameplaySharedData->IsAnalysisSessionValid() ? EVisibility::SelfHitTestInvisible : EVisibility::Collapsed; })
 			+SOverlay::Slot()
 			.HAlign(HAlign_Center)
 			.VAlign(VAlign_Bottom)
