Gameplay Insights s
-wait until post engine init and only create local store if one doesn't currently exist.
-This allows Gameplay Insights to coexist with other plugin based insight tools that do the same thing.
-This is a temporary hold over until the formal local store server service is implemented.

#RNX

diff --git a/Engine/Plugins/Animation/GameplayInsights/Source/GameplayInsights/Private/GameplayInsightsModule.cpp b/Engine/Plugins/Animation/GameplayInsights/Source/GameplayInsights/Private/GameplayInsightsModule.cpp
--- a/Engine/Plugins/Animation/GameplayInsights/Source/GameplayInsights/Private/GameplayInsightsModule.cpp
+++ b/Engine/Plugins/Animation/GameplayInsights/Source/GameplayInsights/Private/GameplayInsightsModule.cpp
@@ -57,92 +57,76 @@ void FGameplayInsightsModule::StartupModule()
 			GameplayTimingViewExtender.GetCustomDebugObjects(InAnimationBlueprintEditor, OutDebugList);
 		});
 
-		// Create the Store Service.
-		FString StoreDir = FPaths::ProjectSavedDir() / TEXT("TraceSessions");
-		Trace::FStoreService::FDesc StoreServiceDesc;
-		StoreServiceDesc.StoreDir = *StoreDir;
-		StoreServiceDesc.RecorderPort = 0; // Let system decide port
-		StoreServiceDesc.ThreadCount = 2;
-		StoreService = TSharedPtr<Trace::FStoreService>(Trace::FStoreService::Create(StoreServiceDesc));
-
-		FCoreDelegates::OnPreExit.AddLambda([this]() {
-			StoreService.Reset();
-		});
-
-		// Connect to our newly created store and setup the insights module
-		UnrealInsightsModule.ConnectToStore(TEXT("localhost"), StoreService->GetPort());
-		Trace::SendTo(TEXT("localhost"), StoreService->GetRecorderPort());
-
 		const float DPIScaleFactor = FPlatformApplicationMisc::GetDPIScaleFactorAtPoint(10.0f, 10.0f);
 
 		TSharedRef<FTabManager::FLayout> MajorTabsLayout = FTabManager::NewLayout("GameplayInsightsMajorLayout_v1.0")
-		->AddArea
-		(
-			FTabManager::NewArea(1280.f * DPIScaleFactor, 720.0f * DPIScaleFactor)
-			->Split
+			->AddArea
 			(
+				FTabManager::NewArea(1280.f * DPIScaleFactor, 720.0f * DPIScaleFactor)
+				->Split
+				(
 				FTabManager::NewStack()
 				->AddTab(FInsightsManagerTabs::TimingProfilerTabId, ETabState::ClosedTab)
 			)
-		);
+			);
 
 		FInsightsMajorTabConfig TimingProfilerConfig;
 		TimingProfilerConfig.TabLabel = LOCTEXT("AnimationInsightsTabName", "Animation Insights");
 		TimingProfilerConfig.TabTooltip = LOCTEXT("AnimationInsightsTabTooltip", "Open the Animation Insights tab.");
 		TimingProfilerConfig.Layout = FTabManager::NewLayout("GameplayInsightsTimingLayout_v1.2")
-		->AddArea
-		(
-			FTabManager::NewPrimaryArea()
-			->SetOrientation(Orient_Vertical)
-			->Split
+			->AddArea
 			(
+				FTabManager::NewPrimaryArea()
+				->SetOrientation(Orient_Vertical)
+				->Split
+				(
 				FTabManager::NewStack()
 				->AddTab(FTimingProfilerTabs::ToolbarID, ETabState::ClosedTab)
 				->SetHideTabWell(true)
 			)
-			->Split
-			(
+				->Split
+				(
 				FTabManager::NewSplitter()
 				->SetOrientation(Orient_Horizontal)
 				->Split
 				(
-					FTabManager::NewSplitter()
-					->SetOrientation(Orient_Vertical)
-					->SetSizeCoefficient(0.7f)
-					->Split
-					(
-						FTabManager::NewStack()
-						->SetSizeCoefficient(0.1f)
-						->SetHideTabWell(true)
-						->AddTab(FTimingProfilerTabs::FramesTrackID, ETabState::OpenedTab)
-					)
-					->Split
-					(
-						FTabManager::NewStack()
-						->SetSizeCoefficient(0.9f)
-						->SetHideTabWell(true)
-						->AddTab(FTimingProfilerTabs::TimingViewID, ETabState::OpenedTab)
-					)
-				)
+				FTabManager::NewSplitter()
+				->SetOrientation(Orient_Vertical)
+				->SetSizeCoefficient(0.7f)
 				->Split
 				(
-					FTabManager::NewStack()
-					->SetSizeCoefficient(0.3f)
-					->AddTab(GameplayInsightsTabs::DocumentTab, ETabState::ClosedTab)
-					->AddTab(FTimingProfilerTabs::TimersID, ETabState::ClosedTab)
-					->AddTab(FTimingProfilerTabs::StatsCountersID, ETabState::ClosedTab)
-					->AddTab(FTimingProfilerTabs::CallersID, ETabState::ClosedTab)
-					->AddTab(FTimingProfilerTabs::CalleesID, ETabState::ClosedTab)
-				)
+				FTabManager::NewStack()
+				->SetSizeCoefficient(0.1f)
+				->SetHideTabWell(true)
+				->AddTab(FTimingProfilerTabs::FramesTrackID, ETabState::OpenedTab)
 			)
-			->Split
-			(
+				->Split
+				(
+				FTabManager::NewStack()
+				->SetSizeCoefficient(0.9f)
+				->SetHideTabWell(true)
+				->AddTab(FTimingProfilerTabs::TimingViewID, ETabState::OpenedTab)
+			)
+			)
+				->Split
+				(
+				FTabManager::NewStack()
+				->SetSizeCoefficient(0.3f)
+				->AddTab(GameplayInsightsTabs::DocumentTab, ETabState::ClosedTab)
+				->AddTab(FTimingProfilerTabs::TimersID, ETabState::ClosedTab)
+				->AddTab(FTimingProfilerTabs::StatsCountersID, ETabState::ClosedTab)
+				->AddTab(FTimingProfilerTabs::CallersID, ETabState::ClosedTab)
+				->AddTab(FTimingProfilerTabs::CalleesID, ETabState::ClosedTab)
+			)
+			)
+				->Split
+				(
 				FTabManager::NewStack()
 				->AddTab(FTimingProfilerTabs::LogViewID, ETabState::ClosedTab)
 			)
-		);
+			);
 		TimingProfilerConfig.WorkspaceGroup = WorkspaceMenu::GetMenuStructure().GetDeveloperToolsProfilingCategory();
-
+		/*
 		UnrealInsightsModule.RegisterMajorTabConfig(FInsightsManagerTabs::TimingProfilerTabId, TimingProfilerConfig);
 		UnrealInsightsModule.RegisterMajorTabConfig(FInsightsManagerTabs::StartPageTabId, FInsightsMajorTabConfig::Unavailable());
 		UnrealInsightsModule.RegisterMajorTabConfig(FInsightsManagerTabs::SessionInfoTabId, FInsightsMajorTabConfig::Unavailable());
@@ -150,9 +134,36 @@ void FGameplayInsightsModule::StartupModule()
 		UnrealInsightsModule.RegisterMajorTabConfig(FInsightsManagerTabs::NetworkingProfilerTabId, FInsightsMajorTabConfig::Unavailable());
 
 		UnrealInsightsModule.SetUnrealInsightsLayoutIni(GEditorLayoutIni);
+		*/
+
+
+		// Create store and start analysis session - should only be done after engine has init and all plugins are loaded
+		FCoreDelegates::OnFEngineLoopInitComplete.AddLambda([this]
+		{
+			IUnrealInsightsModule& UnrealInsightsModule = FModuleManager::LoadModuleChecked<IUnrealInsightsModule>("TraceInsights");
+			if (!UnrealInsightsModule.GetStoreClient())
+			{
+				// Create the Store Service.
+				FString StoreDir = FPaths::ProjectSavedDir() / TEXT("TraceSessions");
+				Trace::FStoreService::FDesc StoreServiceDesc;
+				StoreServiceDesc.StoreDir = *StoreDir;
+				StoreServiceDesc.RecorderPort = 0; // Let system decide port
+				StoreServiceDesc.ThreadCount = 2;
+				StoreService = TSharedPtr<Trace::FStoreService>(Trace::FStoreService::Create(StoreServiceDesc));
+
+				FCoreDelegates::OnPreExit.AddLambda([this]() {
+					StoreService.Reset();
+				});
+
+				// Connect to our newly created store and setup the insights module
+				UnrealInsightsModule.ConnectToStore(TEXT("localhost"), StoreService->GetPort());
+				Trace::SendTo(TEXT("localhost"), StoreService->GetRecorderPort());
+
+				UnrealInsightsModule.CreateSessionViewer(false);
+				UnrealInsightsModule.StartAnalysisForLastLiveSession();
+			}
+		});
 
-		UnrealInsightsModule.CreateSessionViewer(false);
-		UnrealInsightsModule.StartAnalysisForLastLiveSession();
 	}
 
 #else
