Fix formatting and renable Major Tab config in Gameplay Insights

diff --git a/Engine/Plugins/Animation/GameplayInsights/Source/GameplayInsights/Private/GameplayInsightsModule.cpp b/Engine/Plugins/Animation/GameplayInsights/Source/GameplayInsights/Private/GameplayInsightsModule.cpp
--- a/Engine/Plugins/Animation/GameplayInsights/Source/GameplayInsights/Private/GameplayInsightsModule.cpp
+++ b/Engine/Plugins/Animation/GameplayInsights/Source/GameplayInsights/Private/GameplayInsightsModule.cpp
@@ -60,73 +60,74 @@ void FGameplayInsightsModule::StartupModule()
 		const float DPIScaleFactor = FPlatformApplicationMisc::GetDPIScaleFactorAtPoint(10.0f, 10.0f);
 
 		TSharedRef<FTabManager::FLayout> MajorTabsLayout = FTabManager::NewLayout("GameplayInsightsMajorLayout_v1.0")
-			->AddArea
+		->AddArea
+		(
+			FTabManager::NewArea(1280.f * DPIScaleFactor, 720.0f * DPIScaleFactor)
+			->Split
 			(
-				FTabManager::NewArea(1280.f * DPIScaleFactor, 720.0f * DPIScaleFactor)
-				->Split
-				(
 				FTabManager::NewStack()
 				->AddTab(FInsightsManagerTabs::TimingProfilerTabId, ETabState::ClosedTab)
 			)
-			);
+		);
 
 		FInsightsMajorTabConfig TimingProfilerConfig;
 		TimingProfilerConfig.TabLabel = LOCTEXT("AnimationInsightsTabName", "Animation Insights");
 		TimingProfilerConfig.TabTooltip = LOCTEXT("AnimationInsightsTabTooltip", "Open the Animation Insights tab.");
 		TimingProfilerConfig.Layout = FTabManager::NewLayout("GameplayInsightsTimingLayout_v1.2")
-			->AddArea
+		->AddArea
+		(
+			FTabManager::NewPrimaryArea()
+			->SetOrientation(Orient_Vertical)
+			->Split
 			(
-				FTabManager::NewPrimaryArea()
-				->SetOrientation(Orient_Vertical)
-				->Split
-				(
 				FTabManager::NewStack()
 				->AddTab(FTimingProfilerTabs::ToolbarID, ETabState::ClosedTab)
 				->SetHideTabWell(true)
 			)
-				->Split
-				(
+			->Split
+			(
 				FTabManager::NewSplitter()
 				->SetOrientation(Orient_Horizontal)
 				->Split
 				(
-				FTabManager::NewSplitter()
-				->SetOrientation(Orient_Vertical)
-				->SetSizeCoefficient(0.7f)
+					FTabManager::NewSplitter()
+					->SetOrientation(Orient_Vertical)
+					->SetSizeCoefficient(0.7f)
+					->Split
+					(
+						FTabManager::NewStack()
+						->SetSizeCoefficient(0.1f)
+						->SetHideTabWell(true)
+						->AddTab(FTimingProfilerTabs::FramesTrackID, ETabState::OpenedTab)
+					)
+					->Split
+					(
+						FTabManager::NewStack()
+						->SetSizeCoefficient(0.9f)
+						->SetHideTabWell(true)
+						->AddTab(FTimingProfilerTabs::TimingViewID, ETabState::OpenedTab)
+					)
+				)
 				->Split
 				(
-				FTabManager::NewStack()
-				->SetSizeCoefficient(0.1f)
-				->SetHideTabWell(true)
-				->AddTab(FTimingProfilerTabs::FramesTrackID, ETabState::OpenedTab)
+					FTabManager::NewStack()
+					->SetSizeCoefficient(0.3f)
+					->AddTab(GameplayInsightsTabs::DocumentTab, ETabState::ClosedTab)
+					->AddTab(FTimingProfilerTabs::TimersID, ETabState::ClosedTab)
+					->AddTab(FTimingProfilerTabs::StatsCountersID, ETabState::ClosedTab)
+					->AddTab(FTimingProfilerTabs::CallersID, ETabState::ClosedTab)
+					->AddTab(FTimingProfilerTabs::CalleesID, ETabState::ClosedTab)
+				)
 			)
-				->Split
-				(
-				FTabManager::NewStack()
-				->SetSizeCoefficient(0.9f)
-				->SetHideTabWell(true)
-				->AddTab(FTimingProfilerTabs::TimingViewID, ETabState::OpenedTab)
-			)
-			)
-				->Split
-				(
-				FTabManager::NewStack()
-				->SetSizeCoefficient(0.3f)
-				->AddTab(GameplayInsightsTabs::DocumentTab, ETabState::ClosedTab)
-				->AddTab(FTimingProfilerTabs::TimersID, ETabState::ClosedTab)
-				->AddTab(FTimingProfilerTabs::StatsCountersID, ETabState::ClosedTab)
-				->AddTab(FTimingProfilerTabs::CallersID, ETabState::ClosedTab)
-				->AddTab(FTimingProfilerTabs::CalleesID, ETabState::ClosedTab)
-			)
-			)
-				->Split
-				(
+			->Split
+			(
 				FTabManager::NewStack()
 				->AddTab(FTimingProfilerTabs::LogViewID, ETabState::ClosedTab)
 			)
-			);
+		);
+
 		TimingProfilerConfig.WorkspaceGroup = WorkspaceMenu::GetMenuStructure().GetDeveloperToolsProfilingCategory();
-		/*
+
 		UnrealInsightsModule.RegisterMajorTabConfig(FInsightsManagerTabs::TimingProfilerTabId, TimingProfilerConfig);
 		UnrealInsightsModule.RegisterMajorTabConfig(FInsightsManagerTabs::StartPageTabId, FInsightsMajorTabConfig::Unavailable());
 		UnrealInsightsModule.RegisterMajorTabConfig(FInsightsManagerTabs::SessionInfoTabId, FInsightsMajorTabConfig::Unavailable());
@@ -134,8 +135,6 @@ void FGameplayInsightsModule::StartupModule()
 		UnrealInsightsModule.RegisterMajorTabConfig(FInsightsManagerTabs::NetworkingProfilerTabId, FInsightsMajorTabConfig::Unavailable());
 
 		UnrealInsightsModule.SetUnrealInsightsLayoutIni(GEditorLayoutIni);
-		*/
-
 
 		// Create store and start analysis session - should only be done after engine has init and all plugins are loaded
 		FCoreDelegates::OnFEngineLoopInitComplete.AddLambda([this]
