Fixed Animation Insights traced data not appearing until PIE is stopped when enabling tracing while PIE is in progress

Logic to build anim tracks originally skipped if there were no event tracks present for an object (e.g. a BeginPlay for a component). Now we go directly to the animation provider to ask if data is present.

Overall much less tracks are created, meaning we are less likely to hit the 1000-track limit and stop visualizing data.

diff --git a/Engine/Plugins/Animation/GameplayInsights/Source/GameplayInsights/Private/AnimationProvider.cpp b/Engine/Plugins/Animation/GameplayInsights/Source/GameplayInsights/Private/AnimationProvider.cpp
--- a/Engine/Plugins/Animation/GameplayInsights/Source/GameplayInsights/Private/AnimationProvider.cpp
+++ b/Engine/Plugins/Animation/GameplayInsights/Source/GameplayInsights/Private/AnimationProvider.cpp
@@ -14,6 +14,7 @@ FAnimationProvider::FAnimationProvider(Trace::IAnalysisSession& InSession, FGame
 	, SkeletalMeshPoseTransforms(InSession.GetLinearAllocator(), 256)
 	, SkeletalMeshCurves(InSession.GetLinearAllocator(), 256)
 	, SkeletalMeshParentIndices(InSession.GetLinearAllocator(), 256)
+	, bHasAnyData(false)
 {
 	GameplayProvider.OnObjectEndPlay().AddRaw(this, &FAnimationProvider::HandleObjectEndPlay);
 }
@@ -368,10 +369,19 @@ FText FAnimationProvider::FormatNodeValue(const FAnimNodeValueMessage& InMessage
 	return Text;
 }
 
+bool FAnimationProvider::HasAnyData() const
+{
+	Session.ReadAccessCheck();
+
+	return bHasAnyData;
+}
+
 void FAnimationProvider::AppendTickRecord(uint64 InAnimInstanceId, double InTime, uint64 InAssetId, int32 InNodeId, float InBlendWeight, float InPlaybackTime, float InRootMotionWeight, float InPlayRate, float InBlendSpacePositionX, float InBlendSpacePositionY, uint16 InFrameCounter, bool bInLooping, bool bInIsBlendSpace)
 {
 	Session.WriteAccessCheck();
 
+	bHasAnyData = true;
+
 	TSharedPtr<FTickRecordTimelineStorage> TimelineStorage;
 	uint32* TimelineStorageIndexPtr = ObjectIdToTickRecordTimelineStorage.Find(InAnimInstanceId);
 	if(TimelineStorageIndexPtr != nullptr)
@@ -415,6 +425,8 @@ void FAnimationProvider::AppendSkeletalMesh(uint64 InObjectId, const TArrayView<
 
 	if(SkeletalMeshIdToIndexMap.Find(InObjectId) == nullptr)
 	{
+		bHasAnyData = true;
+
 		FSkeletalMeshInfo NewSkeletalMeshInfo;
 		NewSkeletalMeshInfo.Id = InObjectId;
 		NewSkeletalMeshInfo.BoneCount = (uint32)InParentIndices.Num();
@@ -434,6 +446,8 @@ void FAnimationProvider::AppendSkeletalMeshComponent(uint64 InObjectId, uint64 I
 {
 	Session.WriteAccessCheck();
 
+	bHasAnyData = true;
+
 	TSharedPtr<FSkeletalMeshTimelineStorage> TimelineStorage;
 	uint32* IndexPtr = ObjectIdToSkeletalMeshPoseTimelines.Find(InObjectId);
 	if(IndexPtr != nullptr)
@@ -488,6 +502,8 @@ void FAnimationProvider::AppendSkeletalMeshComponent(uint64 InObjectId, uint64 I
 {
 	Session.WriteAccessCheck();
 
+	bHasAnyData = true;
+
 	TSharedPtr<FSkeletalMeshTimelineStorage> TimelineStorage;
 	uint32* IndexPtr = ObjectIdToSkeletalMeshPoseTimelines.Find(InObjectId);
 	if(IndexPtr != nullptr)
@@ -544,6 +560,8 @@ void FAnimationProvider::AppendName(uint32 InId, const TCHAR* InName)
 {
 	Session.WriteAccessCheck();
 
+	bHasAnyData = true;
+
 	NameMap.Add(InId, Session.StoreString(InName));
 }
 
@@ -566,6 +584,8 @@ void FAnimationProvider::AppendSkeletalMeshFrame(uint64 InObjectId, double InTim
 {
 	Session.WriteAccessCheck();
 
+	bHasAnyData = true;
+
 	TSharedPtr<Trace::TIntervalTimeline<FSkeletalMeshFrameMessage>> Timeline;
 	uint32* IndexPtr = ObjectIdToSkeletalMeshFrameTimelines.Find(InObjectId);
 	if(IndexPtr != nullptr)
@@ -601,6 +621,8 @@ void FAnimationProvider::AppendAnimGraph(uint64 InAnimInstanceId, double InStart
 {
 	Session.WriteAccessCheck();
 
+	bHasAnyData = true;
+
 	TSharedPtr<Trace::TIntervalTimeline<FAnimGraphMessage>> Timeline;
 	uint32* IndexPtr = ObjectIdToAnimGraphTimelines.Find(InAnimInstanceId);
 	if(IndexPtr != nullptr)
@@ -631,6 +653,8 @@ void FAnimationProvider::AppendAnimNodeStart(uint64 InAnimInstanceId, double InS
 {
 	Session.WriteAccessCheck();
 
+	bHasAnyData = true;
+
 	TSharedPtr<Trace::TPointTimeline<FAnimNodeMessage>> Timeline;
 	uint32* IndexPtr = ObjectIdToAnimNodeTimelines.Find(InAnimInstanceId);
 	if(IndexPtr != nullptr)
@@ -664,6 +688,8 @@ void FAnimationProvider::AppendAnimSequencePlayer(uint64 InAnimInstanceId, doubl
 {
 	Session.WriteAccessCheck();
 
+	bHasAnyData = true;
+
 	TSharedPtr<Trace::TPointTimeline<FAnimSequencePlayerMessage>> Timeline;
 	uint32* IndexPtr = ObjectIdToAnimSequencePlayerTimelines.Find(InAnimInstanceId);
 	if(IndexPtr != nullptr)
@@ -694,6 +720,8 @@ void FAnimationProvider::AppendBlendSpacePlayer(uint64 InAnimInstanceId, double
 {
 	Session.WriteAccessCheck();
 
+	bHasAnyData = true;
+
 	TSharedPtr<Trace::TPointTimeline<FBlendSpacePlayerMessage>> Timeline;
 	uint32* IndexPtr = ObjectIdToBlendSpacePlayerTimelines.Find(InAnimInstanceId);
 	if(IndexPtr != nullptr)
@@ -797,6 +825,8 @@ void FAnimationProvider::AppendAnimNodeValue(uint64 InAnimInstanceId, double InT
 {
 	Session.WriteAccessCheck();
 
+	bHasAnyData = true;
+
 	TSharedPtr<Trace::TPointTimeline<FAnimNodeValueMessage>> Timeline;
 	uint32* IndexPtr = ObjectIdToAnimNodeValueTimelines.Find(InAnimInstanceId);
 	if(IndexPtr != nullptr)
@@ -824,6 +854,8 @@ void FAnimationProvider::AppendStateMachineState(uint64 InAnimInstanceId, double
 {
 	Session.WriteAccessCheck();
 
+	bHasAnyData = true;
+
 	TSharedPtr<Trace::TPointTimeline<FAnimStateMachineMessage>> Timeline;
 	uint32* IndexPtr = ObjectIdToStateMachineTimelines.Find(InAnimInstanceId);
 	if(IndexPtr != nullptr)
@@ -854,6 +886,8 @@ void FAnimationProvider::AppendNotify(uint64 InAnimInstanceId, double InTime, ui
 {
 	Session.WriteAccessCheck();
 
+	bHasAnyData = true;
+
 	// Check if stateful or event-based
 	if(InNotifyEventType == EAnimNotifyMessageType::Begin || InNotifyEventType == EAnimNotifyMessageType::End)
 	{
@@ -938,6 +972,8 @@ void FAnimationProvider::AppendMontage(uint64 InAnimInstanceId, double InTime, u
 {
 	Session.WriteAccessCheck();
 
+	bHasAnyData = true;
+
 	TSharedPtr<FMontageTimelineStorage> TimelineStorage;
 	uint32* IndexPtr = ObjectIdToAnimMontageTimelines.Find(InAnimInstanceId);
 	if(IndexPtr != nullptr)
diff --git a/Engine/Plugins/Animation/GameplayInsights/Source/GameplayInsights/Private/AnimationProvider.h b/Engine/Plugins/Animation/GameplayInsights/Source/GameplayInsights/Private/AnimationProvider.h
--- a/Engine/Plugins/Animation/GameplayInsights/Source/GameplayInsights/Private/AnimationProvider.h
+++ b/Engine/Plugins/Animation/GameplayInsights/Source/GameplayInsights/Private/AnimationProvider.h
@@ -43,6 +43,9 @@ public:
 	virtual FText FormatNodeKeyValue(const FAnimNodeValueMessage& InMessage) const override;
 	virtual FText FormatNodeValue(const FAnimNodeValueMessage& InMessage) const override;
 
+	// Check to see if any data is present for this provider. Used as an early-out when processing.
+	bool HasAnyData() const;
+
 	/** Add a tick record */
 	void AppendTickRecord(uint64 InAnimInstanceId, double InTime, uint64 InAssetId, int32 InNodeId, float InBlendWeight, float InPlaybackTime, float InRootMotionWeight, float InPlayRate, float InBlendSpacePositionX, float InBlendSpacePositionY, uint16 InFrameCounter, bool bInLooping, bool bInIsBlendSpace);
 
@@ -167,4 +170,7 @@ private:
 	TPagedArray<FTransform> SkeletalMeshPoseTransforms;
 	TPagedArray<FSkeletalMeshNamedCurve> SkeletalMeshCurves;
 	TPagedArray<int32> SkeletalMeshParentIndices;
+
+	// Flag to indicate if any data is present
+	bool bHasAnyData;
 };
\ No newline at end of file
diff --git a/Engine/Plugins/Animation/GameplayInsights/Source/GameplayInsights/Private/AnimationSharedData.cpp b/Engine/Plugins/Animation/GameplayInsights/Source/GameplayInsights/Private/AnimationSharedData.cpp
--- a/Engine/Plugins/Animation/GameplayInsights/Source/GameplayInsights/Private/AnimationSharedData.cpp
+++ b/Engine/Plugins/Animation/GameplayInsights/Source/GameplayInsights/Private/AnimationSharedData.cpp
@@ -79,165 +79,187 @@ void FAnimationSharedData::Tick(Insights::ITimingViewSession& InTimingViewSessio
 	{
 		Trace::FAnalysisSessionReadScope SessionReadScope(GetAnalysisSession());
 
-		// Add tracks for each tracked object's animation data
-		GameplayProvider->EnumerateObjects([this, &InTimingViewSession, &InAnalysisSession, &AnimationProvider, &GameplayProvider](const FObjectInfo& InObjectInfo)
+		if(AnimationProvider->HasAnyData() && AreAnyAnimationTracksEnabled())
 		{
-			GameplayProvider->ReadObjectEventsTimeline(InObjectInfo.Id, [this, &InTimingViewSession, &InAnalysisSession, &InObjectInfo, &AnimationProvider, &GameplayProvider](const IGameplayProvider::ObjectEventsTimeline& InTimeline)
+			// Add tracks for each tracked object's animation data
+			GameplayProvider->EnumerateObjects([this, &InTimingViewSession, &InAnalysisSession, &AnimationProvider, &GameplayProvider](const FObjectInfo& InObjectInfo)
 			{
-				if(InTimeline.GetEventCount() > 0)
-				{
-					TSharedRef<FObjectEventsTrack> ObjectEventsTrack = GameplaySharedData.GetObjectEventsTrackForId(InTimingViewSession, InAnalysisSession, InObjectInfo);
+				TSharedPtr<FObjectEventsTrack> ObjectEventsTrack;
 
-					AnimationProvider->ReadSkeletalMeshPoseTimeline(InObjectInfo.Id, [this, &InObjectInfo, &ObjectEventsTrack, &InTimingViewSession](const IAnimationProvider::SkeletalMeshPoseTimeline& InTimeline, bool bInHasCurves)
+				AnimationProvider->ReadSkeletalMeshPoseTimeline(InObjectInfo.Id, [this, &InObjectInfo, &ObjectEventsTrack, &InTimingViewSession, &InAnalysisSession](const IAnimationProvider::SkeletalMeshPoseTimeline& InTimeline, bool bInHasCurves)
+				{
+					if(!ObjectEventsTrack.IsValid())
 					{
-						auto FindSkeletalMeshPoseTrack = [](const FBaseTimingTrack& InTrack)
-						{
-							return InTrack.Is<FSkeletalMeshPoseTrack>();
-						};
-
-						TSharedPtr<FSkeletalMeshPoseTrack> ExistingSkeletalMeshPoseTrack = StaticCastSharedPtr<FSkeletalMeshPoseTrack>(ObjectEventsTrack->GetGameplayTrack().FindChildTrack(InObjectInfo.Id, FindSkeletalMeshPoseTrack));
-						if(!ExistingSkeletalMeshPoseTrack.IsValid())
-						{
-							TSharedPtr<FSkeletalMeshPoseTrack> SkeletalMeshPoseTrack = MakeShared<FSkeletalMeshPoseTrack>(*this, InObjectInfo.Id, InObjectInfo.Name);
-							SkeletalMeshPoseTrack->SetVisibilityFlag(bSkeletalMeshPoseTracksEnabled);
-							SkeletalMeshPoseTracks.Add(SkeletalMeshPoseTrack.ToSharedRef());
-
-							InTimingViewSession.AddScrollableTrack(SkeletalMeshPoseTrack);
-							GameplaySharedData.InvalidateObjectTracksOrder();
-
-							ObjectEventsTrack->GetGameplayTrack().AddChildTrack(SkeletalMeshPoseTrack->GetGameplayTrack());
+						ObjectEventsTrack = GameplaySharedData.GetObjectEventsTrackForId(InTimingViewSession, InAnalysisSession, InObjectInfo);
+					}
 
-							GameplaySharedData.MakeTrackAndAncestorsVisible(ObjectEventsTrack, true);
-						}
-
-						if(bInHasCurves)
-						{
-							auto FindSkeletalMeshCurvesTrack = [](const FBaseTimingTrack& InTrack)
-							{
-								return InTrack.Is<FSkeletalMeshCurvesTrack>();
-							};
+					auto FindSkeletalMeshPoseTrack = [](const FBaseTimingTrack& InTrack)
+					{
+						return InTrack.Is<FSkeletalMeshPoseTrack>();
+					};
 
-							TSharedPtr<FSkeletalMeshCurvesTrack> ExistingSkeletalMeshCurvesTrack = StaticCastSharedPtr<FSkeletalMeshCurvesTrack>(ObjectEventsTrack->GetGameplayTrack().FindChildTrack(InObjectInfo.Id, FindSkeletalMeshCurvesTrack));
-							if(!ExistingSkeletalMeshCurvesTrack.IsValid())
-							{
-								TSharedPtr<FSkeletalMeshCurvesTrack> SkeletalMeshCurvesTrack = MakeShared<FSkeletalMeshCurvesTrack>(*this, InObjectInfo.Id, InObjectInfo.Name);
-								SkeletalMeshCurvesTrack->SetVisibilityFlag(bSkeletalMeshCurveTracksEnabled);
-								SkeletalMeshCurvesTracks.Add(SkeletalMeshCurvesTrack.ToSharedRef());
+					TSharedPtr<FSkeletalMeshPoseTrack> ExistingSkeletalMeshPoseTrack = StaticCastSharedPtr<FSkeletalMeshPoseTrack>(ObjectEventsTrack->GetGameplayTrack().FindChildTrack(InObjectInfo.Id, FindSkeletalMeshPoseTrack));
+					if(!ExistingSkeletalMeshPoseTrack.IsValid())
+					{
+						TSharedPtr<FSkeletalMeshPoseTrack> SkeletalMeshPoseTrack = MakeShared<FSkeletalMeshPoseTrack>(*this, InObjectInfo.Id, InObjectInfo.Name);
+						SkeletalMeshPoseTrack->SetVisibilityFlag(bSkeletalMeshPoseTracksEnabled);
+						SkeletalMeshPoseTracks.Add(SkeletalMeshPoseTrack.ToSharedRef());
 
-								InTimingViewSession.AddScrollableTrack(SkeletalMeshCurvesTrack);
-								GameplaySharedData.InvalidateObjectTracksOrder();
+						InTimingViewSession.AddScrollableTrack(SkeletalMeshPoseTrack);
+						GameplaySharedData.InvalidateObjectTracksOrder();
 
-								ObjectEventsTrack->GetGameplayTrack().AddChildTrack(SkeletalMeshCurvesTrack->GetGameplayTrack());
+						ObjectEventsTrack->GetGameplayTrack().AddChildTrack(SkeletalMeshPoseTrack->GetGameplayTrack());
 
-								GameplaySharedData.MakeTrackAndAncestorsVisible(ObjectEventsTrack, true);
-							}
-						}
-					});
+						GameplaySharedData.MakeTrackAndAncestorsVisible(ObjectEventsTrack.ToSharedRef(), true);
+					}
 
-					AnimationProvider->ReadTickRecordTimeline(InObjectInfo.Id, [this, &InObjectInfo, &ObjectEventsTrack, &InTimingViewSession, &GameplayProvider](const IAnimationProvider::TickRecordTimeline& InTimeline)
+					if(bInHasCurves)
 					{
-						auto FindTickRecordTrackWithAssetId = [](const FBaseTimingTrack& InTrack)
+						auto FindSkeletalMeshCurvesTrack = [](const FBaseTimingTrack& InTrack)
 						{
-							return InTrack.Is<FAnimationTickRecordsTrack>();
+							return InTrack.Is<FSkeletalMeshCurvesTrack>();
 						};
 
-						TSharedPtr<FAnimationTickRecordsTrack> ExistingAnimationTickRecordsTrack = StaticCastSharedPtr<FAnimationTickRecordsTrack>(ObjectEventsTrack->GetGameplayTrack().FindChildTrack(InObjectInfo.Id, FindTickRecordTrackWithAssetId));
-						if(!ExistingAnimationTickRecordsTrack.IsValid())
+						TSharedPtr<FSkeletalMeshCurvesTrack> ExistingSkeletalMeshCurvesTrack = StaticCastSharedPtr<FSkeletalMeshCurvesTrack>(ObjectEventsTrack->GetGameplayTrack().FindChildTrack(InObjectInfo.Id, FindSkeletalMeshCurvesTrack));
+						if(!ExistingSkeletalMeshCurvesTrack.IsValid())
 						{
-							TSharedPtr<FAnimationTickRecordsTrack> AnimationTickRecordsTrack = MakeShared<FAnimationTickRecordsTrack>(*this, InObjectInfo.Id, InObjectInfo.Name);
-							AnimationTickRecordsTrack->SetVisibilityFlag(bTickRecordTracksEnabled);
-							AnimationTickRecordsTracks.Add(AnimationTickRecordsTrack.ToSharedRef());
+							TSharedPtr<FSkeletalMeshCurvesTrack> SkeletalMeshCurvesTrack = MakeShared<FSkeletalMeshCurvesTrack>(*this, InObjectInfo.Id, InObjectInfo.Name);
+							SkeletalMeshCurvesTrack->SetVisibilityFlag(bSkeletalMeshCurveTracksEnabled);
+							SkeletalMeshCurvesTracks.Add(SkeletalMeshCurvesTrack.ToSharedRef());
 
-							InTimingViewSession.AddScrollableTrack(AnimationTickRecordsTrack);
+							InTimingViewSession.AddScrollableTrack(SkeletalMeshCurvesTrack);
 							GameplaySharedData.InvalidateObjectTracksOrder();
 
-							ObjectEventsTrack->GetGameplayTrack().AddChildTrack(AnimationTickRecordsTrack->GetGameplayTrack());
+							ObjectEventsTrack->GetGameplayTrack().AddChildTrack(SkeletalMeshCurvesTrack->GetGameplayTrack());
 
-							GameplaySharedData.MakeTrackAndAncestorsVisible(ObjectEventsTrack, true);
+							GameplaySharedData.MakeTrackAndAncestorsVisible(ObjectEventsTrack.ToSharedRef(), true);
 						}
-					});
+					}
+				});
 
-					AnimationProvider->ReadAnimGraphTimeline(InObjectInfo.Id, [this, &InObjectInfo, &ObjectEventsTrack, &InTimingViewSession](const IAnimationProvider::AnimGraphTimeline& InTimeline)
+				AnimationProvider->ReadTickRecordTimeline(InObjectInfo.Id, [this, &InObjectInfo, &ObjectEventsTrack, &InTimingViewSession, &GameplayProvider, &InAnalysisSession](const IAnimationProvider::TickRecordTimeline& InTimeline)
+				{
+					if(!ObjectEventsTrack.IsValid())
 					{
-						auto FindAnimNodesTrack = [](const FBaseTimingTrack& InTrack)
-						{
-							return InTrack.Is<FAnimNodesTrack>();
-						};
+						ObjectEventsTrack = GameplaySharedData.GetObjectEventsTrackForId(InTimingViewSession, InAnalysisSession, InObjectInfo);
+					}
 
-						TSharedPtr<FAnimNodesTrack> ExistingAnimNodesTrack = StaticCastSharedPtr<FAnimNodesTrack>(ObjectEventsTrack->GetGameplayTrack().FindChildTrack(InObjectInfo.Id, FindAnimNodesTrack));
-						if(!ExistingAnimNodesTrack.IsValid())
-						{
-							TSharedPtr<FAnimNodesTrack> AnimNodesTrack = MakeShared<FAnimNodesTrack>(*this, InObjectInfo.Id, InObjectInfo.Name);
-							AnimNodesTrack->SetVisibilityFlag(bAnimNodeTracksEnabled);
-							AnimNodesTracks.Add(AnimNodesTrack.ToSharedRef());
+					auto FindTickRecordTrackWithAssetId = [](const FBaseTimingTrack& InTrack)
+					{
+						return InTrack.Is<FAnimationTickRecordsTrack>();
+					};
 
-							InTimingViewSession.AddScrollableTrack(AnimNodesTrack);
-							GameplaySharedData.InvalidateObjectTracksOrder();
+					TSharedPtr<FAnimationTickRecordsTrack> ExistingAnimationTickRecordsTrack = StaticCastSharedPtr<FAnimationTickRecordsTrack>(ObjectEventsTrack->GetGameplayTrack().FindChildTrack(InObjectInfo.Id, FindTickRecordTrackWithAssetId));
+					if(!ExistingAnimationTickRecordsTrack.IsValid())
+					{
+						TSharedPtr<FAnimationTickRecordsTrack> AnimationTickRecordsTrack = MakeShared<FAnimationTickRecordsTrack>(*this, InObjectInfo.Id, InObjectInfo.Name);
+						AnimationTickRecordsTrack->SetVisibilityFlag(bTickRecordTracksEnabled);
+						AnimationTickRecordsTracks.Add(AnimationTickRecordsTrack.ToSharedRef());
 
-							ObjectEventsTrack->GetGameplayTrack().AddChildTrack(AnimNodesTrack->GetGameplayTrack());
+						InTimingViewSession.AddScrollableTrack(AnimationTickRecordsTrack);
+						GameplaySharedData.InvalidateObjectTracksOrder();
 
-							GameplaySharedData.MakeTrackAndAncestorsVisible(ObjectEventsTrack, true);
-						}
-					});
+						ObjectEventsTrack->GetGameplayTrack().AddChildTrack(AnimationTickRecordsTrack->GetGameplayTrack());
+
+						GameplaySharedData.MakeTrackAndAncestorsVisible(ObjectEventsTrack.ToSharedRef(), true);
+					}
+				});
 
-					auto ProcessNotifyTimeline = [this, &InObjectInfo, &ObjectEventsTrack, &InTimingViewSession]()
+				AnimationProvider->ReadAnimGraphTimeline(InObjectInfo.Id, [this, &InObjectInfo, &ObjectEventsTrack, &InTimingViewSession, &InAnalysisSession](const IAnimationProvider::AnimGraphTimeline& InTimeline)
+				{
+					if(!ObjectEventsTrack.IsValid())
 					{
-						auto FindAnimNotifyTrack = [](const FBaseTimingTrack& InTrack)
-						{
-							return InTrack.Is<FAnimNotifiesTrack>();
-						};
+						ObjectEventsTrack = GameplaySharedData.GetObjectEventsTrackForId(InTimingViewSession, InAnalysisSession, InObjectInfo);
+					}
 
-						TSharedPtr<FAnimNotifiesTrack> ExistingAnimNotifyTrack = StaticCastSharedPtr<FAnimNotifiesTrack>(ObjectEventsTrack->GetGameplayTrack().FindChildTrack(InObjectInfo.Id, FindAnimNotifyTrack));
-						if(!ExistingAnimNotifyTrack.IsValid())
-						{
-							TSharedPtr<FAnimNotifiesTrack> AnimNotifyTrack = MakeShared<FAnimNotifiesTrack>(*this, InObjectInfo.Id, InObjectInfo.Name);
-							AnimNotifyTrack->SetVisibilityFlag(bAnimNotifyTracksEnabled);
-							AnimNotifyTracks.Add(AnimNotifyTrack.ToSharedRef());
+					auto FindAnimNodesTrack = [](const FBaseTimingTrack& InTrack)
+					{
+						return InTrack.Is<FAnimNodesTrack>();
+					};
 
-							InTimingViewSession.AddScrollableTrack(AnimNotifyTrack);
-							GameplaySharedData.InvalidateObjectTracksOrder();
+					TSharedPtr<FAnimNodesTrack> ExistingAnimNodesTrack = StaticCastSharedPtr<FAnimNodesTrack>(ObjectEventsTrack->GetGameplayTrack().FindChildTrack(InObjectInfo.Id, FindAnimNodesTrack));
+					if(!ExistingAnimNodesTrack.IsValid())
+					{
+						TSharedPtr<FAnimNodesTrack> AnimNodesTrack = MakeShared<FAnimNodesTrack>(*this, InObjectInfo.Id, InObjectInfo.Name);
+						AnimNodesTrack->SetVisibilityFlag(bAnimNodeTracksEnabled);
+						AnimNodesTracks.Add(AnimNodesTrack.ToSharedRef());
 
-							ObjectEventsTrack->GetGameplayTrack().AddChildTrack(AnimNotifyTrack->GetGameplayTrack());
+						InTimingViewSession.AddScrollableTrack(AnimNodesTrack);
+						GameplaySharedData.InvalidateObjectTracksOrder();
 
-							GameplaySharedData.MakeTrackAndAncestorsVisible(ObjectEventsTrack, true);
-						}
+						ObjectEventsTrack->GetGameplayTrack().AddChildTrack(AnimNodesTrack->GetGameplayTrack());
+
+						GameplaySharedData.MakeTrackAndAncestorsVisible(ObjectEventsTrack.ToSharedRef(), true);
+					}
+				});
+
+				auto ProcessNotifyTimeline = [this, &InObjectInfo, &ObjectEventsTrack, &InTimingViewSession, &InAnalysisSession]()
+				{
+					if(!ObjectEventsTrack.IsValid())
+					{
+						ObjectEventsTrack = GameplaySharedData.GetObjectEventsTrackForId(InTimingViewSession, InAnalysisSession, InObjectInfo);
+					}
+
+					auto FindAnimNotifyTrack = [](const FBaseTimingTrack& InTrack)
+					{
+						return InTrack.Is<FAnimNotifiesTrack>();
 					};
 
-					AnimationProvider->ReadNotifyTimeline(InObjectInfo.Id, [&ProcessNotifyTimeline](const IAnimationProvider::AnimNotifyTimeline& InTimeline)
+					TSharedPtr<FAnimNotifiesTrack> ExistingAnimNotifyTrack = StaticCastSharedPtr<FAnimNotifiesTrack>(ObjectEventsTrack->GetGameplayTrack().FindChildTrack(InObjectInfo.Id, FindAnimNotifyTrack));
+					if(!ExistingAnimNotifyTrack.IsValid())
 					{
-						ProcessNotifyTimeline();
-					});
+						TSharedPtr<FAnimNotifiesTrack> AnimNotifyTrack = MakeShared<FAnimNotifiesTrack>(*this, InObjectInfo.Id, InObjectInfo.Name);
+						AnimNotifyTrack->SetVisibilityFlag(bAnimNotifyTracksEnabled);
+						AnimNotifyTracks.Add(AnimNotifyTrack.ToSharedRef());
+
+						InTimingViewSession.AddScrollableTrack(AnimNotifyTrack);
+						GameplaySharedData.InvalidateObjectTracksOrder();
+
+						ObjectEventsTrack->GetGameplayTrack().AddChildTrack(AnimNotifyTrack->GetGameplayTrack());
+
+						GameplaySharedData.MakeTrackAndAncestorsVisible(ObjectEventsTrack.ToSharedRef(), true);
+					}
+				};
+
+				AnimationProvider->ReadNotifyTimeline(InObjectInfo.Id, [&ProcessNotifyTimeline](const IAnimationProvider::AnimNotifyTimeline& InTimeline)
+				{
+					ProcessNotifyTimeline();
+				});
 
-					AnimationProvider->EnumerateNotifyStateTimelines(InObjectInfo.Id, [&ProcessNotifyTimeline](uint32 InNotifyNameId, const IAnimationProvider::AnimNotifyTimeline& InTimeline)
+				AnimationProvider->EnumerateNotifyStateTimelines(InObjectInfo.Id, [&ProcessNotifyTimeline](uint32 InNotifyNameId, const IAnimationProvider::AnimNotifyTimeline& InTimeline)
+				{
+					ProcessNotifyTimeline();
+				});
+
+				AnimationProvider->ReadMontageTimeline(InObjectInfo.Id, [this, &InObjectInfo, &ObjectEventsTrack, &InTimingViewSession, &InAnalysisSession](const IAnimationProvider::AnimMontageTimeline& InTimeline)
+				{
+					if(!ObjectEventsTrack.IsValid())
 					{
-						ProcessNotifyTimeline();
-					});
+						ObjectEventsTrack = GameplaySharedData.GetObjectEventsTrackForId(InTimingViewSession, InAnalysisSession, InObjectInfo);
+					}
 
-					AnimationProvider->ReadMontageTimeline(InObjectInfo.Id, [this, &InObjectInfo, &ObjectEventsTrack, &InTimingViewSession](const IAnimationProvider::AnimMontageTimeline& InTimeline)
+					auto FindMontageTrack = [](const FBaseTimingTrack& InTrack)
 					{
-						auto FindMontageTrack = [](const FBaseTimingTrack& InTrack)
-						{
-							return InTrack.Is<FMontageTrack>();
-						};
+						return InTrack.Is<FMontageTrack>();
+					};
 
-						TSharedPtr<FMontageTrack> ExistingMontageTrack = StaticCastSharedPtr<FMontageTrack>(ObjectEventsTrack->GetGameplayTrack().FindChildTrack(InObjectInfo.Id, FindMontageTrack));
-						if(!ExistingMontageTrack.IsValid())
-						{
-							TSharedPtr<FMontageTrack> MontageTrack = MakeShared<FMontageTrack>(*this, InObjectInfo.Id, InObjectInfo.Name);
-							MontageTrack->SetVisibilityFlag(bMontageTracksEnabled);
-							MontageTracks.Add(MontageTrack.ToSharedRef());
+					TSharedPtr<FMontageTrack> ExistingMontageTrack = StaticCastSharedPtr<FMontageTrack>(ObjectEventsTrack->GetGameplayTrack().FindChildTrack(InObjectInfo.Id, FindMontageTrack));
+					if(!ExistingMontageTrack.IsValid())
+					{
+						TSharedPtr<FMontageTrack> MontageTrack = MakeShared<FMontageTrack>(*this, InObjectInfo.Id, InObjectInfo.Name);
+						MontageTrack->SetVisibilityFlag(bMontageTracksEnabled);
+						MontageTracks.Add(MontageTrack.ToSharedRef());
 
-							InTimingViewSession.AddScrollableTrack(MontageTrack);
-							GameplaySharedData.InvalidateObjectTracksOrder();
+						InTimingViewSession.AddScrollableTrack(MontageTrack);
+						GameplaySharedData.InvalidateObjectTracksOrder();
 
-							ObjectEventsTrack->GetGameplayTrack().AddChildTrack(MontageTrack->GetGameplayTrack());
+						ObjectEventsTrack->GetGameplayTrack().AddChildTrack(MontageTrack->GetGameplayTrack());
 
-							GameplaySharedData.MakeTrackAndAncestorsVisible(ObjectEventsTrack, true);
-						}
-					});
-				}
+						GameplaySharedData.MakeTrackAndAncestorsVisible(ObjectEventsTrack.ToSharedRef(), true);
+					}
+				});
 			});
-		});
+		}
 	}
 
 	// Prevent mouse movement throttling if we are drawing stuff that can change when the mouse is dragged
@@ -287,7 +309,7 @@ void FAnimationSharedData::ExtendFilterMenu(FMenuBuilder& InMenuBuilder)
 			FUIAction(
 				FExecuteAction::CreateRaw(this, &FAnimationSharedData::ToggleAnimationTracks),
 				FCanExecuteAction(),
-				FIsActionChecked::CreateRaw(this, &FAnimationSharedData::AreAnimationTracksEnabled)),
+				FIsActionChecked::CreateRaw(this, &FAnimationSharedData::AreAllAnimationTracksEnabled)),
 			NAME_None,
 			EUserInterfaceActionType::ToggleButton
 		);
@@ -369,7 +391,7 @@ void FAnimationSharedData::ExtendFilterMenu(FMenuBuilder& InMenuBuilder)
 
 void FAnimationSharedData::ToggleAnimationTracks()
 {
-	bool bAnimationTracksEnabled = !AreAnimationTracksEnabled();
+	bool bAnimationTracksEnabled = !AreAllAnimationTracksEnabled();
 
 	bSkeletalMeshPoseTracksEnabled = bAnimationTracksEnabled;
 	bSkeletalMeshCurveTracksEnabled = bAnimationTracksEnabled;
@@ -409,11 +431,16 @@ void FAnimationSharedData::ToggleAnimationTracks()
 	}
 }
 
-bool FAnimationSharedData::AreAnimationTracksEnabled() const
+bool FAnimationSharedData::AreAllAnimationTracksEnabled() const
 {
 	return bSkeletalMeshPoseTracksEnabled && bSkeletalMeshCurveTracksEnabled && bTickRecordTracksEnabled && bAnimNodeTracksEnabled && bAnimNotifyTracksEnabled && bMontageTracksEnabled;
 }
 
+bool FAnimationSharedData::AreAnyAnimationTracksEnabled() const
+{
+	return bSkeletalMeshPoseTracksEnabled || bSkeletalMeshCurveTracksEnabled || bTickRecordTracksEnabled || bAnimNodeTracksEnabled || bAnimNotifyTracksEnabled || bMontageTracksEnabled;
+}
+
 void FAnimationSharedData::ToggleSkeletalMeshPoseTracks()
 {
 	bSkeletalMeshPoseTracksEnabled = !bSkeletalMeshPoseTracksEnabled;
diff --git a/Engine/Plugins/Animation/GameplayInsights/Source/GameplayInsights/Private/AnimationSharedData.h b/Engine/Plugins/Animation/GameplayInsights/Source/GameplayInsights/Private/AnimationSharedData.h
--- a/Engine/Plugins/Animation/GameplayInsights/Source/GameplayInsights/Private/AnimationSharedData.h
+++ b/Engine/Plugins/Animation/GameplayInsights/Source/GameplayInsights/Private/AnimationSharedData.h
@@ -50,8 +50,11 @@ public:
 	void HideAllMeshes();
 #endif
 
-	// Check whether animation tracks are enabled
-	bool AreAnimationTracksEnabled() const;
+	// Check whether all animation tracks are enabled
+	bool AreAllAnimationTracksEnabled() const;
+
+	// Check whether any animation tracks are enabled
+	bool AreAnyAnimationTracksEnabled() const;
 
 	// Get the last cached analysis session
 	const Trace::IAnalysisSession& GetAnalysisSession() const { return *AnalysisSession; }
@@ -122,7 +125,6 @@ private:
 	/** Validity flags for pose times/ranges */
 	bool bTimeMarkerValid;
 	bool bMarkerFrameValid;
-	bool bSelectionRangeValid;
 
 	// Flags controlling whether check type of our animation tracks are enabled
 	bool bSkeletalMeshPoseTracksEnabled;
diff --git a/Engine/Plugins/Animation/GameplayInsights/Source/GameplayInsights/Private/GameplayProvider.cpp b/Engine/Plugins/Animation/GameplayInsights/Source/GameplayInsights/Private/GameplayProvider.cpp
--- a/Engine/Plugins/Animation/GameplayInsights/Source/GameplayInsights/Private/GameplayProvider.cpp
+++ b/Engine/Plugins/Animation/GameplayInsights/Source/GameplayInsights/Private/GameplayProvider.cpp
@@ -9,6 +9,8 @@ FName FGameplayProvider::ProviderName("GameplayProvider");
 FGameplayProvider::FGameplayProvider(Trace::IAnalysisSession& InSession)
 	: Session(InSession)
 	, EndPlayEvent(nullptr)
+	, bHasAnyData(false)
+	, bHasObjectProperties(false)
 {
 }
 
@@ -230,6 +232,8 @@ void FGameplayProvider::AppendClass(uint64 InClassId, uint64 InSuperId, const TC
 {
 	Session.WriteAccessCheck();
 
+	bHasAnyData = true;
+
 	if(ClassIdToIndexMap.Find(InClassId) == nullptr)
 	{
 		const TCHAR* NewClassName = Session.StoreString(InClassName);
@@ -251,6 +255,8 @@ void FGameplayProvider::AppendObject(uint64 InObjectId, uint64 InOuterId, uint64
 {
 	Session.WriteAccessCheck();
 
+	bHasAnyData = true;
+
 	if(ObjectIdToIndexMap.Find(InObjectId) == nullptr)
 	{
 		const TCHAR* NewObjectName = Session.StoreString(InObjectName);
@@ -272,6 +278,8 @@ void FGameplayProvider::AppendObjectEvent(uint64 InObjectId, double InTime, cons
 {
 	Session.WriteAccessCheck();
 
+	bHasAnyData = true;
+
 	// Important events need some extra routing
 	if(EndPlayEvent == nullptr)
 	{
@@ -312,6 +320,8 @@ void FGameplayProvider::AppendWorld(uint64 InObjectId, int32 InPIEInstanceId, ui
 {
 	Session.WriteAccessCheck();
 
+	bHasAnyData = true;
+
 	if (WorldIdToIndexMap.Find(InObjectId) == nullptr)
 	{
 		FWorldInfo NewWorldInfo;
@@ -330,6 +340,8 @@ void FGameplayProvider::AppendClassPropertyStringId(uint32 InStringId, const FSt
 {
 	Session.WriteAccessCheck();
 
+	bHasAnyData = true;
+
 	const TCHAR* StoredString = Session.StoreString(InString);
 
 	PropertyStrings.Add(InStringId, StoredString);
@@ -339,6 +351,8 @@ void FGameplayProvider::AppendClassProperty(uint64 InClassId, int32 InId, int32
 {
 	Session.WriteAccessCheck();
 
+	bHasAnyData = true;
+
 	if(int32* ClassInfoIndexPtr = ClassIdToIndexMap.Find(InClassId))
 	{
 		FClassInfo& ClassInfo = ClassInfos[*ClassInfoIndexPtr];
@@ -357,6 +371,9 @@ void FGameplayProvider::AppendPropertiesStart(uint64 InObjectId, double InTime,
 {
 	Session.WriteAccessCheck();
 
+	bHasAnyData = true;
+	bHasObjectProperties = true;
+
 	TSharedPtr<FObjectPropertiesStorage> Storage;
 	uint32* IndexPtr = ObjectIdToPropertiesStorage.Find(InObjectId);
 	if(IndexPtr != nullptr)
@@ -383,6 +400,9 @@ void FGameplayProvider::AppendPropertiesEnd(uint64 InObjectId, double InTime)
 {
 	Session.WriteAccessCheck();
 
+	bHasAnyData = true;
+	bHasObjectProperties = true;
+
 	TSharedPtr<FObjectPropertiesStorage> Storage;
 	uint32* IndexPtr = ObjectIdToPropertiesStorage.Find(InObjectId);
 	if(IndexPtr != nullptr)
@@ -410,6 +430,9 @@ void FGameplayProvider::AppendPropertyValue(uint64 InObjectId, double InTime, ui
 {
 	Session.WriteAccessCheck();
 
+	bHasAnyData = true;
+	bHasObjectProperties = true;
+
 	TSharedPtr<FObjectPropertiesStorage> Storage;
 	uint32* IndexPtr = ObjectIdToPropertiesStorage.Find(InObjectId);
 	if(IndexPtr != nullptr)
@@ -435,4 +458,18 @@ void FGameplayProvider::AppendPropertyValue(uint64 InObjectId, double InTime, ui
 	}
 }
 
+bool FGameplayProvider::HasAnyData() const 
+{
+	Session.ReadAccessCheck();
+
+	return bHasAnyData;
+}
+
+bool FGameplayProvider::HasObjectProperties() const 
+{
+	Session.ReadAccessCheck();
+
+	return bHasObjectProperties;
+}
+
 #undef LOCTEXT_NAMESPACE
\ No newline at end of file
diff --git a/Engine/Plugins/Animation/GameplayInsights/Source/GameplayInsights/Private/GameplayProvider.h b/Engine/Plugins/Animation/GameplayInsights/Source/GameplayInsights/Private/GameplayProvider.h
--- a/Engine/Plugins/Animation/GameplayInsights/Source/GameplayInsights/Private/GameplayProvider.h
+++ b/Engine/Plugins/Animation/GameplayInsights/Source/GameplayInsights/Private/GameplayProvider.h
@@ -63,6 +63,12 @@ public:
 	/** Add a property value message */
 	void AppendPropertyValue(uint64 InObjectId, double InTime, uint64 InEventId, int32 InPropertyId, const FStringView& InValue);
 
+	/** Check whether we have any data */
+	bool HasAnyData() const;
+
+	/** Check whether we have any object property data */
+	bool HasObjectProperties() const;
+
 private:
 	Trace::IAnalysisSession& Session;
 
@@ -112,4 +118,10 @@ private:
 
 	/** Map from string ID to stored string */
 	TMap<uint32, const TCHAR*> PropertyStrings;
+
+	/** Whether we have any data */
+	bool bHasAnyData;
+
+	/** Whether we have any object properties */
+	bool bHasObjectProperties;
 };
\ No newline at end of file
diff --git a/Engine/Plugins/Animation/GameplayInsights/Source/GameplayInsights/Private/GameplaySharedData.cpp b/Engine/Plugins/Animation/GameplayInsights/Source/GameplayInsights/Private/GameplaySharedData.cpp
--- a/Engine/Plugins/Animation/GameplayInsights/Source/GameplayInsights/Private/GameplaySharedData.cpp
+++ b/Engine/Plugins/Animation/GameplayInsights/Source/GameplayInsights/Private/GameplaySharedData.cpp
@@ -146,17 +146,30 @@ void FGameplaySharedData::Tick(Insights::ITimingViewSession& InTimingViewSession
 	{
 		Trace::FAnalysisSessionReadScope SessionReadScope(GetAnalysisSession());
 
-		// Add a track for each tracked object
-		GameplayProvider->EnumerateObjects([this, &InTimingViewSession, &InAnalysisSession, &GameplayProvider](const FObjectInfo& InObjectInfo)
+		if(GameplayProvider->HasAnyData())
 		{
-			GameplayProvider->ReadObjectEventsTimeline(InObjectInfo.Id, [this, &InTimingViewSession, &InAnalysisSession, &InObjectInfo, &GameplayProvider](const IGameplayProvider::ObjectEventsTimeline& InTimeline)
+			// Add a track for each tracked object
+			GameplayProvider->EnumerateObjects([this, &InTimingViewSession, &InAnalysisSession, &GameplayProvider](const FObjectInfo& InObjectInfo)
 			{
-				if(InTimeline.GetEventCount() > 0)
+				if(bObjectTracksEnabled || bObjectPropertyTracksEnabled)
 				{
-					TSharedRef<FObjectEventsTrack> ObjectEventsTrack = GetObjectEventsTrackForId(InTimingViewSession, InAnalysisSession, InObjectInfo);
+					TSharedPtr<FObjectEventsTrack> ObjectEventsTrack;
 
-					GameplayProvider->ReadObjectPropertiesTimeline(InObjectInfo.Id, [this, &InObjectInfo, &ObjectEventsTrack, &InTimingViewSession, &GameplayProvider](const IGameplayProvider::ObjectPropertiesTimeline& InTimeline)
+					GameplayProvider->ReadObjectEventsTimeline(InObjectInfo.Id, [this, &InTimingViewSession, &InAnalysisSession, &InObjectInfo, &GameplayProvider, &ObjectEventsTrack](const IGameplayProvider::ObjectEventsTimeline& InTimeline)
 					{
+						if(bObjectTracksEnabled)
+						{
+							ObjectEventsTrack = GetObjectEventsTrackForId(InTimingViewSession, InAnalysisSession, InObjectInfo);
+						}
+					});
+
+					GameplayProvider->ReadObjectPropertiesTimeline(InObjectInfo.Id, [this, &InObjectInfo, &ObjectEventsTrack, &InTimingViewSession, &InAnalysisSession, &GameplayProvider](const IGameplayProvider::ObjectPropertiesTimeline& InTimeline)
+					{
+						if(!ObjectEventsTrack.IsValid())
+						{
+							ObjectEventsTrack = GetObjectEventsTrackForId(InTimingViewSession, InAnalysisSession, InObjectInfo);
+						}
+
 						auto FindObjectProperties = [](const FBaseTimingTrack& InTrack)
 						{
 							return InTrack.Is<FObjectPropertiesTrack>();
@@ -174,12 +187,12 @@ void FGameplaySharedData::Tick(Insights::ITimingViewSession& InTimingViewSession
 
 							ObjectEventsTrack->GetGameplayTrack().AddChildTrack(ObjectPropertiesTrack->GetGameplayTrack());
 
-							MakeTrackAndAncestorsVisible(ObjectEventsTrack, true);
+							MakeTrackAndAncestorsVisible(ObjectEventsTrack.ToSharedRef(), true);
 						}
 					});
 				}
 			});
-		});
+		}
 
 		if(bObjectTracksDirty)
 		{
