Moved the TraceServices module into its own namespace.

#rnx

diff --git a/Engine/Plugins/Animation/GameplayInsights/Source/GameplayInsights/Private/AnimNodesTrack.cpp b/Engine/Plugins/Animation/GameplayInsights/Source/GameplayInsights/Private/AnimNodesTrack.cpp
--- a/Engine/Plugins/Animation/GameplayInsights/Source/GameplayInsights/Private/AnimNodesTrack.cpp
+++ b/Engine/Plugins/Animation/GameplayInsights/Source/GameplayInsights/Private/AnimNodesTrack.cpp
@@ -108,14 +108,14 @@ void FAnimNodesTrack::BuildDrawState(ITimingEventsTrackDrawStateBuilder& Builder
 
 	if(AnimationProvider)
 	{
-		Trace::FAnalysisSessionReadScope SessionReadScope(SharedData.GetAnalysisSession());
+		TraceServices::FAnalysisSessionReadScope SessionReadScope(SharedData.GetAnalysisSession());
 
 		AnimationProvider->ReadAnimGraphTimeline(GetGameplayTrack().GetObjectId(), [&Context, &Builder](const FAnimationProvider::AnimGraphTimeline& InTimeline)
 		{
 			InTimeline.EnumerateEvents(Context.GetViewport().GetStartTime(), Context.GetViewport().GetEndTime(), [&Builder](double InStartTime, double InEndTime, uint32 InDepth, const FAnimGraphMessage& InMessage)
 			{
 				Builder.AddEvent(InStartTime, InEndTime, 0, GetPhaseName(InMessage.Phase));
-				return Trace::EEventEnumerate::Continue;
+				return TraceServices::EEventEnumerate::Continue;
 			});
 		});
 	}
@@ -168,14 +168,14 @@ void FAnimNodesTrack::FindAnimGraphMessage(const FTimingEventSearchParameters& I
 
 			if(AnimationProvider)
 			{
-				Trace::FAnalysisSessionReadScope SessionReadScope(SharedData.GetAnalysisSession());
+				TraceServices::FAnalysisSessionReadScope SessionReadScope(SharedData.GetAnalysisSession());
 
 				AnimationProvider->ReadAnimGraphTimeline(GetGameplayTrack().GetObjectId(), [&InContext](const FAnimationProvider::AnimGraphTimeline& InTimeline)
 				{
 					InTimeline.EnumerateEvents(InContext.GetParameters().StartTime, InContext.GetParameters().EndTime, [&InContext](double InEventStartTime, double InEventEndTime, uint32 InDepth, const FAnimGraphMessage& InMessage)
 					{
 						InContext.Check(InEventStartTime, InEventEndTime, 0, InMessage);
-						return Trace::EEventEnumerate::Continue;
+						return TraceServices::EEventEnumerate::Continue;
 					});
 				});
 			}
@@ -208,7 +208,7 @@ void FAnimNodesTrack::BuildContextMenu(FMenuBuilder& MenuBuilder)
 							const FGameplayProvider* GameplayProvider = SharedData.GetAnalysisSession().ReadProvider<FGameplayProvider>(FGameplayProvider::ProviderName);
 							if(GameplayProvider)
 							{
-								Trace::FAnalysisSessionReadScope SessionReadScope(SharedData.GetAnalysisSession());
+								TraceServices::FAnalysisSessionReadScope SessionReadScope(SharedData.GetAnalysisSession());
 
 								if(const FObjectInfo* AnimInstanceObjectInfo = GameplayProvider->FindObjectInfo(GetGameplayTrack().GetObjectId()))
 								{
@@ -242,7 +242,7 @@ void FAnimNodesTrack::BuildContextMenu(FMenuBuilder& MenuBuilder)
 							const FGameplayProvider* GameplayProvider = SharedData.GetAnalysisSession().ReadProvider<FGameplayProvider>(FGameplayProvider::ProviderName);
 							if(GameplayProvider)
 							{
-								Trace::FAnalysisSessionReadScope SessionReadScope(SharedData.GetAnalysisSession());
+								TraceServices::FAnalysisSessionReadScope SessionReadScope(SharedData.GetAnalysisSession());
 
 								if(const FObjectInfo* AnimInstanceObjectInfo = GameplayProvider->FindObjectInfo(GetGameplayTrack().GetObjectId()))
 								{
@@ -300,7 +300,7 @@ UAnimInstance* FAnimNodesTrack::LazyCreateAnimInstance(USkeletalMeshComponent* I
 	return nullptr;
 }
 
-void FAnimNodesTrack::UpdateDebugData(const Trace::FFrame& InFrame)
+void FAnimNodesTrack::UpdateDebugData(const TraceServices::FFrame& InFrame)
 {
 	if(InstanceClass.LoadSynchronous())
 	{
@@ -309,7 +309,7 @@ void FAnimNodesTrack::UpdateDebugData(const Trace::FFrame& InFrame)
 
 		if(AnimationProvider && GameplayProvider)
 		{
-			Trace::FAnalysisSessionReadScope SessionReadScope(SharedData.GetAnalysisSession());
+			TraceServices::FAnalysisSessionReadScope SessionReadScope(SharedData.GetAnalysisSession());
 
 			AnimationProvider->ReadAnimGraphTimeline(GetGameplayTrack().GetObjectId(), [this, AnimationProvider, GameplayProvider, InFrame](const FAnimationProvider::AnimGraphTimeline& InGraphTimeline)
 			{
@@ -333,7 +333,7 @@ void FAnimNodesTrack::UpdateDebugData(const Trace::FFrame& InFrame)
 									{
 										DebugData.RecordNodeVisit(InMessage.NodeId, InMessage.PreviousNodeId, InMessage.Weight);
 									}
-									return Trace::EEventEnumerate::Continue;
+									return TraceServices::EEventEnumerate::Continue;
 								});
 							});
 
@@ -342,7 +342,7 @@ void FAnimNodesTrack::UpdateDebugData(const Trace::FFrame& InFrame)
 								InStateMachinesTimeline.EnumerateEvents(InGraphStartTime, InGraphEndTime, [&DebugData](double InStartTime, double InEndTime, uint32 InDepth, const FAnimStateMachineMessage& InMessage)
 								{
 									DebugData.RecordStateData(InMessage.StateMachineIndex, InMessage.StateIndex, InMessage.StateWeight, InMessage.ElapsedTime);
-									return Trace::EEventEnumerate::Continue;
+									return TraceServices::EEventEnumerate::Continue;
 								});
 							});
 
@@ -352,7 +352,7 @@ void FAnimNodesTrack::UpdateDebugData(const Trace::FFrame& InFrame)
 								{
 									FText Text = AnimationProvider->FormatNodeKeyValue(InMessage);
 									DebugData.RecordNodeValue(InMessage.NodeId, Text.ToString());
-									return Trace::EEventEnumerate::Continue;
+									return TraceServices::EEventEnumerate::Continue;
 								});
 							});
 
@@ -361,7 +361,7 @@ void FAnimNodesTrack::UpdateDebugData(const Trace::FFrame& InFrame)
 								InSequencePlayersTimeline.EnumerateEvents(InGraphStartTime, InGraphEndTime, [&DebugData](double InStartTime, double InEndTime, uint32 InDepth, const FAnimSequencePlayerMessage& InMessage)
 								{
 									DebugData.RecordSequencePlayer(InMessage.NodeId, InMessage.Position, InMessage.Length, InMessage.FrameCounter);
-									return Trace::EEventEnumerate::Continue;
+									return TraceServices::EEventEnumerate::Continue;
 								});
 							});
 
@@ -377,12 +377,12 @@ void FAnimNodesTrack::UpdateDebugData(const Trace::FFrame& InFrame)
 									}
 
 									DebugData.RecordBlendSpacePlayer(InMessage.NodeId, BlendSpaceBase, InMessage.PositionX, InMessage.PositionY, InMessage.PositionZ);
-									return Trace::EEventEnumerate::Continue;
+									return TraceServices::EEventEnumerate::Continue;
 								});
 							});
 						}
 					}
-					return Trace::EEventEnumerate::Continue;
+					return TraceServices::EEventEnumerate::Continue;
 				});
 			});
 		}
@@ -398,7 +398,7 @@ void FAnimNodesTrack::GetCustomDebugObjects(const IAnimationBlueprintEditor& InA
 			const FGameplayProvider* GameplayProvider = SharedData.GetAnalysisSession().ReadProvider<FGameplayProvider>(FGameplayProvider::ProviderName);
 			if(GameplayProvider)
 			{
-				Trace::FAnalysisSessionReadScope SessionReadScope(SharedData.GetAnalysisSession());
+				TraceServices::FAnalysisSessionReadScope SessionReadScope(SharedData.GetAnalysisSession());
 
 				if(const FObjectInfo* AnimInstanceObjectInfo = GameplayProvider->FindObjectInfo(GetGameplayTrack().GetObjectId()))
 				{
diff --git a/Engine/Plugins/Animation/GameplayInsights/Source/GameplayInsights/Private/AnimNodesTrack.h b/Engine/Plugins/Animation/GameplayInsights/Source/GameplayInsights/Private/AnimNodesTrack.h
--- a/Engine/Plugins/Animation/GameplayInsights/Source/GameplayInsights/Private/AnimNodesTrack.h
+++ b/Engine/Plugins/Animation/GameplayInsights/Source/GameplayInsights/Private/AnimNodesTrack.h
@@ -37,7 +37,7 @@ public:
 
 #if WITH_EDITOR
 	// Update the debug data for any classes that are running on this track
-	void UpdateDebugData(const Trace::FFrame& InFrame);
+	void UpdateDebugData(const TraceServices::FFrame& InFrame);
 
 	// Get custom debug objects for integration with anim blueprint debugging
 	void GetCustomDebugObjects(const IAnimationBlueprintEditor& InAnimationBlueprintEditor, TArray<FCustomDebugObject>& OutDebugList);
@@ -79,4 +79,4 @@ private:
 	FDelegateHandle OnWorldBeginTearDownHandle;
 	FDelegateHandle OnPreWorldFinishDestroyHandle;
 #endif
-};
\ No newline at end of file
+};
diff --git a/Engine/Plugins/Animation/GameplayInsights/Source/GameplayInsights/Private/AnimNotifiesTrack.cpp b/Engine/Plugins/Animation/GameplayInsights/Source/GameplayInsights/Private/AnimNotifiesTrack.cpp
--- a/Engine/Plugins/Animation/GameplayInsights/Source/GameplayInsights/Private/AnimNotifiesTrack.cpp
+++ b/Engine/Plugins/Animation/GameplayInsights/Source/GameplayInsights/Private/AnimNotifiesTrack.cpp
@@ -30,7 +30,7 @@ void FAnimNotifiesTrack::BuildDrawState(ITimingEventsTrackDrawStateBuilder& Buil
 
 	if(AnimationProvider)
 	{
-		Trace::FAnalysisSessionReadScope SessionReadScope(SharedData.GetAnalysisSession());
+		TraceServices::FAnalysisSessionReadScope SessionReadScope(SharedData.GetAnalysisSession());
 
 		bHasNotifies = false;
 
@@ -40,7 +40,7 @@ void FAnimNotifiesTrack::BuildDrawState(ITimingEventsTrackDrawStateBuilder& Buil
 			{
 				Builder.AddEvent(InStartTime, InEndTime, 0, InMessage.Name);
 				bHasNotifies = true;
-				return Trace::EEventEnumerate::Continue;
+				return TraceServices::EEventEnumerate::Continue;
 			});
 		});
 
@@ -63,7 +63,7 @@ void FAnimNotifiesTrack::BuildDrawState(ITimingEventsTrackDrawStateBuilder& Buil
 			InTimeline.EnumerateEvents(Context.GetViewport().GetStartTime(), Context.GetViewport().GetEndTime(), [&Builder, &DepthToUse](double InStartTime, double InEndTime, uint32 InDepth, const FAnimNotifyMessage& InMessage)
 			{
 				Builder.AddEvent(InStartTime, InEndTime, DepthToUse, InMessage.Name);
-				return Trace::EEventEnumerate::Continue;
+				return TraceServices::EEventEnumerate::Continue;
 			});
 		});
 	}
@@ -138,7 +138,7 @@ void FAnimNotifiesTrack::InitTooltip(FTooltipDrawState& Tooltip, const ITimingEv
 		Tooltip.AddNameValueTextLine(LOCTEXT("NotifyType", "Type").ToString(), GetMessageType(InMessage.NotifyEventType, InFoundDepth == 0));
 
 		{
-			Trace::FAnalysisSessionReadScope SessionReadScope(SharedData.GetAnalysisSession());
+			TraceServices::FAnalysisSessionReadScope SessionReadScope(SharedData.GetAnalysisSession());
 
 			const FAnimationProvider* AnimationProvider = SharedData.GetAnalysisSession().ReadProvider<FAnimationProvider>(FAnimationProvider::ProviderName);
 			if(AnimationProvider)
@@ -153,7 +153,7 @@ void FAnimNotifiesTrack::InitTooltip(FTooltipDrawState& Tooltip, const ITimingEv
 
 		if(InMessage.NotifyEventType != EAnimNotifyMessageType::SyncMarker)
 		{
-			Trace::FAnalysisSessionReadScope SessionReadScope(SharedData.GetAnalysisSession());
+			TraceServices::FAnalysisSessionReadScope SessionReadScope(SharedData.GetAnalysisSession());
 
 			const FGameplayProvider* GameplayProvider = SharedData.GetAnalysisSession().ReadProvider<FGameplayProvider>(FGameplayProvider::ProviderName);
 			if(GameplayProvider)
@@ -197,7 +197,7 @@ void FAnimNotifiesTrack::FindAnimNotifyMessage(const FTimingEventSearchParameter
 
 			if(AnimationProvider)
 			{
-				Trace::FAnalysisSessionReadScope SessionReadScope(SharedData.GetAnalysisSession());
+				TraceServices::FAnalysisSessionReadScope SessionReadScope(SharedData.GetAnalysisSession());
 
 				if(bHasNotifies)
 				{
@@ -206,7 +206,7 @@ void FAnimNotifiesTrack::FindAnimNotifyMessage(const FTimingEventSearchParameter
 						InTimeline.EnumerateEvents(InContext.GetParameters().StartTime, InContext.GetParameters().EndTime, [this, &InContext](double InEventStartTime, double InEventEndTime, uint32 InDepth, const FAnimNotifyMessage& InMessage)
 						{
 							InContext.Check(InEventStartTime, InEventEndTime, 0, InMessage);
-							return Trace::EEventEnumerate::Continue;
+							return TraceServices::EEventEnumerate::Continue;
 						});
 					});
 				}
@@ -221,7 +221,7 @@ void FAnimNotifiesTrack::FindAnimNotifyMessage(const FTimingEventSearchParameter
 							{
 								InContext.Check(InEventStartTime, InEventEndTime, *FoundDepthPtr, InMessage);
 							}
-							return Trace::EEventEnumerate::Continue;
+							return TraceServices::EEventEnumerate::Continue;
 						});
 					});
 				}
@@ -234,17 +234,17 @@ void FAnimNotifiesTrack::FindAnimNotifyMessage(const FTimingEventSearchParameter
 		});
 }
 
-void FAnimNotifiesTrack::GetVariantsAtFrame(const Trace::FFrame& InFrame, TArray<TSharedRef<FVariantTreeNode>>& OutVariants) const 
+void FAnimNotifiesTrack::GetVariantsAtFrame(const TraceServices::FFrame& InFrame, TArray<TSharedRef<FVariantTreeNode>>& OutVariants) const 
 {
 	TSharedRef<FVariantTreeNode> Header = OutVariants.Add_GetRef(FVariantTreeNode::MakeHeader(LOCTEXT("AnimNotifiesHeader", "Notifies and Sync Markers"), 0));
 
-	const Trace::IFrameProvider& FramesProvider = Trace::ReadFrameProvider(SharedData.GetAnalysisSession());
+	const TraceServices::IFrameProvider& FramesProvider = TraceServices::ReadFrameProvider(SharedData.GetAnalysisSession());
 	const FAnimationProvider* AnimationProvider = SharedData.GetAnalysisSession().ReadProvider<FAnimationProvider>(FAnimationProvider::ProviderName);
 	const FGameplayProvider* GameplayProvider = SharedData.GetAnalysisSession().ReadProvider<FGameplayProvider>(FGameplayProvider::ProviderName);
 
 	if(AnimationProvider && GameplayProvider)
 	{
-		Trace::FAnalysisSessionReadScope SessionReadScope(SharedData.GetAnalysisSession());
+		TraceServices::FAnalysisSessionReadScope SessionReadScope(SharedData.GetAnalysisSession());
 
 		bool bForEvent = true;
 		auto ProcessEvent = [this, &AnimationProvider, &GameplayProvider, &Header, &bForEvent](double InStartTime, double InEndTime, uint32 InDepth, const FAnimNotifyMessage& InMessage)
@@ -265,7 +265,7 @@ void FAnimNotifiesTrack::GetVariantsAtFrame(const Trace::FFrame& InFrame, TArray
 				NotifyHeader->AddChild(FVariantTreeNode::MakeClass((bIsState ? LOCTEXT("Notify State Class", "Notify State Class") : LOCTEXT("NotifyClass", "Notify Class")), NotifyInfo.ClassId));
 			}
 
-			return Trace::EEventEnumerate::Continue;
+			return TraceServices::EEventEnumerate::Continue;
 		};
 
 		AnimationProvider->ReadNotifyTimeline(GetGameplayTrack().GetObjectId(), [&InFrame, &ProcessEvent](const FAnimationProvider::AnimNotifyTimeline& InTimeline)
diff --git a/Engine/Plugins/Animation/GameplayInsights/Source/GameplayInsights/Private/AnimNotifiesTrack.h b/Engine/Plugins/Animation/GameplayInsights/Source/GameplayInsights/Private/AnimNotifiesTrack.h
--- a/Engine/Plugins/Animation/GameplayInsights/Source/GameplayInsights/Private/AnimNotifiesTrack.h
+++ b/Engine/Plugins/Animation/GameplayInsights/Source/GameplayInsights/Private/AnimNotifiesTrack.h
@@ -20,7 +20,7 @@ public:
 	virtual void Draw(const ITimingTrackDrawContext& Context) const override;
 	virtual void InitTooltip(FTooltipDrawState& Tooltip, const ITimingEvent& HoveredTimingEvent) const override;
 	virtual const TSharedPtr<const ITimingEvent> SearchEvent(const FTimingEventSearchParameters& InSearchParameters) const override;
-	virtual void GetVariantsAtFrame(const Trace::FFrame& InFrame, TArray<TSharedRef<FVariantTreeNode>>& OutVariants) const override;
+	virtual void GetVariantsAtFrame(const TraceServices::FFrame& InFrame, TArray<TSharedRef<FVariantTreeNode>>& OutVariants) const override;
 
 private:
 	// Helper function used to find an anim notify message
@@ -35,4 +35,4 @@ private:
 
 	// Whether we are currently drawing any notifies
 	bool bHasNotifies;
-};
\ No newline at end of file
+};
diff --git a/Engine/Plugins/Animation/GameplayInsights/Source/GameplayInsights/Private/AnimationAnalyzer.cpp b/Engine/Plugins/Animation/GameplayInsights/Source/GameplayInsights/Private/AnimationAnalyzer.cpp
--- a/Engine/Plugins/Animation/GameplayInsights/Source/GameplayInsights/Private/AnimationAnalyzer.cpp
+++ b/Engine/Plugins/Animation/GameplayInsights/Source/GameplayInsights/Private/AnimationAnalyzer.cpp
@@ -6,7 +6,7 @@
 #include "AnimationProvider.h"
 #include "Containers/ArrayView.h"
 
-FAnimationAnalyzer::FAnimationAnalyzer(Trace::IAnalysisSession& InSession, FAnimationProvider& InAnimationProvider)
+FAnimationAnalyzer::FAnimationAnalyzer(TraceServices::IAnalysisSession& InSession, FAnimationProvider& InAnimationProvider)
 	: Session(InSession)
 	, AnimationProvider(InAnimationProvider)
 {
@@ -43,7 +43,7 @@ void FAnimationAnalyzer::OnAnalysisBegin(const FOnAnalysisContext& Context)
 
 bool FAnimationAnalyzer::OnEvent(uint16 RouteId, EStyle Style, const FOnEventContext& Context)
 {
-	Trace::FAnalysisSessionEditScope _(Session);
+	TraceServices::FAnalysisSessionEditScope _(Session);
 
 	const auto& EventData = Context.EventData;
 	switch (RouteId)
@@ -119,7 +119,7 @@ bool FAnimationAnalyzer::OnEvent(uint16 RouteId, EStyle Style, const FOnEventCon
 	case RouteId_Name:
 	{
 		uint32 Id = EventData.GetValue<uint32>("Id");
-		FString Name = FTraceAnalyzerUtils::LegacyAttachmentString<TCHAR>("Name", Context);
+		FString Name = TraceServices::FTraceAnalyzerUtils::LegacyAttachmentString<TCHAR>("Name", Context);
 		AnimationProvider.AppendName(Id, *Name);
 		break;
 	}
@@ -296,4 +296,4 @@ bool FAnimationAnalyzer::OnEvent(uint16 RouteId, EStyle Style, const FOnEventCon
 	}
 
 	return true;
-}
\ No newline at end of file
+}
diff --git a/Engine/Plugins/Animation/GameplayInsights/Source/GameplayInsights/Private/AnimationAnalyzer.h b/Engine/Plugins/Animation/GameplayInsights/Source/GameplayInsights/Private/AnimationAnalyzer.h
--- a/Engine/Plugins/Animation/GameplayInsights/Source/GameplayInsights/Private/AnimationAnalyzer.h
+++ b/Engine/Plugins/Animation/GameplayInsights/Source/GameplayInsights/Private/AnimationAnalyzer.h
@@ -5,12 +5,12 @@
 #include "Trace/Analyzer.h"
 
 class FAnimationProvider;
-namespace Trace { class IAnalysisSession; }
+namespace TraceServices { class IAnalysisSession; }
 
 class FAnimationAnalyzer : public Trace::IAnalyzer
 {
 public:
-	FAnimationAnalyzer(Trace::IAnalysisSession& InSession, FAnimationProvider& InAnimationProvider);
+	FAnimationAnalyzer(TraceServices::IAnalysisSession& InSession, FAnimationProvider& InAnimationProvider);
 
 	virtual void OnAnalysisBegin(const FOnAnalysisContext& Context) override;
 	virtual void OnAnalysisEnd() override {}
@@ -44,6 +44,6 @@ private:
 		RouteId_Montage,
 	};
 
-	Trace::IAnalysisSession& Session;
+	TraceServices::IAnalysisSession& Session;
 	FAnimationProvider& AnimationProvider;
-};
\ No newline at end of file
+};
diff --git a/Engine/Plugins/Animation/GameplayInsights/Source/GameplayInsights/Private/AnimationProvider.cpp b/Engine/Plugins/Animation/GameplayInsights/Source/GameplayInsights/Private/AnimationProvider.cpp
--- a/Engine/Plugins/Animation/GameplayInsights/Source/GameplayInsights/Private/AnimationProvider.cpp
+++ b/Engine/Plugins/Animation/GameplayInsights/Source/GameplayInsights/Private/AnimationProvider.cpp
@@ -8,7 +8,7 @@
 
 FName FAnimationProvider::ProviderName("AnimationProvider");
 
-FAnimationProvider::FAnimationProvider(Trace::IAnalysisSession& InSession, FGameplayProvider& InGameplayProvider)
+FAnimationProvider::FAnimationProvider(TraceServices::IAnalysisSession& InSession, FGameplayProvider& InGameplayProvider)
 	: Session(InSession)
 	, GameplayProvider(InGameplayProvider)
 	, SkeletalMeshPoseTransforms(InSession.GetLinearAllocator(), 256)
@@ -392,7 +392,7 @@ void FAnimationProvider::AppendTickRecord(uint64 InAnimInstanceId, double InTime
 	{
 		ObjectIdToTickRecordTimelineStorage.Add(InAnimInstanceId, TickRecordTimelineStorage.Num());
 		TimelineStorage = TickRecordTimelineStorage.Add_GetRef(MakeShared<FTickRecordTimelineStorage>());
-		TimelineStorage->Timeline = MakeShared<Trace::TPointTimeline<FTickRecordMessage>>(Session.GetLinearAllocator());
+		TimelineStorage->Timeline = MakeShared<TraceServices::TPointTimeline<FTickRecordMessage>>(Session.GetLinearAllocator());
 		TimelineStorage->Timeline->SetEnumerateOutsideRange(true);
 	}
 
@@ -457,7 +457,7 @@ void FAnimationProvider::AppendSkeletalMeshComponent(uint64 InObjectId, uint64 I
 	else
 	{
 		TimelineStorage = MakeShared<FSkeletalMeshTimelineStorage>();
-		TimelineStorage->Timeline = MakeShared<Trace::TIntervalTimeline<FSkeletalMeshPoseMessage>>(Session.GetLinearAllocator());
+		TimelineStorage->Timeline = MakeShared<TraceServices::TIntervalTimeline<FSkeletalMeshPoseMessage>>(Session.GetLinearAllocator());
 		ObjectIdToSkeletalMeshPoseTimelines.Add(InObjectId, SkeletalMeshPoseTimelineStorage.Num());
 		SkeletalMeshPoseTimelineStorage.Add(TimelineStorage.ToSharedRef());
 	}
@@ -513,7 +513,7 @@ void FAnimationProvider::AppendSkeletalMeshComponent(uint64 InObjectId, uint64 I
 	else
 	{
 		TimelineStorage = MakeShared<FSkeletalMeshTimelineStorage>();
-		TimelineStorage->Timeline = MakeShared<Trace::TIntervalTimeline<FSkeletalMeshPoseMessage>>(Session.GetLinearAllocator());
+		TimelineStorage->Timeline = MakeShared<TraceServices::TIntervalTimeline<FSkeletalMeshPoseMessage>>(Session.GetLinearAllocator());
 		ObjectIdToSkeletalMeshPoseTimelines.Add(InObjectId, SkeletalMeshPoseTimelineStorage.Num());
 		SkeletalMeshPoseTimelineStorage.Add(TimelineStorage.ToSharedRef());
 	}
@@ -571,7 +571,7 @@ void FAnimationProvider::HandleObjectEndPlay(uint64 InObjectId, double InTime, c
 	uint32* SkelMeshPoseIndexPtr = ObjectIdToSkeletalMeshPoseTimelines.Find(InObjectId);
 	if(SkelMeshPoseIndexPtr != nullptr)
 	{
-		TSharedPtr<Trace::TIntervalTimeline<FSkeletalMeshPoseMessage>> Timeline = SkeletalMeshPoseTimelineStorage[*SkelMeshPoseIndexPtr]->Timeline;
+		TSharedPtr<TraceServices::TIntervalTimeline<FSkeletalMeshPoseMessage>> Timeline = SkeletalMeshPoseTimelineStorage[*SkelMeshPoseIndexPtr]->Timeline;
 		uint64 NumEvents = Timeline->GetEventCount();
 		if(NumEvents > 0)
 		{
@@ -586,7 +586,7 @@ void FAnimationProvider::AppendSkeletalMeshFrame(uint64 InObjectId, double InTim
 
 	bHasAnyData = true;
 
-	TSharedPtr<Trace::TIntervalTimeline<FSkeletalMeshFrameMessage>> Timeline;
+	TSharedPtr<TraceServices::TIntervalTimeline<FSkeletalMeshFrameMessage>> Timeline;
 	uint32* IndexPtr = ObjectIdToSkeletalMeshFrameTimelines.Find(InObjectId);
 	if(IndexPtr != nullptr)
 	{
@@ -595,7 +595,7 @@ void FAnimationProvider::AppendSkeletalMeshFrame(uint64 InObjectId, double InTim
 	}
 	else
 	{
-		Timeline = MakeShared<Trace::TIntervalTimeline<FSkeletalMeshFrameMessage>>(Session.GetLinearAllocator());
+		Timeline = MakeShared<TraceServices::TIntervalTimeline<FSkeletalMeshFrameMessage>>(Session.GetLinearAllocator());
 		ObjectIdToSkeletalMeshFrameTimelines.Add(InObjectId, SkeletalMeshFrameTimelines.Num());
 		SkeletalMeshFrameTimelines.Add(Timeline.ToSharedRef());
 	}
@@ -623,7 +623,7 @@ void FAnimationProvider::AppendAnimGraph(uint64 InAnimInstanceId, double InStart
 
 	bHasAnyData = true;
 
-	TSharedPtr<Trace::TIntervalTimeline<FAnimGraphMessage>> Timeline;
+	TSharedPtr<TraceServices::TIntervalTimeline<FAnimGraphMessage>> Timeline;
 	uint32* IndexPtr = ObjectIdToAnimGraphTimelines.Find(InAnimInstanceId);
 	if(IndexPtr != nullptr)
 	{
@@ -632,7 +632,7 @@ void FAnimationProvider::AppendAnimGraph(uint64 InAnimInstanceId, double InStart
 	}
 	else
 	{
-		Timeline = MakeShared<Trace::TIntervalTimeline<FAnimGraphMessage>>(Session.GetLinearAllocator());
+		Timeline = MakeShared<TraceServices::TIntervalTimeline<FAnimGraphMessage>>(Session.GetLinearAllocator());
 		ObjectIdToAnimGraphTimelines.Add(InAnimInstanceId, AnimGraphTimelines.Num());
 		AnimGraphTimelines.Add(Timeline.ToSharedRef());
 	}
@@ -655,7 +655,7 @@ void FAnimationProvider::AppendAnimNodeStart(uint64 InAnimInstanceId, double InS
 
 	bHasAnyData = true;
 
-	TSharedPtr<Trace::TPointTimeline<FAnimNodeMessage>> Timeline;
+	TSharedPtr<TraceServices::TPointTimeline<FAnimNodeMessage>> Timeline;
 	uint32* IndexPtr = ObjectIdToAnimNodeTimelines.Find(InAnimInstanceId);
 	if(IndexPtr != nullptr)
 	{
@@ -664,7 +664,7 @@ void FAnimationProvider::AppendAnimNodeStart(uint64 InAnimInstanceId, double InS
 	}
 	else
 	{
-		Timeline = MakeShared<Trace::TPointTimeline<FAnimNodeMessage>>(Session.GetLinearAllocator());
+		Timeline = MakeShared<TraceServices::TPointTimeline<FAnimNodeMessage>>(Session.GetLinearAllocator());
 		ObjectIdToAnimNodeTimelines.Add(InAnimInstanceId, AnimNodeTimelines.Num());
 		AnimNodeTimelines.Add(Timeline.ToSharedRef());
 	}
@@ -690,7 +690,7 @@ void FAnimationProvider::AppendAnimSequencePlayer(uint64 InAnimInstanceId, doubl
 
 	bHasAnyData = true;
 
-	TSharedPtr<Trace::TPointTimeline<FAnimSequencePlayerMessage>> Timeline;
+	TSharedPtr<TraceServices::TPointTimeline<FAnimSequencePlayerMessage>> Timeline;
 	uint32* IndexPtr = ObjectIdToAnimSequencePlayerTimelines.Find(InAnimInstanceId);
 	if(IndexPtr != nullptr)
 	{
@@ -699,7 +699,7 @@ void FAnimationProvider::AppendAnimSequencePlayer(uint64 InAnimInstanceId, doubl
 	}
 	else
 	{
-		Timeline = MakeShared<Trace::TPointTimeline<FAnimSequencePlayerMessage>>(Session.GetLinearAllocator());
+		Timeline = MakeShared<TraceServices::TPointTimeline<FAnimSequencePlayerMessage>>(Session.GetLinearAllocator());
 		ObjectIdToAnimSequencePlayerTimelines.Add(InAnimInstanceId, AnimSequencePlayerTimelines.Num());
 		AnimSequencePlayerTimelines.Add(Timeline.ToSharedRef());
 	}
@@ -722,7 +722,7 @@ void FAnimationProvider::AppendBlendSpacePlayer(uint64 InAnimInstanceId, double
 
 	bHasAnyData = true;
 
-	TSharedPtr<Trace::TPointTimeline<FBlendSpacePlayerMessage>> Timeline;
+	TSharedPtr<TraceServices::TPointTimeline<FBlendSpacePlayerMessage>> Timeline;
 	uint32* IndexPtr = ObjectIdToBlendSpacePlayerTimelines.Find(InAnimInstanceId);
 	if(IndexPtr != nullptr)
 	{
@@ -731,7 +731,7 @@ void FAnimationProvider::AppendBlendSpacePlayer(uint64 InAnimInstanceId, double
 	}
 	else
 	{
-		Timeline = MakeShared<Trace::TPointTimeline<FBlendSpacePlayerMessage>>(Session.GetLinearAllocator());
+		Timeline = MakeShared<TraceServices::TPointTimeline<FBlendSpacePlayerMessage>>(Session.GetLinearAllocator());
 		ObjectIdToBlendSpacePlayerTimelines.Add(InAnimInstanceId, BlendSpacePlayerTimelines.Num());
 		BlendSpacePlayerTimelines.Add(Timeline.ToSharedRef());
 	}
@@ -827,7 +827,7 @@ void FAnimationProvider::AppendAnimNodeValue(uint64 InAnimInstanceId, double InT
 
 	bHasAnyData = true;
 
-	TSharedPtr<Trace::TPointTimeline<FAnimNodeValueMessage>> Timeline;
+	TSharedPtr<TraceServices::TPointTimeline<FAnimNodeValueMessage>> Timeline;
 	uint32* IndexPtr = ObjectIdToAnimNodeValueTimelines.Find(InAnimInstanceId);
 	if(IndexPtr != nullptr)
 	{
@@ -835,7 +835,7 @@ void FAnimationProvider::AppendAnimNodeValue(uint64 InAnimInstanceId, double InT
 	}
 	else
 	{
-		Timeline = MakeShared<Trace::TPointTimeline<FAnimNodeValueMessage>>(Session.GetLinearAllocator());
+		Timeline = MakeShared<TraceServices::TPointTimeline<FAnimNodeValueMessage>>(Session.GetLinearAllocator());
 		ObjectIdToAnimNodeValueTimelines.Add(InAnimInstanceId, AnimNodeValueTimelines.Num());
 		AnimNodeValueTimelines.Add(Timeline.ToSharedRef());
 	}
@@ -856,7 +856,7 @@ void FAnimationProvider::AppendStateMachineState(uint64 InAnimInstanceId, double
 
 	bHasAnyData = true;
 
-	TSharedPtr<Trace::TPointTimeline<FAnimStateMachineMessage>> Timeline;
+	TSharedPtr<TraceServices::TPointTimeline<FAnimStateMachineMessage>> Timeline;
 	uint32* IndexPtr = ObjectIdToStateMachineTimelines.Find(InAnimInstanceId);
 	if(IndexPtr != nullptr)
 	{
@@ -864,7 +864,7 @@ void FAnimationProvider::AppendStateMachineState(uint64 InAnimInstanceId, double
 	}
 	else
 	{
-		Timeline = MakeShared<Trace::TPointTimeline<FAnimStateMachineMessage>>(Session.GetLinearAllocator());
+		Timeline = MakeShared<TraceServices::TPointTimeline<FAnimStateMachineMessage>>(Session.GetLinearAllocator());
 		ObjectIdToStateMachineTimelines.Add(InAnimInstanceId, StateMachineTimelines.Num());
 		StateMachineTimelines.Add(Timeline.ToSharedRef());
 	}
@@ -891,7 +891,7 @@ void FAnimationProvider::AppendNotify(uint64 InAnimInstanceId, double InTime, ui
 	// Check if stateful or event-based
 	if(InNotifyEventType == EAnimNotifyMessageType::Begin || InNotifyEventType == EAnimNotifyMessageType::End)
 	{
-		TSharedPtr<Trace::TIntervalTimeline<FAnimNotifyMessage>> Timeline;
+		TSharedPtr<TraceServices::TIntervalTimeline<FAnimNotifyMessage>> Timeline;
 		TSharedPtr<FAnimNotifyStateTimelineStorage> TimelineStorage;
 		uint32* TimelineStorageIndexPtr = ObjectIdToAnimNotifyStateTimelines.Find(InAnimInstanceId);
 		if(TimelineStorageIndexPtr != nullptr)
@@ -913,7 +913,7 @@ void FAnimationProvider::AppendNotify(uint64 InAnimInstanceId, double InTime, ui
 		}
 		else
 		{
-			Timeline = MakeShared<Trace::TIntervalTimeline<FAnimNotifyMessage>>(Session.GetLinearAllocator());
+			Timeline = MakeShared<TraceServices::TIntervalTimeline<FAnimNotifyMessage>>(Session.GetLinearAllocator());
 			TimelineStorage->NotifyIdToAnimNotifyStateTimeline.Add(InNotifyId, TimelineStorage->Timelines.Num());
 			TimelineStorage->Timelines.Add(Timeline.ToSharedRef());
 		}
@@ -939,7 +939,7 @@ void FAnimationProvider::AppendNotify(uint64 InAnimInstanceId, double InTime, ui
 	}
 	else
 	{
-		TSharedPtr<Trace::TPointTimeline<FAnimNotifyMessage>> Timeline;
+		TSharedPtr<TraceServices::TPointTimeline<FAnimNotifyMessage>> Timeline;
 		uint32* IndexPtr = ObjectIdToAnimNotifyTimelines.Find(InAnimInstanceId);
 		if(IndexPtr != nullptr)
 		{
@@ -947,7 +947,7 @@ void FAnimationProvider::AppendNotify(uint64 InAnimInstanceId, double InTime, ui
 		}
 		else
 		{
-			Timeline = MakeShared<Trace::TPointTimeline<FAnimNotifyMessage>>(Session.GetLinearAllocator());
+			Timeline = MakeShared<TraceServices::TPointTimeline<FAnimNotifyMessage>>(Session.GetLinearAllocator());
 			ObjectIdToAnimNotifyTimelines.Add(InAnimInstanceId, AnimNotifyTimelines.Num());
 			AnimNotifyTimelines.Add(Timeline.ToSharedRef());
 		}
@@ -983,7 +983,7 @@ void FAnimationProvider::AppendMontage(uint64 InAnimInstanceId, double InTime, u
 	else
 	{
 		TimelineStorage = MakeShared<FMontageTimelineStorage>();
-		TimelineStorage->Timeline = MakeShared<Trace::TPointTimeline<FAnimMontageMessage>>(Session.GetLinearAllocator());
+		TimelineStorage->Timeline = MakeShared<TraceServices::TPointTimeline<FAnimMontageMessage>>(Session.GetLinearAllocator());
 		TimelineStorage->Timeline->SetEnumerateOutsideRange(true);
 		ObjectIdToAnimMontageTimelines.Add(InAnimInstanceId, AnimMontageTimelineStorage.Num());
 		AnimMontageTimelineStorage.Add(TimelineStorage.ToSharedRef());
@@ -1005,4 +1005,4 @@ void FAnimationProvider::AppendMontage(uint64 InAnimInstanceId, double InTime, u
 	Session.UpdateDurationSeconds(InTime);
 }
 
-#undef LOCTEXT_NAMESPACE
\ No newline at end of file
+#undef LOCTEXT_NAMESPACE
diff --git a/Engine/Plugins/Animation/GameplayInsights/Source/GameplayInsights/Private/AnimationProvider.h b/Engine/Plugins/Animation/GameplayInsights/Source/GameplayInsights/Private/AnimationProvider.h
--- a/Engine/Plugins/Animation/GameplayInsights/Source/GameplayInsights/Private/AnimationProvider.h
+++ b/Engine/Plugins/Animation/GameplayInsights/Source/GameplayInsights/Private/AnimationProvider.h
@@ -9,7 +9,7 @@
 #include "Containers/ArrayView.h"
 #include "Model/IntervalTimeline.h"
 
-namespace Trace { class IAnalysisSession; }
+namespace TraceServices { class IAnalysisSession; }
 class FGameplayProvider;
 struct FObjectInfo;
 class USkeletalMesh;
@@ -19,7 +19,7 @@ class FAnimationProvider : public IAnimationProvider
 public:
 	static FName ProviderName;
 
-	FAnimationProvider(Trace::IAnalysisSession& InSession, FGameplayProvider& InGameplayProvider);
+	FAnimationProvider(TraceServices::IAnalysisSession& InSession, FGameplayProvider& InGameplayProvider);
 
 	/** IAnimationProvider interface */
 	virtual bool ReadSkeletalMeshPoseTimeline(uint64 InObjectId, TFunctionRef<void(const SkeletalMeshPoseTimeline&, bool)> Callback) const override;
@@ -104,7 +104,7 @@ private:
 	void HandleObjectEndPlay(uint64 InObjectId, double InTime, const FObjectInfo& InObjectInfo);
 
 private:
-	Trace::IAnalysisSession& Session;
+	TraceServices::IAnalysisSession& Session;
 	FGameplayProvider& GameplayProvider;
 
 	/** Maps into timeline arrays per-object */
@@ -133,44 +133,44 @@ private:
 	struct FTickRecordTimelineStorage
 	{
 		TSet<TTuple<uint64, int32>> AssetIdAndPlayers;
-		TSharedPtr<Trace::TPointTimeline<FTickRecordMessage>> Timeline;
+		TSharedPtr<TraceServices::TPointTimeline<FTickRecordMessage>> Timeline;
 	};
 
 	struct FSkeletalMeshTimelineStorage
 	{
-		TSharedPtr<Trace::TIntervalTimeline<FSkeletalMeshPoseMessage>> Timeline;
+		TSharedPtr<TraceServices::TIntervalTimeline<FSkeletalMeshPoseMessage>> Timeline;
 		TSet<uint32> AllCurveIds;
 	};
 
 	struct FAnimNotifyStateTimelineStorage
 	{
 		TMap<uint64, uint32> NotifyIdToAnimNotifyStateTimeline;
-		TArray<TSharedRef<Trace::TIntervalTimeline<FAnimNotifyMessage>>> Timelines;
+		TArray<TSharedRef<TraceServices::TIntervalTimeline<FAnimNotifyMessage>>> Timelines;
 	};
 
 	struct FMontageTimelineStorage
 	{
-		TSharedPtr<Trace::TPointTimeline<FAnimMontageMessage>> Timeline;
+		TSharedPtr<TraceServices::TPointTimeline<FAnimMontageMessage>> Timeline;
 		TSet<uint64> AllMontageIds;
 	};
 
 	/** Message storage */
 	TArray<TSharedRef<FTickRecordTimelineStorage>> TickRecordTimelineStorage;
 	TArray<TSharedRef<FSkeletalMeshTimelineStorage>> SkeletalMeshPoseTimelineStorage;
-	TArray<TSharedRef<Trace::TIntervalTimeline<FSkeletalMeshFrameMessage>>> SkeletalMeshFrameTimelines;
-	TArray<TSharedRef<Trace::TIntervalTimeline<FAnimGraphMessage>>> AnimGraphTimelines;
-	TArray<TSharedRef<Trace::TPointTimeline<FAnimNodeMessage>>> AnimNodeTimelines;
-	TArray<TSharedRef<Trace::TPointTimeline<FAnimNodeValueMessage>>> AnimNodeValueTimelines;
-	TArray<TSharedRef<Trace::TPointTimeline<FAnimSequencePlayerMessage>>> AnimSequencePlayerTimelines;
-	TArray<TSharedRef<Trace::TPointTimeline<FBlendSpacePlayerMessage>>> BlendSpacePlayerTimelines;
-	TArray<TSharedRef<Trace::TPointTimeline<FAnimStateMachineMessage>>> StateMachineTimelines;
+	TArray<TSharedRef<TraceServices::TIntervalTimeline<FSkeletalMeshFrameMessage>>> SkeletalMeshFrameTimelines;
+	TArray<TSharedRef<TraceServices::TIntervalTimeline<FAnimGraphMessage>>> AnimGraphTimelines;
+	TArray<TSharedRef<TraceServices::TPointTimeline<FAnimNodeMessage>>> AnimNodeTimelines;
+	TArray<TSharedRef<TraceServices::TPointTimeline<FAnimNodeValueMessage>>> AnimNodeValueTimelines;
+	TArray<TSharedRef<TraceServices::TPointTimeline<FAnimSequencePlayerMessage>>> AnimSequencePlayerTimelines;
+	TArray<TSharedRef<TraceServices::TPointTimeline<FBlendSpacePlayerMessage>>> BlendSpacePlayerTimelines;
+	TArray<TSharedRef<TraceServices::TPointTimeline<FAnimStateMachineMessage>>> StateMachineTimelines;
 	TArray<TSharedRef<FAnimNotifyStateTimelineStorage>> AnimNotifyStateTimelineStorage;
-	TArray<TSharedRef<Trace::TPointTimeline<FAnimNotifyMessage>>> AnimNotifyTimelines;
+	TArray<TSharedRef<TraceServices::TPointTimeline<FAnimNotifyMessage>>> AnimNotifyTimelines;
 	TArray<TSharedRef<FMontageTimelineStorage>> AnimMontageTimelineStorage;
-	TPagedArray<FTransform> SkeletalMeshPoseTransforms;
-	TPagedArray<FSkeletalMeshNamedCurve> SkeletalMeshCurves;
-	TPagedArray<int32> SkeletalMeshParentIndices;
+	TraceServices::TPagedArray<FTransform> SkeletalMeshPoseTransforms;
+	TraceServices::TPagedArray<FSkeletalMeshNamedCurve> SkeletalMeshCurves;
+	TraceServices::TPagedArray<int32> SkeletalMeshParentIndices;
 
 	// Flag to indicate if any data is present
 	bool bHasAnyData;
-};
\ No newline at end of file
+};
diff --git a/Engine/Plugins/Animation/GameplayInsights/Source/GameplayInsights/Private/AnimationSharedData.cpp b/Engine/Plugins/Animation/GameplayInsights/Source/GameplayInsights/Private/AnimationSharedData.cpp
--- a/Engine/Plugins/Animation/GameplayInsights/Source/GameplayInsights/Private/AnimationSharedData.cpp
+++ b/Engine/Plugins/Animation/GameplayInsights/Source/GameplayInsights/Private/AnimationSharedData.cpp
@@ -68,7 +68,7 @@ void FAnimationSharedData::OnEndSession(Insights::ITimingViewSession& InTimingVi
 	TimingViewSession = nullptr;
 }
 
-void FAnimationSharedData::Tick(Insights::ITimingViewSession& InTimingViewSession, const Trace::IAnalysisSession& InAnalysisSession)
+void FAnimationSharedData::Tick(Insights::ITimingViewSession& InTimingViewSession, const TraceServices::IAnalysisSession& InAnalysisSession)
 {
 	AnalysisSession = &InAnalysisSession;
 
@@ -77,7 +77,7 @@ void FAnimationSharedData::Tick(Insights::ITimingViewSession& InTimingViewSessio
 
 	if(AnimationProvider && GameplayProvider)
 	{
-		Trace::FAnalysisSessionReadScope SessionReadScope(GetAnalysisSession());
+		TraceServices::FAnalysisSessionReadScope SessionReadScope(GetAnalysisSession());
 
 		if(AnimationProvider->HasAnyData() && AreAnyAnimationTracksEnabled())
 		{
@@ -508,14 +508,14 @@ void FAnimationSharedData::OnTimeMarkerChanged(Insights::ETimeChangedFlags InFla
 
 	if(bTimeMarkerValid && AnalysisSession != nullptr)
 	{
-		Trace::FAnalysisSessionReadScope SessionReadScope(*AnalysisSession);
+		TraceServices::FAnalysisSessionReadScope SessionReadScope(*AnalysisSession);
 
-		const Trace::IFrameProvider& FramesProvider = Trace::ReadFrameProvider(*AnalysisSession);
+		const TraceServices::IFrameProvider& FramesProvider = TraceServices::ReadFrameProvider(*AnalysisSession);
 		bMarkerFrameValid = FramesProvider.GetFrameFromTime(ETraceFrameType::TraceFrameType_Game, MarkerTime, MarkerFrame);
 	}
 	else
 	{
-		MarkerFrame = Trace::FFrame();
+		MarkerFrame = TraceServices::FFrame();
 		bMarkerFrameValid = false;
 	}
 
@@ -676,7 +676,7 @@ void FAnimationSharedData::OpenAnimGraphTab(uint64 InAnimInstanceId) const
 		const FGameplayProvider* GameplayProvider = AnalysisSession->ReadProvider<FGameplayProvider>(FGameplayProvider::ProviderName);
 		if(GameplayProvider)
 		{
-			Trace::FAnalysisSessionReadScope SessionReadScope(*AnalysisSession);
+			TraceServices::FAnalysisSessionReadScope SessionReadScope(*AnalysisSession);
 
 			const FObjectInfo& ObjectInfo = GameplayProvider->GetObjectInfo(InAnimInstanceId);
 			Tab->SetLabel(FText::FromString(ObjectInfo.Name));
diff --git a/Engine/Plugins/Animation/GameplayInsights/Source/GameplayInsights/Private/AnimationSharedData.h b/Engine/Plugins/Animation/GameplayInsights/Source/GameplayInsights/Private/AnimationSharedData.h
--- a/Engine/Plugins/Animation/GameplayInsights/Source/GameplayInsights/Private/AnimationSharedData.h
+++ b/Engine/Plugins/Animation/GameplayInsights/Source/GameplayInsights/Private/AnimationSharedData.h
@@ -6,7 +6,7 @@
 #include "TraceServices/Model/Frames.h"
 
 class FGameplaySharedData;
-namespace Trace { class IAnalysisSession; }
+namespace TraceServices { class IAnalysisSession; }
 namespace Insights { class ITimingViewSession; }
 namespace Insights { enum class ETimeChangedFlags : int32; }
 class FSkeletalMeshPoseTrack;
@@ -29,7 +29,7 @@ public:
 
 	void OnBeginSession(Insights::ITimingViewSession& InTimingViewSession);
 	void OnEndSession(Insights::ITimingViewSession& InTimingViewSession);
-	void Tick(Insights::ITimingViewSession& InTimingViewSession, const Trace::IAnalysisSession& InAnalysisSession);
+	void Tick(Insights::ITimingViewSession& InTimingViewSession, const TraceServices::IAnalysisSession& InAnalysisSession);
 	void ExtendFilterMenu(FMenuBuilder& InMenuBuilder);
 
 #if WITH_ENGINE
@@ -57,7 +57,7 @@ public:
 	bool AreAnyAnimationTracksEnabled() const;
 
 	// Get the last cached analysis session
-	const Trace::IAnalysisSession& GetAnalysisSession() const { return *AnalysisSession; }
+	const TraceServices::IAnalysisSession& GetAnalysisSession() const { return *AnalysisSession; }
 
 	// Check whether the analysis session is valid
 	bool IsAnalysisSessionValid() const { return AnalysisSession != nullptr; }
@@ -81,7 +81,7 @@ public:
 	void OpenAnimGraphTab(uint64 InAnimInstanceId) const;
 
 	/** Get the current marker frame, if any. @return true if the frame is valid */
-	bool GetCurrentMarkerFrame(Trace::FFrame& OutFrame) const { OutFrame = MarkerFrame; return bMarkerFrameValid; }
+	bool GetCurrentMarkerFrame(TraceServices::FFrame& OutFrame) const { OutFrame = MarkerFrame; return bMarkerFrameValid; }
 
 private:
 	// UI handlers
@@ -99,7 +99,7 @@ private:
 	FGameplaySharedData& GameplaySharedData;
 
 	// Cached analysis session, set in Tick()
-	const Trace::IAnalysisSession* AnalysisSession;
+	const TraceServices::IAnalysisSession* AnalysisSession;
 
 	// Cached timing view session, set in OnBeginSession/OnEndSession
 	Insights::ITimingViewSession* TimingViewSession;
@@ -120,7 +120,7 @@ private:
 
 	/** Various times and ranges */
 	double MarkerTime;
-	Trace::FFrame MarkerFrame;
+	TraceServices::FFrame MarkerFrame;
 
 	/** Validity flags for pose times/ranges */
 	bool bTimeMarkerValid;
@@ -133,4 +133,4 @@ private:
 	bool bAnimNodeTracksEnabled;
 	bool bAnimNotifyTracksEnabled;
 	bool bMontageTracksEnabled;
-};
\ No newline at end of file
+};
diff --git a/Engine/Plugins/Animation/GameplayInsights/Source/GameplayInsights/Private/AnimationTickRecordsTrack.cpp b/Engine/Plugins/Animation/GameplayInsights/Source/GameplayInsights/Private/AnimationTickRecordsTrack.cpp
--- a/Engine/Plugins/Animation/GameplayInsights/Source/GameplayInsights/Private/AnimationTickRecordsTrack.cpp
+++ b/Engine/Plugins/Animation/GameplayInsights/Source/GameplayInsights/Private/AnimationTickRecordsTrack.cpp
@@ -98,7 +98,7 @@ void FAnimationTickRecordsTrack::AddAllSeries()
 	const FAnimationProvider* AnimationProvider = SharedData.GetAnalysisSession().ReadProvider<FAnimationProvider>(FAnimationProvider::ProviderName);
 	if(GameplayProvider && AnimationProvider)
 	{
-		Trace::FAnalysisSessionReadScope SessionReadScope(SharedData.GetAnalysisSession());
+		TraceServices::FAnalysisSessionReadScope SessionReadScope(SharedData.GetAnalysisSession());
 
 		AnimationProvider->EnumerateTickRecordIds(GetGameplayTrack().GetObjectId(), [this, &GameplayProvider](uint64 InAssetId, int32 InNodeId)
 		{
@@ -204,7 +204,7 @@ bool FAnimationTickRecordsTrack::UpdateSeriesBoundsHelper(FTickRecordSeries& InS
 
 	if(AnimationProvider)
 	{
-		Trace::FAnalysisSessionReadScope SessionReadScope(SharedData.GetAnalysisSession());
+		TraceServices::FAnalysisSessionReadScope SessionReadScope(SharedData.GetAnalysisSession());
 
 		InSeries.CurrentMin = 0.0;
 		InSeries.CurrentMax = 0.0;
@@ -220,7 +220,7 @@ bool FAnimationTickRecordsTrack::UpdateSeriesBoundsHelper(FTickRecordSeries& InS
 					InSeries.CurrentMax = FMath::Max(InSeries.CurrentMax, Value);
 					bFoundEvents = true;
 				}
-				return Trace::EEventEnumerate::Continue;
+				return TraceServices::EEventEnumerate::Continue;
 			});
 		});
 	}
@@ -235,7 +235,7 @@ void FAnimationTickRecordsTrack::UpdateSeriesHelper(FTickRecordSeries& InSeries,
 
 	if(AnimationProvider)
 	{
-		Trace::FAnalysisSessionReadScope SessionReadScope(SharedData.GetAnalysisSession());
+		TraceServices::FAnalysisSessionReadScope SessionReadScope(SharedData.GetAnalysisSession());
 
 		FGraphTrackBuilder Builder(*this, InSeries, InViewport);
 
@@ -254,7 +254,7 @@ void FAnimationTickRecordsTrack::UpdateSeriesHelper(FTickRecordSeries& InSeries,
 
 					LastFrameWithTickRecord = InMessage.FrameCounter;
 				}
-				return Trace::EEventEnumerate::Continue;
+				return TraceServices::EEventEnumerate::Continue;
 			});
 		});
 	}
@@ -365,14 +365,14 @@ void FAnimationTickRecordsTrack::FindTickRecordMessage(const FTimingEventSearchP
 
 			if(AnimationProvider)
 			{
-				Trace::FAnalysisSessionReadScope SessionReadScope(SharedData.GetAnalysisSession());
+				TraceServices::FAnalysisSessionReadScope SessionReadScope(SharedData.GetAnalysisSession());
 
 				AnimationProvider->ReadTickRecordTimeline(GetGameplayTrack().GetObjectId(), [this, &InContext](const FAnimationProvider::TickRecordTimeline& InTimeline)
 				{
 					InTimeline.EnumerateEvents(InContext.GetParameters().StartTime, InContext.GetParameters().EndTime, [&InContext](double InEventStartTime, double InEventEndTime, uint32 InDepth, const FTickRecordMessage& InMessage)
 					{
 						InContext.Check(InEventStartTime, InEventEndTime, 0, InMessage);
-						return Trace::EEventEnumerate::Continue;
+						return TraceServices::EEventEnumerate::Continue;
 					});
 				});
 			}
@@ -439,13 +439,13 @@ void FAnimationTickRecordsTrack::BuildContextMenu(FMenuBuilder& MenuBuilder)
 	FGameplayGraphTrack::BuildContextMenu(MenuBuilder);
 }
 
-void FAnimationTickRecordsTrack::GetVariantsAtFrame(const Trace::FFrame& InFrame, TArray<TSharedRef<FVariantTreeNode>>& OutVariants) const
+void FAnimationTickRecordsTrack::GetVariantsAtFrame(const TraceServices::FFrame& InFrame, TArray<TSharedRef<FVariantTreeNode>>& OutVariants) const
 {
 	const FGameplayProvider* GameplayProvider = SharedData.GetAnalysisSession().ReadProvider<FGameplayProvider>(FGameplayProvider::ProviderName);
 	const FAnimationProvider* AnimationProvider = SharedData.GetAnalysisSession().ReadProvider<FAnimationProvider>(FAnimationProvider::ProviderName);
 	if(GameplayProvider && AnimationProvider)
 	{
-		Trace::FAnalysisSessionReadScope SessionReadScope(SharedData.GetAnalysisSession());
+		TraceServices::FAnalysisSessionReadScope SessionReadScope(SharedData.GetAnalysisSession());
 
 		AnimationProvider->ReadTickRecordTimeline(GetGameplayTrack().GetObjectId(), [&GameplayProvider, &OutVariants, &InFrame](const FAnimationProvider::TickRecordTimeline& InTimeline)
 		{
@@ -466,7 +466,7 @@ void FAnimationTickRecordsTrack::GetVariantsAtFrame(const Trace::FFrame& InFrame
 						Header->AddChild(FVariantTreeNode::MakeFloat(LOCTEXT("BlendSpacePositionY", "Blend Space Position Y"), InMessage.BlendSpacePositionY));
 					}
 				}
-				return Trace::EEventEnumerate::Continue;
+				return TraceServices::EEventEnumerate::Continue;
 			});
 		});
 	}
diff --git a/Engine/Plugins/Animation/GameplayInsights/Source/GameplayInsights/Private/AnimationTickRecordsTrack.h b/Engine/Plugins/Animation/GameplayInsights/Source/GameplayInsights/Private/AnimationTickRecordsTrack.h
--- a/Engine/Plugins/Animation/GameplayInsights/Source/GameplayInsights/Private/AnimationTickRecordsTrack.h
+++ b/Engine/Plugins/Animation/GameplayInsights/Source/GameplayInsights/Private/AnimationTickRecordsTrack.h
@@ -54,7 +54,7 @@ public:
 	virtual bool UpdateSeriesBounds(FGameplayGraphSeries& InSeries, const FTimingTrackViewport& InViewport) override;
 	virtual void UpdateSeries(FGameplayGraphSeries& InSeries, const FTimingTrackViewport& InViewport) override;
 	virtual void AddAllSeries() override;
-	virtual void GetVariantsAtFrame(const Trace::FFrame& InFrame, TArray<TSharedRef<FVariantTreeNode>>& OutVariants) const override;
+	virtual void GetVariantsAtFrame(const TraceServices::FFrame& InFrame, TArray<TSharedRef<FVariantTreeNode>>& OutVariants) const override;
 
 private:
 	// Helper function used to find a tick record
@@ -79,4 +79,4 @@ private:
 
 	/** Cached hovered event */
 	TWeakPtr<const FGraphTrackEvent> CachedHoveredEvent;
-};
\ No newline at end of file
+};
diff --git a/Engine/Plugins/Animation/GameplayInsights/Source/GameplayInsights/Private/GameplayAnalyzer.cpp b/Engine/Plugins/Animation/GameplayInsights/Source/GameplayInsights/Private/GameplayAnalyzer.cpp
--- a/Engine/Plugins/Animation/GameplayInsights/Source/GameplayInsights/Private/GameplayAnalyzer.cpp
+++ b/Engine/Plugins/Animation/GameplayInsights/Source/GameplayInsights/Private/GameplayAnalyzer.cpp
@@ -4,7 +4,7 @@
 #include "TraceServices/Model/AnalysisSession.h"
 #include "GameplayProvider.h"
 
-FGameplayAnalyzer::FGameplayAnalyzer(Trace::IAnalysisSession& InSession, FGameplayProvider& InGameplayProvider)
+FGameplayAnalyzer::FGameplayAnalyzer(TraceServices::IAnalysisSession& InSession, FGameplayProvider& InGameplayProvider)
 	: Session(InSession)
 	, GameplayProvider(InGameplayProvider)
 {
@@ -27,7 +27,7 @@ void FGameplayAnalyzer::OnAnalysisBegin(const FOnAnalysisContext& Context)
 
 bool FGameplayAnalyzer::OnEvent(uint16 RouteId, EStyle Style, const FOnEventContext& Context)
 {
-	Trace::FAnalysisSessionEditScope _(Session);
+	TraceServices::FAnalysisSessionEditScope _(Session);
 
 	const auto& EventData = Context.EventData;
 	switch (RouteId)
@@ -116,4 +116,4 @@ bool FGameplayAnalyzer::OnEvent(uint16 RouteId, EStyle Style, const FOnEventCont
 	}
 
 	return true;
-}
\ No newline at end of file
+}
diff --git a/Engine/Plugins/Animation/GameplayInsights/Source/GameplayInsights/Private/GameplayAnalyzer.h b/Engine/Plugins/Animation/GameplayInsights/Source/GameplayInsights/Private/GameplayAnalyzer.h
--- a/Engine/Plugins/Animation/GameplayInsights/Source/GameplayInsights/Private/GameplayAnalyzer.h
+++ b/Engine/Plugins/Animation/GameplayInsights/Source/GameplayInsights/Private/GameplayAnalyzer.h
@@ -5,12 +5,12 @@
 #include "Trace/Analyzer.h"
 
 class FGameplayProvider;
-namespace Trace { class IAnalysisSession; }
+namespace TraceServices { class IAnalysisSession; }
 
 class FGameplayAnalyzer : public Trace::IAnalyzer
 {
 public:
-	FGameplayAnalyzer(Trace::IAnalysisSession& InSession, FGameplayProvider& InGameplayProvider);
+	FGameplayAnalyzer(TraceServices::IAnalysisSession& InSession, FGameplayProvider& InGameplayProvider);
 
 	virtual void OnAnalysisBegin(const FOnAnalysisContext& Context) override;
 	virtual void OnAnalysisEnd() override {}
@@ -30,6 +30,6 @@ private:
 		RouteId_PropertyValue,
 	};
 
-	Trace::IAnalysisSession& Session;
+	TraceServices::IAnalysisSession& Session;
 	FGameplayProvider& GameplayProvider;
-};
\ No newline at end of file
+};
diff --git a/Engine/Plugins/Animation/GameplayInsights/Source/GameplayInsights/Private/GameplayGraphTrack.cpp b/Engine/Plugins/Animation/GameplayInsights/Source/GameplayInsights/Private/GameplayGraphTrack.cpp
--- a/Engine/Plugins/Animation/GameplayInsights/Source/GameplayInsights/Private/GameplayGraphTrack.cpp
+++ b/Engine/Plugins/Animation/GameplayInsights/Source/GameplayInsights/Private/GameplayGraphTrack.cpp
@@ -309,7 +309,7 @@ void FGameplayGraphTrack::BuildContextMenu(FMenuBuilder& MenuBuilder)
 	MenuBuilder.EndSection();
 }
 
-void FGameplayGraphTrack::GetVariantsAtFrame(const Trace::FFrame& InFrame, TArray<TSharedRef<FVariantTreeNode>>& OutVariants) const
+void FGameplayGraphTrack::GetVariantsAtFrame(const TraceServices::FFrame& InFrame, TArray<TSharedRef<FVariantTreeNode>>& OutVariants) const
 { 
 	GetVariantsAtTime(InFrame.StartTime, OutVariants); 
 }
diff --git a/Engine/Plugins/Animation/GameplayInsights/Source/GameplayInsights/Private/GameplayGraphTrack.h b/Engine/Plugins/Animation/GameplayInsights/Source/GameplayInsights/Private/GameplayGraphTrack.h
--- a/Engine/Plugins/Animation/GameplayInsights/Source/GameplayInsights/Private/GameplayGraphTrack.h
+++ b/Engine/Plugins/Animation/GameplayInsights/Source/GameplayInsights/Private/GameplayGraphTrack.h
@@ -9,7 +9,7 @@
 class FTimingEventSearchParameters;
 class FGameplayGraphTrack;
 struct FVariantTreeNode;
-namespace Trace { struct FFrame; }
+namespace TraceServices { struct FFrame; }
 
 /** The various layouts that we display series with */
 enum class EGameplayGraphLayout : int32
@@ -57,7 +57,7 @@ public:
 
 	// Get all variants at the specified time/frame
 	virtual void GetVariantsAtTime(double InTime, TArray<TSharedRef<FVariantTreeNode>>& OutVariants) const {}
-	virtual void GetVariantsAtFrame(const Trace::FFrame& InFrame, TArray<TSharedRef<FVariantTreeNode>>& OutVariants) const;
+	virtual void GetVariantsAtFrame(const TraceServices::FFrame& InFrame, TArray<TSharedRef<FVariantTreeNode>>& OutVariants) const;
 
 	// Get the requested track size scale
 	float GetRequestedTrackSizeScale() const { return RequestedTrackSizeScale; }
diff --git a/Engine/Plugins/Animation/GameplayInsights/Source/GameplayInsights/Private/GameplayInsightsModule.cpp b/Engine/Plugins/Animation/GameplayInsights/Source/GameplayInsights/Private/GameplayInsightsModule.cpp
--- a/Engine/Plugins/Animation/GameplayInsights/Source/GameplayInsights/Private/GameplayInsightsModule.cpp
+++ b/Engine/Plugins/Animation/GameplayInsights/Source/GameplayInsights/Private/GameplayInsightsModule.cpp
@@ -38,7 +38,7 @@ const FName GameplayInsightsTabs::DocumentTab("DocumentTab");
 
 void FGameplayInsightsModule::StartupModule()
 {
-	IModularFeatures::Get().RegisterModularFeature(Trace::ModuleFeatureName, &GameplayTraceModule);
+	IModularFeatures::Get().RegisterModularFeature(TraceServices::ModuleFeatureName, &GameplayTraceModule);
 	IModularFeatures::Get().RegisterModularFeature(Insights::TimingViewExtenderFeatureName, &GameplayTimingViewExtender);
 
 	TickerHandle = FTicker::GetCoreTicker().AddTicker(TEXT("GameplayInsights"), 0.0f, [this](float DeltaTime)
@@ -207,7 +207,7 @@ void FGameplayInsightsModule::ShutdownModule()
 
 	FTicker::GetCoreTicker().RemoveTicker(TickerHandle);
 
-	IModularFeatures::Get().UnregisterModularFeature(Trace::ModuleFeatureName, &GameplayTraceModule);
+	IModularFeatures::Get().UnregisterModularFeature(TraceServices::ModuleFeatureName, &GameplayTraceModule);
 	IModularFeatures::Get().UnregisterModularFeature(Insights::TimingViewExtenderFeatureName, &GameplayTimingViewExtender);
 }
 
diff --git a/Engine/Plugins/Animation/GameplayInsights/Source/GameplayInsights/Private/GameplayProvider.cpp b/Engine/Plugins/Animation/GameplayInsights/Source/GameplayInsights/Private/GameplayProvider.cpp
--- a/Engine/Plugins/Animation/GameplayInsights/Source/GameplayInsights/Private/GameplayProvider.cpp
+++ b/Engine/Plugins/Animation/GameplayInsights/Source/GameplayInsights/Private/GameplayProvider.cpp
@@ -6,7 +6,7 @@ FName FGameplayProvider::ProviderName("GameplayProvider");
 
 #define LOCTEXT_NAMESPACE "GameplayProvider"
 
-FGameplayProvider::FGameplayProvider(Trace::IAnalysisSession& InSession)
+FGameplayProvider::FGameplayProvider(TraceServices::IAnalysisSession& InSession)
 	: Session(InSession)
 	, EndPlayEvent(nullptr)
 	, bHasAnyData(false)
@@ -286,7 +286,7 @@ void FGameplayProvider::AppendObjectEvent(uint64 InObjectId, double InTime, cons
 		EndPlayEvent = Session.StoreString(TEXT("EndPlay"));
 	}
 
-	TSharedPtr<Trace::TPointTimeline<FObjectEventMessage>> Timeline;
+	TSharedPtr<TraceServices::TPointTimeline<FObjectEventMessage>> Timeline;
 	uint32* IndexPtr = ObjectIdToEventTimelines.Find(InObjectId);
 	if(IndexPtr != nullptr)
 	{
@@ -294,7 +294,7 @@ void FGameplayProvider::AppendObjectEvent(uint64 InObjectId, double InTime, cons
 	}
 	else
 	{
-		Timeline = MakeShared<Trace::TPointTimeline<FObjectEventMessage>>(Session.GetLinearAllocator());
+		Timeline = MakeShared<TraceServices::TPointTimeline<FObjectEventMessage>>(Session.GetLinearAllocator());
 		ObjectIdToEventTimelines.Add(InObjectId, EventTimelines.Num());
 		EventTimelines.Add(Timeline.ToSharedRef());
 	}
@@ -383,7 +383,7 @@ void FGameplayProvider::AppendPropertiesStart(uint64 InObjectId, double InTime,
 	else
 	{
 		Storage = MakeShared<FObjectPropertiesStorage>();
-		Storage->Timeline = MakeShared<Trace::TIntervalTimeline<FObjectPropertiesMessage>>(Session.GetLinearAllocator());
+		Storage->Timeline = MakeShared<TraceServices::TIntervalTimeline<FObjectPropertiesMessage>>(Session.GetLinearAllocator());
 		ObjectIdToPropertiesStorage.Add(InObjectId, PropertiesStorage.Num());
 		PropertiesStorage.Add(Storage.ToSharedRef());
 	}
@@ -412,7 +412,7 @@ void FGameplayProvider::AppendPropertiesEnd(uint64 InObjectId, double InTime)
 	else
 	{
 		Storage = MakeShared<FObjectPropertiesStorage>();
-		Storage->Timeline = MakeShared<Trace::TIntervalTimeline<FObjectPropertiesMessage>>(Session.GetLinearAllocator());
+		Storage->Timeline = MakeShared<TraceServices::TIntervalTimeline<FObjectPropertiesMessage>>(Session.GetLinearAllocator());
 		ObjectIdToPropertiesStorage.Add(InObjectId, PropertiesStorage.Num());
 		PropertiesStorage.Add(Storage.ToSharedRef());
 	}
@@ -442,7 +442,7 @@ void FGameplayProvider::AppendPropertyValue(uint64 InObjectId, double InTime, ui
 	else
 	{
 		Storage = MakeShared<FObjectPropertiesStorage>();
-		Storage->Timeline = MakeShared<Trace::TIntervalTimeline<FObjectPropertiesMessage>>(Session.GetLinearAllocator());
+		Storage->Timeline = MakeShared<TraceServices::TIntervalTimeline<FObjectPropertiesMessage>>(Session.GetLinearAllocator());
 		ObjectIdToPropertiesStorage.Add(InObjectId, PropertiesStorage.Num());
 		PropertiesStorage.Add(Storage.ToSharedRef());
 	}
@@ -472,4 +472,4 @@ bool FGameplayProvider::HasObjectProperties() const
 	return bHasObjectProperties;
 }
 
-#undef LOCTEXT_NAMESPACE
\ No newline at end of file
+#undef LOCTEXT_NAMESPACE
diff --git a/Engine/Plugins/Animation/GameplayInsights/Source/GameplayInsights/Private/GameplayProvider.h b/Engine/Plugins/Animation/GameplayInsights/Source/GameplayInsights/Private/GameplayProvider.h
--- a/Engine/Plugins/Animation/GameplayInsights/Source/GameplayInsights/Private/GameplayProvider.h
+++ b/Engine/Plugins/Animation/GameplayInsights/Source/GameplayInsights/Private/GameplayProvider.h
@@ -9,14 +9,14 @@
 #include "Containers/StringView.h"
 #include "Model/IntervalTimeline.h"
 
-namespace Trace { class IAnalysisSession; }
+namespace TraceServices { class IAnalysisSession; }
 
 class FGameplayProvider : public IGameplayProvider
 {
 public:
 	static FName ProviderName;
 
-	FGameplayProvider(Trace::IAnalysisSession& InSession);
+	FGameplayProvider(TraceServices::IAnalysisSession& InSession);
 
 	/** IGameplayProvider interface */
 	virtual bool ReadObjectEventsTimeline(uint64 InObjectId, TFunctionRef<void(const ObjectEventsTimeline&)> Callback) const override;
@@ -70,7 +70,7 @@ public:
 	bool HasObjectProperties() const;
 
 private:
-	Trace::IAnalysisSession& Session;
+	TraceServices::IAnalysisSession& Session;
 
 	/** All class info, grow only for stable indices */
 	TArray<FClassInfo> ClassInfos;
@@ -99,12 +99,12 @@ private:
 		double OpenStartTime;
 		uint64 OpenEventId;
 		FObjectPropertiesMessage OpenEvent;
-		TSharedPtr<Trace::TIntervalTimeline<FObjectPropertiesMessage>> Timeline;
+		TSharedPtr<TraceServices::TIntervalTimeline<FObjectPropertiesMessage>> Timeline;
 		TArray<FObjectPropertyValue> Values;
 	};
 
 	/** Message storage */
-	TArray<TSharedRef<Trace::TPointTimeline<FObjectEventMessage>>> EventTimelines;
+	TArray<TSharedRef<TraceServices::TPointTimeline<FObjectEventMessage>>> EventTimelines;
 	TArray<TSharedRef<FObjectPropertiesStorage>> PropertiesStorage;
 
 	/** Map of class path name to ClassInfo index */
@@ -124,4 +124,4 @@ private:
 
 	/** Whether we have any object properties */
 	bool bHasObjectProperties;
-};
\ No newline at end of file
+};
diff --git a/Engine/Plugins/Animation/GameplayInsights/Source/GameplayInsights/Private/GameplaySharedData.cpp b/Engine/Plugins/Animation/GameplayInsights/Source/GameplayInsights/Private/GameplaySharedData.cpp
--- a/Engine/Plugins/Animation/GameplayInsights/Source/GameplayInsights/Private/GameplaySharedData.cpp
+++ b/Engine/Plugins/Animation/GameplayInsights/Source/GameplayInsights/Private/GameplaySharedData.cpp
@@ -38,7 +38,7 @@ void FGameplaySharedData::OnEndSession(Insights::ITimingViewSession& InTimingVie
 	TimingViewSession = nullptr;
 }
 
-TSharedRef<FObjectEventsTrack> FGameplaySharedData::GetObjectEventsTrackForId(Insights::ITimingViewSession& InTimingViewSession, const Trace::IAnalysisSession& InAnalysisSession, const FObjectInfo& InObjectInfo)
+TSharedRef<FObjectEventsTrack> FGameplaySharedData::GetObjectEventsTrackForId(Insights::ITimingViewSession& InTimingViewSession, const TraceServices::IAnalysisSession& InAnalysisSession, const FObjectInfo& InObjectInfo)
 {
 	const FGameplayProvider* GameplayProvider = InAnalysisSession.ReadProvider<FGameplayProvider>(FGameplayProvider::ProviderName);
 	check(GameplayProvider);
@@ -136,7 +136,7 @@ void FGameplaySharedData::MakeTrackAndAncestorsVisible(const TSharedRef<FObjectE
 	InvalidateObjectTracksOrder();
 }
 
-void FGameplaySharedData::Tick(Insights::ITimingViewSession& InTimingViewSession, const Trace::IAnalysisSession& InAnalysisSession)
+void FGameplaySharedData::Tick(Insights::ITimingViewSession& InTimingViewSession, const TraceServices::IAnalysisSession& InAnalysisSession)
 {
 	AnalysisSession = &InAnalysisSession;
 
@@ -144,7 +144,7 @@ void FGameplaySharedData::Tick(Insights::ITimingViewSession& InTimingViewSession
 
 	if(GameplayProvider)
 	{
-		Trace::FAnalysisSessionReadScope SessionReadScope(GetAnalysisSession());
+		TraceServices::FAnalysisSessionReadScope SessionReadScope(GetAnalysisSession());
 
 		if(GameplayProvider->HasAnyData())
 		{
@@ -357,4 +357,4 @@ void FGameplaySharedData::OpenTrackVariantsTab(const FGameplayTrack& InGameplayT
 	}
 }
 
-#undef LOCTEXT_NAMESPACE
\ No newline at end of file
+#undef LOCTEXT_NAMESPACE
diff --git a/Engine/Plugins/Animation/GameplayInsights/Source/GameplayInsights/Private/GameplaySharedData.h b/Engine/Plugins/Animation/GameplayInsights/Source/GameplayInsights/Private/GameplaySharedData.h
--- a/Engine/Plugins/Animation/GameplayInsights/Source/GameplayInsights/Private/GameplaySharedData.h
+++ b/Engine/Plugins/Animation/GameplayInsights/Source/GameplayInsights/Private/GameplaySharedData.h
@@ -5,7 +5,7 @@
 #include "Insights/ViewModels/TimingEventsTrack.h"
 #include "Framework/Docking/TabManager.h"
 
-namespace Trace { class IAnalysisSession; }
+namespace TraceServices { class IAnalysisSession; }
 namespace Insights { class ITimingViewSession; }
 class FObjectEventsTrack;
 class FObjectPropertiesTrack;
@@ -21,11 +21,11 @@ public:
 
 	void OnBeginSession(Insights::ITimingViewSession& InTimingViewSession);
 	void OnEndSession(Insights::ITimingViewSession& InTimingViewSession);
-	void Tick(Insights::ITimingViewSession& InTimingViewSession, const Trace::IAnalysisSession& InAnalysisSession);
+	void Tick(Insights::ITimingViewSession& InTimingViewSession, const TraceServices::IAnalysisSession& InAnalysisSession);
 	void ExtendFilterMenu(FMenuBuilder& InMenuBuilder);
 
 	// Helper function. Builds object track hierarchy on-demand and returns a track for the supplied object info.
-	TSharedRef<FObjectEventsTrack> GetObjectEventsTrackForId(Insights::ITimingViewSession& InTimingViewSession, const Trace::IAnalysisSession& InAnalysisSession, const FObjectInfo& InObjectInfo);
+	TSharedRef<FObjectEventsTrack> GetObjectEventsTrackForId(Insights::ITimingViewSession& InTimingViewSession, const TraceServices::IAnalysisSession& InAnalysisSession, const FObjectInfo& InObjectInfo);
 
 	// Helper function to make tracks visible when children are first added
 	void MakeTrackAndAncestorsVisible(const TSharedRef<FObjectEventsTrack>& InObjectEventsTrack, bool bInVisible);
@@ -37,7 +37,7 @@ public:
 	void InvalidateObjectTracksOrder() { bObjectTracksDirty = true; }
 
 	// Get the last cached analysis session
-	const Trace::IAnalysisSession& GetAnalysisSession() const { return *AnalysisSession; }
+	const TraceServices::IAnalysisSession& GetAnalysisSession() const { return *AnalysisSession; }
 
 	// Get the timing view session
 	Insights::ITimingViewSession& GetTimingViewSession() { return *TimingViewSession; }
@@ -94,7 +94,7 @@ private:
 	TArray<TSharedRef<FBaseTimingTrack>> RootTracks;
 
 	// Cached analysis session, set in Tick()
-	const Trace::IAnalysisSession* AnalysisSession;
+	const TraceServices::IAnalysisSession* AnalysisSession;
 
 	// Cached timing view session, set in OnBeginSession/OnEndSession
 	Insights::ITimingViewSession* TimingViewSession;
@@ -113,4 +113,4 @@ private:
 
 	// Whether all of our object property tracks are enabled
 	bool bObjectPropertyTracksEnabled;
-};
\ No newline at end of file
+};
diff --git a/Engine/Plugins/Animation/GameplayInsights/Source/GameplayInsights/Private/GameplayTimingViewExtender.cpp b/Engine/Plugins/Animation/GameplayInsights/Source/GameplayInsights/Private/GameplayTimingViewExtender.cpp
--- a/Engine/Plugins/Animation/GameplayInsights/Source/GameplayInsights/Private/GameplayTimingViewExtender.cpp
+++ b/Engine/Plugins/Animation/GameplayInsights/Source/GameplayInsights/Private/GameplayTimingViewExtender.cpp
@@ -75,7 +75,7 @@ void FGameplayTimingViewExtender::OnEndSession(Insights::ITimingViewSession& InS
 	PerSessionDataMap.Remove(&InSession);
 }
 
-void FGameplayTimingViewExtender::Tick(Insights::ITimingViewSession& InSession, const Trace::IAnalysisSession& InAnalysisSession)
+void FGameplayTimingViewExtender::Tick(Insights::ITimingViewSession& InSession, const TraceServices::IAnalysisSession& InAnalysisSession)
 {
 	FPerSessionData* PerSessionData = PerSessionDataMap.Find(&InSession);
 	if(PerSessionData != nullptr)
diff --git a/Engine/Plugins/Animation/GameplayInsights/Source/GameplayInsights/Private/GameplayTimingViewExtender.h b/Engine/Plugins/Animation/GameplayInsights/Source/GameplayInsights/Private/GameplayTimingViewExtender.h
--- a/Engine/Plugins/Animation/GameplayInsights/Source/GameplayInsights/Private/GameplayTimingViewExtender.h
+++ b/Engine/Plugins/Animation/GameplayInsights/Source/GameplayInsights/Private/GameplayTimingViewExtender.h
@@ -18,7 +18,7 @@ public:
 	// Insights::ITimingViewExtender interface
 	virtual void OnBeginSession(Insights::ITimingViewSession& InSession) override;
 	virtual void OnEndSession(Insights::ITimingViewSession& InSession) override;
-	virtual void Tick(Insights::ITimingViewSession& InSession, const Trace::IAnalysisSession& InAnalysisSession) override;
+	virtual void Tick(Insights::ITimingViewSession& InSession, const TraceServices::IAnalysisSession& InAnalysisSession) override;
 	virtual void ExtendFilterMenu(Insights::ITimingViewSession& InSession, FMenuBuilder& InMenuBuilder) override;
 
 #if WITH_EDITOR
diff --git a/Engine/Plugins/Animation/GameplayInsights/Source/GameplayInsights/Private/GameplayTraceModule.cpp b/Engine/Plugins/Animation/GameplayInsights/Source/GameplayInsights/Private/GameplayTraceModule.cpp
--- a/Engine/Plugins/Animation/GameplayInsights/Source/GameplayInsights/Private/GameplayTraceModule.cpp
+++ b/Engine/Plugins/Animation/GameplayInsights/Source/GameplayInsights/Private/GameplayTraceModule.cpp
@@ -8,13 +8,13 @@
 
 FName FGameplayTraceModule::ModuleName("GameplayTrace");
 
-void FGameplayTraceModule::GetModuleInfo(Trace::FModuleInfo& OutModuleInfo)
+void FGameplayTraceModule::GetModuleInfo(TraceServices::FModuleInfo& OutModuleInfo)
 {
 	OutModuleInfo.Name = ModuleName;
 	OutModuleInfo.DisplayName = TEXT("Gameplay");
 }
 
-void FGameplayTraceModule::OnAnalysisBegin(Trace::IAnalysisSession& InSession)
+void FGameplayTraceModule::OnAnalysisBegin(TraceServices::IAnalysisSession& InSession)
 {
 	FGameplayProvider* GameplayProvider = new FGameplayProvider(InSession);
 	InSession.AddProvider(FGameplayProvider::ProviderName, GameplayProvider);
@@ -31,12 +31,12 @@ void FGameplayTraceModule::GetLoggers(TArray<const TCHAR *>& OutLoggers)
 	OutLoggers.Add(TEXT("Animation"));
 }
 
-void FGameplayTraceModule::GenerateReports(const Trace::IAnalysisSession& Session, const TCHAR* CmdLine, const TCHAR* OutputDirectory)
+void FGameplayTraceModule::GenerateReports(const TraceServices::IAnalysisSession& Session, const TCHAR* CmdLine, const TCHAR* OutputDirectory)
 {
 
 }
 
-const IAnimationProvider* ReadAnimationProvider(const Trace::IAnalysisSession& Session)
+const IAnimationProvider* ReadAnimationProvider(const TraceServices::IAnalysisSession& Session)
 {
 	return Session.ReadProvider<IAnimationProvider>(FAnimationProvider::ProviderName);
 }
diff --git a/Engine/Plugins/Animation/GameplayInsights/Source/GameplayInsights/Private/GameplayTraceModule.h b/Engine/Plugins/Animation/GameplayInsights/Source/GameplayInsights/Private/GameplayTraceModule.h
--- a/Engine/Plugins/Animation/GameplayInsights/Source/GameplayInsights/Private/GameplayTraceModule.h
+++ b/Engine/Plugins/Animation/GameplayInsights/Source/GameplayInsights/Private/GameplayTraceModule.h
@@ -8,18 +8,18 @@
 class IAnimationProvider;
 namespace Insights { class ITimingViewSession; }
 
-class FGameplayTraceModule : public Trace::IModule
+class FGameplayTraceModule : public TraceServices::IModule
 {
 public:
-	// Trace::IModule interface
-	virtual void GetModuleInfo(Trace::FModuleInfo& OutModuleInfo) override;
-	virtual void OnAnalysisBegin(Trace::IAnalysisSession& Session) override;
+	// TraceServices::IModule interface
+	virtual void GetModuleInfo(TraceServices::FModuleInfo& OutModuleInfo) override;
+	virtual void OnAnalysisBegin(TraceServices::IAnalysisSession& Session) override;
 	virtual void GetLoggers(TArray<const TCHAR *>& OutLoggers) override;
-	virtual void GenerateReports(const Trace::IAnalysisSession& Session, const TCHAR* CmdLine, const TCHAR* OutputDirectory) override;
+	virtual void GenerateReports(const TraceServices::IAnalysisSession& Session, const TCHAR* CmdLine, const TCHAR* OutputDirectory) override;
 	virtual const TCHAR* GetCommandLineArgument() override { return TEXT("objecttrace"); }
 
 private:
 	static FName ModuleName;
 };
 
-const IAnimationProvider* ReadAnimationProvider(const Trace::IAnalysisSession& Session);
\ No newline at end of file
+const IAnimationProvider* ReadAnimationProvider(const TraceServices::IAnalysisSession& Session);
diff --git a/Engine/Plugins/Animation/GameplayInsights/Source/GameplayInsights/Private/GameplayTrack.cpp b/Engine/Plugins/Animation/GameplayInsights/Source/GameplayInsights/Private/GameplayTrack.cpp
--- a/Engine/Plugins/Animation/GameplayInsights/Source/GameplayInsights/Private/GameplayTrack.cpp
+++ b/Engine/Plugins/Animation/GameplayInsights/Source/GameplayInsights/Private/GameplayTrack.cpp
@@ -129,12 +129,12 @@ FText FGameplayTrack::GetWorldPostFix(const FWorldInfo& InWorldInfo)
 	}
 }
 
-FText FGameplayTrack::GetWorldName(const Trace::IAnalysisSession& InAnalysisSession) const
+FText FGameplayTrack::GetWorldName(const TraceServices::IAnalysisSession& InAnalysisSession) const
 {
 	const FGameplayProvider* GameplayProvider = InAnalysisSession.ReadProvider<FGameplayProvider>(FGameplayProvider::ProviderName);
 	if (GameplayProvider)
 	{
-		Trace::FAnalysisSessionReadScope SessionReadScope(InAnalysisSession);
+		TraceServices::FAnalysisSessionReadScope SessionReadScope(InAnalysisSession);
 
 		if(const FWorldInfo* WorldInfo = GameplayProvider->FindWorldInfoFromObject(ObjectId))
 		{
@@ -146,7 +146,7 @@ FText FGameplayTrack::GetWorldName(const Trace::IAnalysisSession& InAnalysisSess
 	return LOCTEXT("UnknownWorld", "Unknown");
 }
 
-void FGameplayTimingEventsTrack::GetVariantsAtFrame(const Trace::FFrame& InFrame, TArray<TSharedRef<FVariantTreeNode>>& OutVariants) const
+void FGameplayTimingEventsTrack::GetVariantsAtFrame(const TraceServices::FFrame& InFrame, TArray<TSharedRef<FVariantTreeNode>>& OutVariants) const
 { 
 	GetVariantsAtTime(InFrame.StartTime, OutVariants);
 }
@@ -170,4 +170,4 @@ void FGameplayTimingEventsTrack::BuildContextMenu(FMenuBuilder& MenuBuilder)
 	MenuBuilder.EndSection();
 }
 
-#undef LOCTEXT_NAMESPACE
\ No newline at end of file
+#undef LOCTEXT_NAMESPACE
diff --git a/Engine/Plugins/Animation/GameplayInsights/Source/GameplayInsights/Private/GameplayTrack.h b/Engine/Plugins/Animation/GameplayInsights/Source/GameplayInsights/Private/GameplayTrack.h
--- a/Engine/Plugins/Animation/GameplayInsights/Source/GameplayInsights/Private/GameplayTrack.h
+++ b/Engine/Plugins/Animation/GameplayInsights/Source/GameplayInsights/Private/GameplayTrack.h
@@ -12,8 +12,8 @@ namespace GameplayTrackConstants
 struct FVariantTreeNode;
 class FGameplaySharedData;
 struct FWorldInfo;
-namespace Trace { class IAnalysisSession; }
-namespace Trace { struct FFrame; }
+namespace TraceServices { class IAnalysisSession; }
+namespace TraceServices { struct FFrame; }
 
 // Provides parent/child hierarchy structure and the owning object Id
 // Designed as a compositional member of outer timing tracks
@@ -56,7 +56,7 @@ public:
 	uint32 GetIndent() const { return Indent; }
 
 	/** Get the world name for this object */
-	FText GetWorldName(const Trace::IAnalysisSession& InAnalysisSession) const;
+	FText GetWorldName(const TraceServices::IAnalysisSession& InAnalysisSession) const;
 
 	/** Helper function to find a postfix for a world (e.g. "Client 0", "Server" etc.) */
 	static FText GetWorldPostFix(const FWorldInfo& InWorldInfo);
@@ -115,7 +115,7 @@ public:
 
 	/** Get all variants at the specified time/frame */
 	virtual void GetVariantsAtTime(double InTime, TArray<TSharedRef<FVariantTreeNode>>& OutVariants) const {}
-	virtual void GetVariantsAtFrame(const Trace::FFrame& InFrame, TArray<TSharedRef<FVariantTreeNode>>& OutVariants) const;
+	virtual void GetVariantsAtFrame(const TraceServices::FFrame& InFrame, TArray<TSharedRef<FVariantTreeNode>>& OutVariants) const;
 
 protected:
 	const FGameplaySharedData& GameplaySharedData;
diff --git a/Engine/Plugins/Animation/GameplayInsights/Source/GameplayInsights/Private/MontageTrack.cpp b/Engine/Plugins/Animation/GameplayInsights/Source/GameplayInsights/Private/MontageTrack.cpp
--- a/Engine/Plugins/Animation/GameplayInsights/Source/GameplayInsights/Private/MontageTrack.cpp
+++ b/Engine/Plugins/Animation/GameplayInsights/Source/GameplayInsights/Private/MontageTrack.cpp
@@ -33,7 +33,7 @@ void FMontageTrack::AddAllSeries()
 
 	if(AnimationProvider && GameplayProvider)
 	{
-		Trace::FAnalysisSessionReadScope SessionReadScope(SharedData.GetAnalysisSession());
+		TraceServices::FAnalysisSessionReadScope SessionReadScope(SharedData.GetAnalysisSession());
 
 		AnimationProvider->EnumerateMontageIds(GetGameplayTrack().GetObjectId(), [this, &AnimationProvider, &GameplayProvider](uint64 InMontageId)
 		{
@@ -77,7 +77,7 @@ bool FMontageTrack::UpdateSeriesBounds(FGameplayGraphSeries& InSeries, const FTi
 
 	if(AnimationProvider)
 	{
-		Trace::FAnalysisSessionReadScope SessionReadScope(SharedData.GetAnalysisSession());
+		TraceServices::FAnalysisSessionReadScope SessionReadScope(SharedData.GetAnalysisSession());
 
 		FMontageSeries& MontageSeries = *static_cast<FMontageSeries*>(&InSeries);
 
@@ -94,7 +94,7 @@ bool FMontageTrack::UpdateSeriesBounds(FGameplayGraphSeries& InSeries, const FTi
 					MontageSeries.CurrentMax = FMath::Max(MontageSeries.CurrentMax, InMessage.Weight);
 					bFoundEvents = true;
 				}
-				return Trace::EEventEnumerate::Continue;
+				return TraceServices::EEventEnumerate::Continue;
 			});
 		});
 	}
@@ -108,7 +108,7 @@ void FMontageTrack::UpdateSeries(FGameplayGraphSeries& InSeries, const FTimingTr
 
 	if(AnimationProvider)
 	{
-		Trace::FAnalysisSessionReadScope SessionReadScope(SharedData.GetAnalysisSession());
+		TraceServices::FAnalysisSessionReadScope SessionReadScope(SharedData.GetAnalysisSession());
 
 		FMontageSeries& MontageSeries = *static_cast<FMontageSeries*>(&InSeries);
 
@@ -129,7 +129,7 @@ void FMontageTrack::UpdateSeries(FGameplayGraphSeries& InSeries, const FTimingTr
 
 					LastFrameWithMontage = FrameCounter;
 				}
-				return Trace::EEventEnumerate::Continue;
+				return TraceServices::EEventEnumerate::Continue;
 			});
 		});
 	}
@@ -152,7 +152,7 @@ void FMontageTrack::InitTooltip(FTooltipDrawState& Tooltip, const ITimingEvent&
 		Tooltip.AddNameValueTextLine(LOCTEXT("EventDesiredWeight", "Desired Weight").ToString(), FText::AsNumber(InMessage.DesiredWeight).ToString());
 
 		{
-			Trace::FAnalysisSessionReadScope SessionReadScope(SharedData.GetAnalysisSession());
+			TraceServices::FAnalysisSessionReadScope SessionReadScope(SharedData.GetAnalysisSession());
 
 			const FGameplayProvider* GameplayProvider = SharedData.GetAnalysisSession().ReadProvider<FGameplayProvider>(FGameplayProvider::ProviderName);
 			if(GameplayProvider)
@@ -200,14 +200,14 @@ void FMontageTrack::FindMontageMessage(const FTimingEventSearchParameters& InPar
 
 			if(AnimationProvider)
 			{
-				Trace::FAnalysisSessionReadScope SessionReadScope(SharedData.GetAnalysisSession());
+				TraceServices::FAnalysisSessionReadScope SessionReadScope(SharedData.GetAnalysisSession());
 
 				AnimationProvider->ReadMontageTimeline(GetGameplayTrack().GetObjectId(), [&InContext](const FAnimationProvider::AnimMontageTimeline& InTimeline)
 				{
 					InTimeline.EnumerateEvents(InContext.GetParameters().StartTime, InContext.GetParameters().EndTime, [&InContext](double InEventStartTime, double InEventEndTime, uint32 InDepth, const FAnimMontageMessage& InMessage)
 					{
 						InContext.Check(InEventStartTime, InEventEndTime, 0, InMessage);
-						return Trace::EEventEnumerate::Continue;
+						return TraceServices::EEventEnumerate::Continue;
 					});
 				});
 			}
@@ -227,17 +227,17 @@ void FMontageTrack::FindMontageMessage(const FTimingEventSearchParameters& InPar
 		TTimingEventSearch<FAnimMontageMessage>::NoMatch);
 }
 
-void FMontageTrack::GetVariantsAtFrame(const Trace::FFrame& InFrame, TArray<TSharedRef<FVariantTreeNode>>& OutVariants) const 
+void FMontageTrack::GetVariantsAtFrame(const TraceServices::FFrame& InFrame, TArray<TSharedRef<FVariantTreeNode>>& OutVariants) const 
 {
 	TSharedRef<FVariantTreeNode> Header = OutVariants.Add_GetRef(FVariantTreeNode::MakeHeader(LOCTEXT("MontagesHeader", "Montages"), 0));
 
-	const Trace::IFrameProvider& FramesProvider = Trace::ReadFrameProvider(SharedData.GetAnalysisSession());
+	const TraceServices::IFrameProvider& FramesProvider = TraceServices::ReadFrameProvider(SharedData.GetAnalysisSession());
 	const FAnimationProvider* AnimationProvider = SharedData.GetAnalysisSession().ReadProvider<FAnimationProvider>(FAnimationProvider::ProviderName);
 	const FGameplayProvider* GameplayProvider = SharedData.GetAnalysisSession().ReadProvider<FGameplayProvider>(FGameplayProvider::ProviderName);
 
 	if(AnimationProvider && GameplayProvider)
 	{
-		Trace::FAnalysisSessionReadScope SessionReadScope(SharedData.GetAnalysisSession());
+		TraceServices::FAnalysisSessionReadScope SessionReadScope(SharedData.GetAnalysisSession());
 
 		AnimationProvider->ReadMontageTimeline(GetGameplayTrack().GetObjectId(), [this, &InFrame, &AnimationProvider, &GameplayProvider, &Header](const FAnimationProvider::AnimMontageTimeline& InTimeline)
 		{
@@ -256,7 +256,7 @@ void FMontageTrack::GetVariantsAtFrame(const Trace::FFrame& InFrame, TArray<TSha
 					const TCHAR* NextSectionName = AnimationProvider->GetName(InMessage.NextSectionNameId);
 					MontageHeader->AddChild(FVariantTreeNode::MakeString(LOCTEXT("NextSectionName", "Next Section"), NextSectionName));
 				}
-				return Trace::EEventEnumerate::Continue;
+				return TraceServices::EEventEnumerate::Continue;
 			});
 		});
 	}
diff --git a/Engine/Plugins/Animation/GameplayInsights/Source/GameplayInsights/Private/MontageTrack.h b/Engine/Plugins/Animation/GameplayInsights/Source/GameplayInsights/Private/MontageTrack.h
--- a/Engine/Plugins/Animation/GameplayInsights/Source/GameplayInsights/Private/MontageTrack.h
+++ b/Engine/Plugins/Animation/GameplayInsights/Source/GameplayInsights/Private/MontageTrack.h
@@ -27,7 +27,7 @@ public:
 	virtual void AddAllSeries() override;
 	virtual bool UpdateSeriesBounds(FGameplayGraphSeries& InSeries, const FTimingTrackViewport& InViewport) override;
 	virtual void UpdateSeries(FGameplayGraphSeries& InSeries, const FTimingTrackViewport& InViewport) override;
-	virtual void GetVariantsAtFrame(const Trace::FFrame& InFrame, TArray<TSharedRef<FVariantTreeNode>>& OutVariants) const override;
+	virtual void GetVariantsAtFrame(const TraceServices::FFrame& InFrame, TArray<TSharedRef<FVariantTreeNode>>& OutVariants) const override;
 
 private:
 	// Helper function used to find a montage message
diff --git a/Engine/Plugins/Animation/GameplayInsights/Source/GameplayInsights/Private/ObjectEventsTrack.cpp b/Engine/Plugins/Animation/GameplayInsights/Source/GameplayInsights/Private/ObjectEventsTrack.cpp
--- a/Engine/Plugins/Animation/GameplayInsights/Source/GameplayInsights/Private/ObjectEventsTrack.cpp
+++ b/Engine/Plugins/Animation/GameplayInsights/Source/GameplayInsights/Private/ObjectEventsTrack.cpp
@@ -31,7 +31,7 @@ void FObjectEventsTrack::BuildDrawState(ITimingEventsTrackDrawStateBuilder& Buil
 
 	if(GameplayProvider)
 	{
-		Trace::FAnalysisSessionReadScope SessionReadScope(SharedData.GetAnalysisSession());
+		TraceServices::FAnalysisSessionReadScope SessionReadScope(SharedData.GetAnalysisSession());
 
 		// object events
 		GameplayProvider->ReadObjectEventsTimeline(GetGameplayTrack().GetObjectId(), [&Context, &Builder](const FGameplayProvider::ObjectEventsTimeline& InTimeline)
@@ -39,7 +39,7 @@ void FObjectEventsTrack::BuildDrawState(ITimingEventsTrackDrawStateBuilder& Buil
 			InTimeline.EnumerateEvents(Context.GetViewport().GetStartTime(), Context.GetViewport().GetEndTime(), [&Builder](double InStartTime, double InEndTime, uint32 InDepth, const FObjectEventMessage& InMessage)
 			{
 				Builder.AddEvent(InStartTime, InEndTime, 0, InMessage.Name);
-				return Trace::EEventEnumerate::Continue;
+				return TraceServices::EEventEnumerate::Continue;
 			});
 		});
 	}
@@ -92,14 +92,14 @@ void FObjectEventsTrack::FindObjectEvent(const FTimingEventSearchParameters& InP
 
 			if(GameplayProvider)
 			{
-				Trace::FAnalysisSessionReadScope SessionReadScope(SharedData.GetAnalysisSession());
+				TraceServices::FAnalysisSessionReadScope SessionReadScope(SharedData.GetAnalysisSession());
 
 				GameplayProvider->ReadObjectEventsTimeline(GetGameplayTrack().GetObjectId(), [&InContext](const FGameplayProvider::ObjectEventsTimeline& InTimeline)
 				{
 					InTimeline.EnumerateEvents(InContext.GetParameters().StartTime, InContext.GetParameters().EndTime, [&InContext](double InEventStartTime, double InEventEndTime, uint32 InDepth, const FObjectEventMessage& InMessage)
 					{
 						InContext.Check(InEventStartTime, InEventEndTime, 0, InMessage);
-						return Trace::EEventEnumerate::Continue;
+						return TraceServices::EEventEnumerate::Continue;
 					});
 				});
 			}
@@ -120,7 +120,7 @@ FText FObjectEventsTrack::MakeTrackName(const FGameplaySharedData& InSharedData,
 	const FGameplayProvider* GameplayProvider = InSharedData.GetAnalysisSession().ReadProvider<FGameplayProvider>(FGameplayProvider::ProviderName);
 	if(GameplayProvider)
 	{
-		Trace::FAnalysisSessionReadScope SessionReadScope(InSharedData.GetAnalysisSession());
+		TraceServices::FAnalysisSessionReadScope SessionReadScope(InSharedData.GetAnalysisSession());
 
 		if(const FObjectInfo* ObjectInfo = GameplayProvider->FindObjectInfo(InObjectID))
 		{
@@ -146,12 +146,12 @@ FText FObjectEventsTrack::MakeTrackName(const FGameplaySharedData& InSharedData,
 	return ClassName;
 }
 
-void FObjectEventsTrack::GetVariantsAtFrame(const Trace::FFrame& InFrame, TArray<TSharedRef<FVariantTreeNode>>& OutVariants) const
+void FObjectEventsTrack::GetVariantsAtFrame(const TraceServices::FFrame& InFrame, TArray<TSharedRef<FVariantTreeNode>>& OutVariants) const
 {
 	const FGameplayProvider* GameplayProvider = SharedData.GetAnalysisSession().ReadProvider<FGameplayProvider>(FGameplayProvider::ProviderName);
 	if(GameplayProvider)
 	{
-		Trace::FAnalysisSessionReadScope SessionReadScope(SharedData.GetAnalysisSession());
+		TraceServices::FAnalysisSessionReadScope SessionReadScope(SharedData.GetAnalysisSession());
 
 		TSharedRef<FVariantTreeNode> Header = OutVariants.Add_GetRef(FVariantTreeNode::MakeHeader(FText::FromString(GetName()), 0));
 
@@ -161,7 +161,7 @@ void FObjectEventsTrack::GetVariantsAtFrame(const Trace::FFrame& InFrame, TArray
 			InTimeline.EnumerateEvents(InFrame.StartTime, InFrame.EndTime, [&Header](double InStartTime, double InEndTime, uint32 InDepth, const FObjectEventMessage& InMessage)
 			{
 				Header->AddChild(FVariantTreeNode::MakeFloat(FText::FromString(InMessage.Name), InStartTime));
-				return Trace::EEventEnumerate::Continue;
+				return TraceServices::EEventEnumerate::Continue;
 			});
 		});
 	}
@@ -178,7 +178,7 @@ void FObjectEventsTrack::BuildContextMenu(FMenuBuilder& MenuBuilder)
 		const FGameplayProvider* GameplayProvider = SharedData.GetAnalysisSession().ReadProvider<FGameplayProvider>(FGameplayProvider::ProviderName);
 		if (GameplayProvider)
 		{
-			Trace::FAnalysisSessionReadScope SessionReadScope(SharedData.GetAnalysisSession());
+			TraceServices::FAnalysisSessionReadScope SessionReadScope(SharedData.GetAnalysisSession());
 
 			if (const FObjectInfo* ObjectInfo = GameplayProvider->FindObjectInfo(GetGameplayTrack().GetObjectId()))
 			{
diff --git a/Engine/Plugins/Animation/GameplayInsights/Source/GameplayInsights/Private/ObjectEventsTrack.h b/Engine/Plugins/Animation/GameplayInsights/Source/GameplayInsights/Private/ObjectEventsTrack.h
--- a/Engine/Plugins/Animation/GameplayInsights/Source/GameplayInsights/Private/ObjectEventsTrack.h
+++ b/Engine/Plugins/Animation/GameplayInsights/Source/GameplayInsights/Private/ObjectEventsTrack.h
@@ -9,7 +9,7 @@ class FGameplaySharedData;
 struct FObjectEventMessage;
 class FTimingEventSearchParameters;
 struct FWorldInfo;
-namespace Trace { struct FFrame; }
+namespace TraceServices { struct FFrame; }
 
 class FObjectEventsTrack : public FGameplayTimingEventsTrack
 {
@@ -22,7 +22,7 @@ public:
 	virtual void Draw(const ITimingTrackDrawContext& Context) const override;
 	virtual void InitTooltip(FTooltipDrawState& Tooltip, const ITimingEvent& HoveredTimingEvent) const override;
 	virtual const TSharedPtr<const ITimingEvent> SearchEvent(const FTimingEventSearchParameters& InSearchParameters) const override;
-	virtual void GetVariantsAtFrame(const Trace::FFrame& InFrame, TArray<TSharedRef<FVariantTreeNode>>& OutVariants) const override;
+	virtual void GetVariantsAtFrame(const TraceServices::FFrame& InFrame, TArray<TSharedRef<FVariantTreeNode>>& OutVariants) const override;
 	virtual void BuildContextMenu(FMenuBuilder& MenuBuilder) override;
 
 private:
@@ -34,4 +34,4 @@ private:
 
 private:
 	const FGameplaySharedData& SharedData;
-};
\ No newline at end of file
+};
diff --git a/Engine/Plugins/Animation/GameplayInsights/Source/GameplayInsights/Private/ObjectPropertiesTrack.cpp b/Engine/Plugins/Animation/GameplayInsights/Source/GameplayInsights/Private/ObjectPropertiesTrack.cpp
--- a/Engine/Plugins/Animation/GameplayInsights/Source/GameplayInsights/Private/ObjectPropertiesTrack.cpp
+++ b/Engine/Plugins/Animation/GameplayInsights/Source/GameplayInsights/Private/ObjectPropertiesTrack.cpp
@@ -28,7 +28,7 @@ void FObjectPropertiesTrack::BuildDrawState(ITimingEventsTrackDrawStateBuilder&
 
 	if(GameplayProvider)
 	{
-		Trace::FAnalysisSessionReadScope SessionReadScope(SharedData.GetAnalysisSession());
+		TraceServices::FAnalysisSessionReadScope SessionReadScope(SharedData.GetAnalysisSession());
 
 		const FText PropertiesText(LOCTEXT("PropertiesEventLabel", "Properties"));
 		const TCHAR* PropertiesString = *PropertiesText.ToString();
@@ -39,7 +39,7 @@ void FObjectPropertiesTrack::BuildDrawState(ITimingEventsTrackDrawStateBuilder&
 			InTimeline.EnumerateEvents(Context.GetViewport().GetStartTime(), Context.GetViewport().GetEndTime(), [&Builder, &PropertiesString](double InStartTime, double InEndTime, uint32 InDepth, const FObjectPropertiesMessage& InMessage)
 			{
 				Builder.AddEvent(InStartTime, InEndTime, 0, PropertiesString);
-				return Trace::EEventEnumerate::Continue;
+				return TraceServices::EEventEnumerate::Continue;
 			});
 		});
 	}
@@ -57,7 +57,7 @@ void FObjectPropertiesTrack::InitTooltip(FTooltipDrawState& Tooltip, const ITimi
 	const FGameplayProvider* GameplayProvider = SharedData.GetAnalysisSession().ReadProvider<FGameplayProvider>(FGameplayProvider::ProviderName);
 	if(GameplayProvider)
 	{
-		Trace::FAnalysisSessionReadScope SessionReadScope(SharedData.GetAnalysisSession());
+		TraceServices::FAnalysisSessionReadScope SessionReadScope(SharedData.GetAnalysisSession());
 
 		FTimingEventSearchParameters SearchParameters(HoveredTimingEvent.GetStartTime(), HoveredTimingEvent.GetEndTime(), ETimingEventSearchFlags::StopAtFirstMatch);
 
@@ -98,14 +98,14 @@ void FObjectPropertiesTrack::FindObjectEvent(const FTimingEventSearchParameters&
 
 			if(GameplayProvider)
 			{
-				Trace::FAnalysisSessionReadScope SessionReadScope(SharedData.GetAnalysisSession());
+				TraceServices::FAnalysisSessionReadScope SessionReadScope(SharedData.GetAnalysisSession());
 
 				GameplayProvider->ReadObjectPropertiesTimeline(GetGameplayTrack().GetObjectId(), [&InContext](const FGameplayProvider::ObjectPropertiesTimeline& InTimeline)
 				{
 					InTimeline.EnumerateEvents(InContext.GetParameters().StartTime, InContext.GetParameters().EndTime, [&InContext](double InEventStartTime, double InEventEndTime, uint32 InDepth, const FObjectPropertiesMessage& InMessage)
 					{
 						InContext.Check(InEventStartTime, InEventEndTime, 0, InMessage);
-						return Trace::EEventEnumerate::Continue;
+						return TraceServices::EEventEnumerate::Continue;
 					});
 				});
 			}
@@ -126,7 +126,7 @@ FText FObjectPropertiesTrack::MakeTrackName(const FGameplaySharedData& InSharedD
 	const FGameplayProvider* GameplayProvider = InSharedData.GetAnalysisSession().ReadProvider<FGameplayProvider>(FGameplayProvider::ProviderName);
 	if(GameplayProvider)
 	{
-		Trace::FAnalysisSessionReadScope SessionReadScope(InSharedData.GetAnalysisSession());
+		TraceServices::FAnalysisSessionReadScope SessionReadScope(InSharedData.GetAnalysisSession());
 
 		FText ObjectName;
 		if (GameplayProvider->IsWorld(InObjectID))
@@ -144,12 +144,12 @@ FText FObjectPropertiesTrack::MakeTrackName(const FGameplaySharedData& InSharedD
 	return ClassName;
 }
 
-void FObjectPropertiesTrack::GetVariantsAtFrame(const Trace::FFrame& InFrame, TArray<TSharedRef<FVariantTreeNode>>& OutVariants) const
+void FObjectPropertiesTrack::GetVariantsAtFrame(const TraceServices::FFrame& InFrame, TArray<TSharedRef<FVariantTreeNode>>& OutVariants) const
 {
 	const FGameplayProvider* GameplayProvider = SharedData.GetAnalysisSession().ReadProvider<FGameplayProvider>(FGameplayProvider::ProviderName);
 	if(GameplayProvider)
 	{
-		Trace::FAnalysisSessionReadScope SessionReadScope(SharedData.GetAnalysisSession());
+		TraceServices::FAnalysisSessionReadScope SessionReadScope(SharedData.GetAnalysisSession());
 
 		const FClassInfo& ClassInfo = GameplayProvider->GetClassInfoFromObject(GetGameplayTrack().GetObjectId());
 
@@ -188,7 +188,7 @@ void FObjectPropertiesTrack::GetVariantsAtFrame(const Trace::FFrame& InFrame, TA
 					{
 						PropertyVariants[InValue.PropertyId]->GetValue().String.Value = InValue.Value;
 					});
-					return Trace::EEventEnumerate::Stop;
+					return TraceServices::EEventEnumerate::Stop;
 				});
 			});
 		}
diff --git a/Engine/Plugins/Animation/GameplayInsights/Source/GameplayInsights/Private/ObjectPropertiesTrack.h b/Engine/Plugins/Animation/GameplayInsights/Source/GameplayInsights/Private/ObjectPropertiesTrack.h
--- a/Engine/Plugins/Animation/GameplayInsights/Source/GameplayInsights/Private/ObjectPropertiesTrack.h
+++ b/Engine/Plugins/Animation/GameplayInsights/Source/GameplayInsights/Private/ObjectPropertiesTrack.h
@@ -9,7 +9,7 @@ class FGameplaySharedData;
 struct FObjectPropertiesMessage;
 class FTimingEventSearchParameters;
 struct FWorldInfo;
-namespace Trace { struct FFrame; }
+namespace TraceServices { struct FFrame; }
 
 class FObjectPropertiesTrack : public FGameplayTimingEventsTrack
 {
@@ -22,7 +22,7 @@ public:
 	virtual void Draw(const ITimingTrackDrawContext& Context) const override;
 	virtual void InitTooltip(FTooltipDrawState& Tooltip, const ITimingEvent& HoveredTimingEvent) const override;
 	virtual const TSharedPtr<const ITimingEvent> SearchEvent(const FTimingEventSearchParameters& InSearchParameters) const override;
-	virtual void GetVariantsAtFrame(const Trace::FFrame& InFrame, TArray<TSharedRef<FVariantTreeNode>>& OutVariants) const override;
+	virtual void GetVariantsAtFrame(const TraceServices::FFrame& InFrame, TArray<TSharedRef<FVariantTreeNode>>& OutVariants) const override;
 
 private:
 	// Helper function used to find an object event
@@ -33,4 +33,4 @@ private:
 
 private:
 	const FGameplaySharedData& SharedData;
-};
\ No newline at end of file
+};
diff --git a/Engine/Plugins/Animation/GameplayInsights/Source/GameplayInsights/Private/SAnimGraphSchematicView.cpp b/Engine/Plugins/Animation/GameplayInsights/Source/GameplayInsights/Private/SAnimGraphSchematicView.cpp
--- a/Engine/Plugins/Animation/GameplayInsights/Source/GameplayInsights/Private/SAnimGraphSchematicView.cpp
+++ b/Engine/Plugins/Animation/GameplayInsights/Source/GameplayInsights/Private/SAnimGraphSchematicView.cpp
@@ -51,7 +51,7 @@ enum class EAnimGraphSchematicFilterState
 class FAnimGraphSchematicPropertyNode
 {
 public:
-	FAnimGraphSchematicPropertyNode(const FText& InName, const TSharedPtr<FAnimNodeValueMessage>& InValue, const Trace::IAnalysisSession& InAnalysisSession)
+	FAnimGraphSchematicPropertyNode(const FText& InName, const TSharedPtr<FAnimNodeValueMessage>& InValue, const TraceServices::IAnalysisSession& InAnalysisSession)
 		: AnalysisSession(InAnalysisSession)
 		, Name(InName)
 		, Value(InValue)
@@ -70,7 +70,7 @@ public:
 		}
 	}
 
-	static TSharedRef<SWidget> StaticMakeValueWidget(const Trace::IAnalysisSession& InAnalysisSession, const TSharedRef<FAnimNodeValueMessage>& InValue) 
+	static TSharedRef<SWidget> StaticMakeValueWidget(const TraceServices::IAnalysisSession& InAnalysisSession, const TSharedRef<FAnimNodeValueMessage>& InValue) 
 	{ 
 		switch(InValue->Value.Type)
 		{
@@ -179,7 +179,7 @@ public:
 			const FGameplayProvider* GameplayProvider = InAnalysisSession.ReadProvider<FGameplayProvider>(FGameplayProvider::ProviderName);
 			if(GameplayProvider)
 			{
-				Trace::FAnalysisSessionReadScope SessionReadScope(InAnalysisSession);
+				TraceServices::FAnalysisSessionReadScope SessionReadScope(InAnalysisSession);
 
 				const FObjectInfo& ObjectInfo = GameplayProvider->GetObjectInfo(InValue->Value.Object.Value);
 #if WITH_EDITOR
@@ -208,7 +208,7 @@ public:
 			const FGameplayProvider* GameplayProvider = InAnalysisSession.ReadProvider<FGameplayProvider>(FGameplayProvider::ProviderName);
 			if(GameplayProvider)
 			{
-				Trace::FAnalysisSessionReadScope SessionReadScope(InAnalysisSession);
+				TraceServices::FAnalysisSessionReadScope SessionReadScope(InAnalysisSession);
 
 				const FClassInfo& ClassInfo = GameplayProvider->GetClassInfo(InValue->Value.Class.Value);
 #if WITH_EDITOR
@@ -236,7 +236,7 @@ public:
 		return SNullWidget::NullWidget; 
 	}
 
-	const Trace::IAnalysisSession& AnalysisSession;
+	const TraceServices::IAnalysisSession& AnalysisSession;
 
 	FText Name;
 
@@ -318,7 +318,7 @@ public:
 class FAnimGraphSchematicNode : public TSharedFromThis<FAnimGraphSchematicNode>
 {
 public:
-	FAnimGraphSchematicNode(int32 InNodeId, const FText& InType, const Trace::IAnalysisSession& InAnalysisSession)
+	FAnimGraphSchematicNode(int32 InNodeId, const FText& InType, const TraceServices::IAnalysisSession& InAnalysisSession)
 		: AnalysisSession(InAnalysisSession)
 		, NodeId(InNodeId)
 		, Type(InType)
@@ -341,7 +341,7 @@ public:
 		}
 	}
 
-	const Trace::IAnalysisSession& AnalysisSession;
+	const TraceServices::IAnalysisSession& AnalysisSession;
 
 	int32 NodeId;
 
@@ -444,7 +444,7 @@ class SAnimGraphSchematicNode : public SMultiColumnTableRow<TSharedRef<FAnimGrap
 	TAttribute<FText> FilterText;
 };
 
-void SAnimGraphSchematicView::Construct(const FArguments& InArgs, uint64 InAnimInstanceId, Insights::ITimingViewSession& InTimingViewSession, const Trace::IAnalysisSession& InAnalysisSession)
+void SAnimGraphSchematicView::Construct(const FArguments& InArgs, uint64 InAnimInstanceId, Insights::ITimingViewSession& InTimingViewSession, const TraceServices::IAnalysisSession& InAnalysisSession)
 {
 	AnimInstanceId = InAnimInstanceId;
 	TimingViewSession = &InTimingViewSession;
@@ -597,13 +597,13 @@ void SAnimGraphSchematicView::RefreshNodes()
 
 	if(AnimationProvider && GameplayProvider)
 	{
-		Trace::FAnalysisSessionReadScope SessionReadScope(*AnalysisSession);
+		TraceServices::FAnalysisSessionReadScope SessionReadScope(*AnalysisSession);
 
 		AnimationProvider->ReadAnimGraphTimeline(AnimInstanceId, [this, &NodeMap, AnimationProvider, GameplayProvider](const FAnimationProvider::AnimGraphTimeline& InGraphTimeline)
 		{
-			const Trace::IFrameProvider& FramesProvider = Trace::ReadFrameProvider(*AnalysisSession);
+			const TraceServices::IFrameProvider& FramesProvider = TraceServices::ReadFrameProvider(*AnalysisSession);
 
-			Trace::FFrame Frame;
+			TraceServices::FFrame Frame;
 			if(FramesProvider.GetFrameFromTime(ETraceFrameType::TraceFrameType_Game, TimeMarker, Frame))
 			{
 				InGraphTimeline.EnumerateEvents(Frame.StartTime, Frame.EndTime, [this, &NodeMap, AnimationProvider, GameplayProvider](double InGraphStartTime, double InGraphEndTime, uint32 InDepth, const FAnimGraphMessage& InMessage)
@@ -663,7 +663,7 @@ void SAnimGraphSchematicView::RefreshNodes()
 										}
 									}
 								}
-								return Trace::EEventEnumerate::Continue;
+								return TraceServices::EEventEnumerate::Continue;
 							});
 						});
 
@@ -719,10 +719,10 @@ void SAnimGraphSchematicView::RefreshNodes()
 								ExistingNode->KeysAndValues.Add(Key, SharedMessage);
 								ExistingNode->Values.Add(SharedMessage);
 							}
-							return Trace::EEventEnumerate::Continue;
+							return TraceServices::EEventEnumerate::Continue;
 						});
 					});
-					return Trace::EEventEnumerate::Continue;
+					return TraceServices::EEventEnumerate::Continue;
 				});
 			};
 		});
diff --git a/Engine/Plugins/Animation/GameplayInsights/Source/GameplayInsights/Private/SAnimGraphSchematicView.h b/Engine/Plugins/Animation/GameplayInsights/Source/GameplayInsights/Private/SAnimGraphSchematicView.h
--- a/Engine/Plugins/Animation/GameplayInsights/Source/GameplayInsights/Private/SAnimGraphSchematicView.h
+++ b/Engine/Plugins/Animation/GameplayInsights/Source/GameplayInsights/Private/SAnimGraphSchematicView.h
@@ -8,7 +8,7 @@
 
 namespace Insights { class ITimingViewSession; }
 namespace Insights { enum class ETimeChangedFlags; }
-namespace Trace { class IAnalysisSession; }
+namespace TraceServices { class IAnalysisSession; }
 class FAnimationSharedData;
 class IInsightsManager;
 class FAnimGraphSchematicNode;
@@ -26,7 +26,7 @@ public:
 
 	SLATE_END_ARGS()
 
-	void Construct(const FArguments& InArgs, uint64 InAnimInstanceId, Insights::ITimingViewSession& InTimingViewSession, const Trace::IAnalysisSession& InAnalysisSession);
+	void Construct(const FArguments& InArgs, uint64 InAnimInstanceId, Insights::ITimingViewSession& InTimingViewSession, const TraceServices::IAnalysisSession& InAnalysisSession);
 
 	uint64 GetAnimInstanceId() const { return AnimInstanceId; }
 
@@ -71,7 +71,7 @@ private:
 	FSlateColor GetViewButtonForegroundColor() const;
 
 private:
-	const Trace::IAnalysisSession* AnalysisSession;
+	const TraceServices::IAnalysisSession* AnalysisSession;
 
 	Insights::ITimingViewSession* TimingViewSession;
 
