Moved the TraceServices module into its own namespace.

diff --git a/Engine/Plugins/Animation/GameplayInsights/Source/GameplayInsights/Private/SGameplayInsightsTransportControls.cpp b/Engine/Plugins/Animation/GameplayInsights/Source/GameplayInsights/Private/SGameplayInsightsTransportControls.cpp
--- a/Engine/Plugins/Animation/GameplayInsights/Source/GameplayInsights/Private/SGameplayInsightsTransportControls.cpp
+++ b/Engine/Plugins/Animation/GameplayInsights/Source/GameplayInsights/Private/SGameplayInsightsTransportControls.cpp
@@ -138,25 +138,25 @@ void SGameplayInsightsTransportControls::Construct(const FArguments& InArgs, FGa
 
 FReply SGameplayInsightsTransportControls::OnClick_Forward_Step()
 {
-	Trace::FAnalysisSessionReadScope SessionReadScope(SharedData->GetAnalysisSession());
-	const Trace::IFrameProvider& FramesProvider = Trace::ReadFrameProvider(SharedData->GetAnalysisSession());
+	TraceServices::FAnalysisSessionReadScope SessionReadScope(SharedData->GetAnalysisSession());
+	const TraceServices::IFrameProvider& FramesProvider = TraceServices::ReadFrameProvider(SharedData->GetAnalysisSession());
 
 	double CurrentTime = SharedData->GetTimingViewSession().GetTimeMarker();
 	if(CurrentTime == std::numeric_limits<double>::infinity())
 	{
-		const Trace::FFrame* FirstFrame = FramesProvider.GetFrame(ETraceFrameType::TraceFrameType_Game, 0);
+		const TraceServices::FFrame* FirstFrame = FramesProvider.GetFrame(ETraceFrameType::TraceFrameType_Game, 0);
 		if(FirstFrame)
 		{
 			CurrentTime = FirstFrame->StartTime + (double)SMALL_NUMBER;
 		}
 	}
 	
-	Trace::FFrame Frame;
+	TraceServices::FFrame Frame;
 	if(FramesProvider.GetFrameFromTime(ETraceFrameType::TraceFrameType_Game, CurrentTime, Frame))
 	{
 		if(Frame.Index < FramesProvider.GetFrameCount(ETraceFrameType::TraceFrameType_Game) - 1)
 		{
-			const Trace::FFrame* NextFrame = FramesProvider.GetFrame(ETraceFrameType::TraceFrameType_Game, Frame.Index + 1);
+			const TraceServices::FFrame* NextFrame = FramesProvider.GetFrame(ETraceFrameType::TraceFrameType_Game, Frame.Index + 1);
 			if(NextFrame)
 			{
 				SetTimeMarker(NextFrame->StartTime + (double)SMALL_NUMBER, false);
@@ -172,12 +172,12 @@ FReply SGameplayInsightsTransportControls::OnClick_Forward_Step()
 
 FReply SGameplayInsightsTransportControls::OnClick_Forward_End()
 {
-	Trace::FAnalysisSessionReadScope SessionReadScope(SharedData->GetAnalysisSession());
-	const Trace::IFrameProvider& FramesProvider = Trace::ReadFrameProvider(SharedData->GetAnalysisSession());
+	TraceServices::FAnalysisSessionReadScope SessionReadScope(SharedData->GetAnalysisSession());
+	const TraceServices::IFrameProvider& FramesProvider = TraceServices::ReadFrameProvider(SharedData->GetAnalysisSession());
 	
 	if(FramesProvider.GetFrameCount(ETraceFrameType::TraceFrameType_Game) > 0)
 	{
-		const Trace::FFrame* LastFrame = FramesProvider.GetFrame(ETraceFrameType::TraceFrameType_Game, FramesProvider.GetFrameCount(ETraceFrameType::TraceFrameType_Game) - 1);
+		const TraceServices::FFrame* LastFrame = FramesProvider.GetFrame(ETraceFrameType::TraceFrameType_Game, FramesProvider.GetFrameCount(ETraceFrameType::TraceFrameType_Game) - 1);
 		if(LastFrame)
 		{
 			SetTimeMarker(LastFrame->StartTime + (double)SMALL_NUMBER, true);
@@ -192,25 +192,25 @@ FReply SGameplayInsightsTransportControls::OnClick_Forward_End()
 
 FReply SGameplayInsightsTransportControls::OnClick_Backward_Step()
 {
-	Trace::FAnalysisSessionReadScope SessionReadScope(SharedData->GetAnalysisSession());
-	const Trace::IFrameProvider& FramesProvider = Trace::ReadFrameProvider(SharedData->GetAnalysisSession());
+	TraceServices::FAnalysisSessionReadScope SessionReadScope(SharedData->GetAnalysisSession());
+	const TraceServices::IFrameProvider& FramesProvider = TraceServices::ReadFrameProvider(SharedData->GetAnalysisSession());
 
 	double CurrentTime = SharedData->GetTimingViewSession().GetTimeMarker();
 	if(CurrentTime == std::numeric_limits<double>::infinity())
 	{
-		const Trace::FFrame* LastFrame = FramesProvider.GetFrame(ETraceFrameType::TraceFrameType_Game, FramesProvider.GetFrameCount(TraceFrameType_Game) - 1);
+		const TraceServices::FFrame* LastFrame = FramesProvider.GetFrame(ETraceFrameType::TraceFrameType_Game, FramesProvider.GetFrameCount(TraceFrameType_Game) - 1);
 		if(LastFrame)
 		{
 			CurrentTime = LastFrame->StartTime + (double)SMALL_NUMBER;
 		}
 	}
 
-	Trace::FFrame Frame;
+	TraceServices::FFrame Frame;
 	if(FramesProvider.GetFrameFromTime(ETraceFrameType::TraceFrameType_Game, CurrentTime, Frame))
 	{
 		if(Frame.Index > 0)
 		{
-			const Trace::FFrame* PrevFrame = FramesProvider.GetFrame(ETraceFrameType::TraceFrameType_Game, Frame.Index - 1);
+			const TraceServices::FFrame* PrevFrame = FramesProvider.GetFrame(ETraceFrameType::TraceFrameType_Game, Frame.Index - 1);
 			if(PrevFrame)
 			{
 				SetTimeMarker(PrevFrame->StartTime + (double)SMALL_NUMBER, false);
@@ -226,10 +226,10 @@ FReply SGameplayInsightsTransportControls::OnClick_Backward_Step()
 
 FReply SGameplayInsightsTransportControls::OnClick_Backward_End()
 {
-	Trace::FAnalysisSessionReadScope SessionReadScope(SharedData->GetAnalysisSession());
-	const Trace::IFrameProvider& FramesProvider = Trace::ReadFrameProvider(SharedData->GetAnalysisSession());
+	TraceServices::FAnalysisSessionReadScope SessionReadScope(SharedData->GetAnalysisSession());
+	const TraceServices::IFrameProvider& FramesProvider = TraceServices::ReadFrameProvider(SharedData->GetAnalysisSession());
 	
-	const Trace::FFrame* FirstFrame = FramesProvider.GetFrame(ETraceFrameType::TraceFrameType_Game, 0);
+	const TraceServices::FFrame* FirstFrame = FramesProvider.GetFrame(ETraceFrameType::TraceFrameType_Game, 0);
 	if(FirstFrame)
 	{
 		SetTimeMarker(FirstFrame->StartTime + (double)SMALL_NUMBER, true);
@@ -246,10 +246,10 @@ FReply SGameplayInsightsTransportControls::OnClick_Forward()
 	double CurrentTime = SharedData->GetTimingViewSession().GetTimeMarker();
 	if(CurrentTime == std::numeric_limits<double>::infinity())
 	{
-		Trace::FAnalysisSessionReadScope SessionReadScope(SharedData->GetAnalysisSession());
-		const Trace::IFrameProvider& FramesProvider = Trace::ReadFrameProvider(SharedData->GetAnalysisSession());
+		TraceServices::FAnalysisSessionReadScope SessionReadScope(SharedData->GetAnalysisSession());
+		const TraceServices::IFrameProvider& FramesProvider = TraceServices::ReadFrameProvider(SharedData->GetAnalysisSession());
 	
-		const Trace::FFrame* FirstFrame = FramesProvider.GetFrame(ETraceFrameType::TraceFrameType_Game, 0);
+		const TraceServices::FFrame* FirstFrame = FramesProvider.GetFrame(ETraceFrameType::TraceFrameType_Game, 0);
 		if(FirstFrame)
 		{
 			SetTimeMarker(FirstFrame->StartTime + (double)SMALL_NUMBER, false);
@@ -266,10 +266,10 @@ FReply SGameplayInsightsTransportControls::OnClick_Backward()
 	double CurrentTime = SharedData->GetTimingViewSession().GetTimeMarker();
 	if(CurrentTime == std::numeric_limits<double>::infinity())
 	{
-		Trace::FAnalysisSessionReadScope SessionReadScope(SharedData->GetAnalysisSession());
-		const Trace::IFrameProvider& FramesProvider = Trace::ReadFrameProvider(SharedData->GetAnalysisSession());
+		TraceServices::FAnalysisSessionReadScope SessionReadScope(SharedData->GetAnalysisSession());
+		const TraceServices::IFrameProvider& FramesProvider = TraceServices::ReadFrameProvider(SharedData->GetAnalysisSession());
 	
-		const Trace::FFrame* LastFrame = FramesProvider.GetFrame(ETraceFrameType::TraceFrameType_Game, FramesProvider.GetFrameCount(TraceFrameType_Game) - 1);
+		const TraceServices::FFrame* LastFrame = FramesProvider.GetFrame(ETraceFrameType::TraceFrameType_Game, FramesProvider.GetFrameCount(TraceFrameType_Game) - 1);
 		if(LastFrame)
 		{
 			SetTimeMarker(LastFrame->StartTime + (double)SMALL_NUMBER, false);
@@ -338,4 +338,4 @@ void SGameplayInsightsTransportControls::HandleSelectionRangeChanged(Insights::E
 
 #undef LOCTEXT_NAMESPACE
 
-#endif
\ No newline at end of file
+#endif
diff --git a/Engine/Plugins/Animation/GameplayInsights/Source/GameplayInsights/Private/STrackVariantValueView.cpp b/Engine/Plugins/Animation/GameplayInsights/Source/GameplayInsights/Private/STrackVariantValueView.cpp
--- a/Engine/Plugins/Animation/GameplayInsights/Source/GameplayInsights/Private/STrackVariantValueView.cpp
+++ b/Engine/Plugins/Animation/GameplayInsights/Source/GameplayInsights/Private/STrackVariantValueView.cpp
@@ -7,7 +7,7 @@
 #include "GameplayGraphTrack.h"
 #include "Insights/ITimingViewSession.h"
 
-void STrackVariantValueView::Construct(const FArguments& InArgs, const TSharedRef<FBaseTimingTrack>& InTimingTrack, Insights::ITimingViewSession& InTimingViewSession, const Trace::IAnalysisSession& InAnalysisSession)
+void STrackVariantValueView::Construct(const FArguments& InArgs, const TSharedRef<FBaseTimingTrack>& InTimingTrack, Insights::ITimingViewSession& InTimingViewSession, const TraceServices::IAnalysisSession& InAnalysisSession)
 {
 	TimingTrack = InTimingTrack;
 	AnalysisSession = &InAnalysisSession;
@@ -17,7 +17,7 @@ void STrackVariantValueView::Construct(const FArguments& InArgs, const TSharedRe
 	ChildSlot
 	[
 		SAssignNew(VariantValueView, SVariantValueView, InAnalysisSession)
-		.OnGetVariantValues_Lambda([this](const Trace::FFrame& InFrame, TArray<TSharedRef<FVariantTreeNode>>& OutValues)
+		.OnGetVariantValues_Lambda([this](const TraceServices::FFrame& InFrame, TArray<TSharedRef<FVariantTreeNode>>& OutValues)
 		{
 			TSharedPtr<FBaseTimingTrack> PinnedTrack = TimingTrack.Pin();
 			if(PinnedTrack.IsValid())
@@ -34,10 +34,10 @@ void STrackVariantValueView::Construct(const FArguments& InArgs, const TSharedRe
 		})
 	];
 
-	Trace::FAnalysisSessionReadScope SessionReadScope(InAnalysisSession);
+	TraceServices::FAnalysisSessionReadScope SessionReadScope(InAnalysisSession);
 
-	const Trace::IFrameProvider& FramesProvider = Trace::ReadFrameProvider(InAnalysisSession);
-	Trace::FFrame MarkerFrame;
+	const TraceServices::IFrameProvider& FramesProvider = TraceServices::ReadFrameProvider(InAnalysisSession);
+	TraceServices::FFrame MarkerFrame;
 	if(FramesProvider.GetFrameFromTime(ETraceFrameType::TraceFrameType_Game, InTimingViewSession.GetTimeMarker(), MarkerFrame))
 	{
 		VariantValueView->RequestRefresh(MarkerFrame);
@@ -46,10 +46,10 @@ void STrackVariantValueView::Construct(const FArguments& InArgs, const TSharedRe
 
 void STrackVariantValueView::HandleTimeMarkerChanged(Insights::ETimeChangedFlags InFlags, double InTimeMarker)
 {
-	Trace::FAnalysisSessionReadScope SessionReadScope(*AnalysisSession);
+	TraceServices::FAnalysisSessionReadScope SessionReadScope(*AnalysisSession);
 
-	const Trace::IFrameProvider& FramesProvider = Trace::ReadFrameProvider(*AnalysisSession);
-	Trace::FFrame MarkerFrame;
+	const TraceServices::IFrameProvider& FramesProvider = TraceServices::ReadFrameProvider(*AnalysisSession);
+	TraceServices::FFrame MarkerFrame;
 	if(FramesProvider.GetFrameFromTime(ETraceFrameType::TraceFrameType_Game, InTimeMarker, MarkerFrame))
 	{
 		VariantValueView->RequestRefresh(MarkerFrame);
diff --git a/Engine/Plugins/Animation/GameplayInsights/Source/GameplayInsights/Private/STrackVariantValueView.h b/Engine/Plugins/Animation/GameplayInsights/Source/GameplayInsights/Private/STrackVariantValueView.h
--- a/Engine/Plugins/Animation/GameplayInsights/Source/GameplayInsights/Private/STrackVariantValueView.h
+++ b/Engine/Plugins/Animation/GameplayInsights/Source/GameplayInsights/Private/STrackVariantValueView.h
@@ -8,7 +8,7 @@
 
 class FBaseTimingTrack;
 namespace Insights { class ITimingViewSession; }
-namespace Trace { class IAnalysisSession; }
+namespace TraceServices { class IAnalysisSession; }
 namespace Insights { enum class ETimeChangedFlags; }
 class SVariantValueView;
 
@@ -20,7 +20,7 @@ public:
 
 	SLATE_END_ARGS()
 
-	void Construct(const FArguments& InArgs, const TSharedRef<FBaseTimingTrack>& InTimingTrack, Insights::ITimingViewSession& InTimingViewSession, const Trace::IAnalysisSession& InAnalysisSession);
+	void Construct(const FArguments& InArgs, const TSharedRef<FBaseTimingTrack>& InTimingTrack, Insights::ITimingViewSession& InTimingViewSession, const TraceServices::IAnalysisSession& InAnalysisSession);
 
 	TSharedPtr<FBaseTimingTrack> GetTimingTrack() const { return TimingTrack.Pin(); }
 
@@ -33,5 +33,5 @@ private:
 
 	TSharedPtr<SVariantValueView> VariantValueView;
 
-	const Trace::IAnalysisSession* AnalysisSession;
-};
\ No newline at end of file
+	const TraceServices::IAnalysisSession* AnalysisSession;
+};
diff --git a/Engine/Plugins/Animation/GameplayInsights/Source/GameplayInsights/Private/SVariantValueView.cpp b/Engine/Plugins/Animation/GameplayInsights/Source/GameplayInsights/Private/SVariantValueView.cpp
--- a/Engine/Plugins/Animation/GameplayInsights/Source/GameplayInsights/Private/SVariantValueView.cpp
+++ b/Engine/Plugins/Animation/GameplayInsights/Source/GameplayInsights/Private/SVariantValueView.cpp
@@ -34,7 +34,7 @@ namespace VariantColumns
 	static const FName Value("Value");
 };
 
-static TSharedRef<SWidget> MakeVariantValueWidget(const Trace::IAnalysisSession& InAnalysisSession, const FVariantValue& InValue, const TAttribute<FText>& InHighlightText) 
+static TSharedRef<SWidget> MakeVariantValueWidget(const TraceServices::IAnalysisSession& InAnalysisSession, const FVariantValue& InValue, const TAttribute<FText>& InHighlightText) 
 { 
 	switch(InValue.Type)
 	{
@@ -149,7 +149,7 @@ static TSharedRef<SWidget> MakeVariantValueWidget(const Trace::IAnalysisSession&
 		const FGameplayProvider* GameplayProvider = InAnalysisSession.ReadProvider<FGameplayProvider>(FGameplayProvider::ProviderName);
 		if(GameplayProvider)
 		{
-			Trace::FAnalysisSessionReadScope SessionReadScope(InAnalysisSession);
+			TraceServices::FAnalysisSessionReadScope SessionReadScope(InAnalysisSession);
 
 			const FObjectInfo& ObjectInfo = GameplayProvider->GetObjectInfo(InValue.Object.Value);
 #if WITH_EDITOR
@@ -179,7 +179,7 @@ static TSharedRef<SWidget> MakeVariantValueWidget(const Trace::IAnalysisSession&
 		const FGameplayProvider* GameplayProvider = InAnalysisSession.ReadProvider<FGameplayProvider>(FGameplayProvider::ProviderName);
 		if(GameplayProvider)
 		{
-			Trace::FAnalysisSessionReadScope SessionReadScope(InAnalysisSession);
+			TraceServices::FAnalysisSessionReadScope SessionReadScope(InAnalysisSession);
 
 			const FClassInfo& ClassInfo = GameplayProvider->GetClassInfo(InValue.Class.Value);
 #if WITH_EDITOR
@@ -218,7 +218,7 @@ public:
 
 	SLATE_END_ARGS()
 
-	void Construct(const FArguments& InArgs, const TSharedRef<STableViewBase>& InOwnerTable, TSharedRef<FVariantTreeNode> InNode, const Trace::IAnalysisSession& InAnalysisSession)
+	void Construct(const FArguments& InArgs, const TSharedRef<STableViewBase>& InOwnerTable, TSharedRef<FVariantTreeNode> InNode, const TraceServices::IAnalysisSession& InAnalysisSession)
 	{
 		Node = InNode;
 		AnalysisSession = &InAnalysisSession;
@@ -281,12 +281,12 @@ public:
 		return SNullWidget::NullWidget;
 	}
 
-	const Trace::IAnalysisSession* AnalysisSession;
+	const TraceServices::IAnalysisSession* AnalysisSession;
 	TSharedPtr<FVariantTreeNode> Node;
 	TAttribute<FText> HighlightText;
 };
 
-void SVariantValueView::Construct(const FArguments& InArgs, const Trace::IAnalysisSession& InAnalysisSession)
+void SVariantValueView::Construct(const FArguments& InArgs, const TraceServices::IAnalysisSession& InAnalysisSession)
 {
 	OnGetVariantValues = InArgs._OnGetVariantValues;
 
diff --git a/Engine/Plugins/Animation/GameplayInsights/Source/GameplayInsights/Private/SVariantValueView.h b/Engine/Plugins/Animation/GameplayInsights/Source/GameplayInsights/Private/SVariantValueView.h
--- a/Engine/Plugins/Animation/GameplayInsights/Source/GameplayInsights/Private/SVariantValueView.h
+++ b/Engine/Plugins/Animation/GameplayInsights/Source/GameplayInsights/Private/SVariantValueView.h
@@ -8,14 +8,14 @@
 #include "TraceServices/Model/Frames.h"
 #include "Containers/StringFwd.h"
 
-namespace Trace { class IAnalysisSession; }
+namespace TraceServices { class IAnalysisSession; }
 struct FVariantTreeNode;
 class FTextFilterExpressionEvaluator;
 enum class EVariantTreeNodeFilterState;
 class FUICommandList;
 
 // Delegate called to get variant values to display
-DECLARE_DELEGATE_TwoParams(FOnGetVariantValues, const Trace::FFrame& /*InFrame*/, TArray<TSharedRef<FVariantTreeNode>>& /*OutValues*/);
+DECLARE_DELEGATE_TwoParams(FOnGetVariantValues, const TraceServices::FFrame& /*InFrame*/, TArray<TSharedRef<FVariantTreeNode>>& /*OutValues*/);
 
 class SVariantValueView : public SCompoundWidget
 {
@@ -26,10 +26,10 @@ public:
 
 	SLATE_END_ARGS()
 
-	void Construct(const FArguments& InArgs, const Trace::IAnalysisSession& InAnalysisSession);
+	void Construct(const FArguments& InArgs, const TraceServices::IAnalysisSession& InAnalysisSession);
 
 	/** Refresh the displayed variants. */
-	void RequestRefresh(const Trace::FFrame& InFrame) { Frame = InFrame; bNeedsRefresh = true; }
+	void RequestRefresh(const TraceServices::FFrame& InFrame) { Frame = InFrame; bNeedsRefresh = true; }
 
 private:
 	// SWidget interface
@@ -69,7 +69,7 @@ private:
 	bool CanCopy() const;
 
 private:
-	const Trace::IAnalysisSession* AnalysisSession;
+	const TraceServices::IAnalysisSession* AnalysisSession;
 
 	TSharedPtr<FTextFilterExpressionEvaluator> TextFilter;
 
@@ -81,7 +81,7 @@ private:
 
 	TSharedPtr<FUICommandList> CommandList;
 
-	Trace::FFrame Frame;
+	TraceServices::FFrame Frame;
 
 	FOnGetVariantValues OnGetVariantValues;
 
diff --git a/Engine/Plugins/Animation/GameplayInsights/Source/GameplayInsights/Private/SkeletalMeshCurvesTrack.cpp b/Engine/Plugins/Animation/GameplayInsights/Source/GameplayInsights/Private/SkeletalMeshCurvesTrack.cpp
--- a/Engine/Plugins/Animation/GameplayInsights/Source/GameplayInsights/Private/SkeletalMeshCurvesTrack.cpp
+++ b/Engine/Plugins/Animation/GameplayInsights/Source/GameplayInsights/Private/SkeletalMeshCurvesTrack.cpp
@@ -31,7 +31,7 @@ void FSkeletalMeshCurvesTrack::AddAllSeries()
 	bool bFirstSeries = AllSeries.Num() == 0;
 	if(AnimationProvider)
 	{
-		Trace::FAnalysisSessionReadScope SessionReadScope(SharedData.GetAnalysisSession());
+		TraceServices::FAnalysisSessionReadScope SessionReadScope(SharedData.GetAnalysisSession());
 
 		AnimationProvider->EnumerateSkeletalMeshCurveIds(GetGameplayTrack().GetObjectId(), [this, &bFirstSeries, &AnimationProvider](uint32 InCurveId)
 		{
@@ -77,7 +77,7 @@ bool FSkeletalMeshCurvesTrack::UpdateSeriesBounds(FGameplayGraphSeries& InSeries
 
 	if(AnimationProvider)
 	{
-		Trace::FAnalysisSessionReadScope SessionReadScope(SharedData.GetAnalysisSession());
+		TraceServices::FAnalysisSessionReadScope SessionReadScope(SharedData.GetAnalysisSession());
 
 		FSkeletalMeshCurveSeries& CurveSeries = *static_cast<FSkeletalMeshCurveSeries*>(&InSeries);
 
@@ -97,7 +97,7 @@ bool FSkeletalMeshCurvesTrack::UpdateSeriesBounds(FGameplayGraphSeries& InSeries
 						bFoundEvents = true;
 					}
 				});
-				return Trace::EEventEnumerate::Continue;
+				return TraceServices::EEventEnumerate::Continue;
 			});
 		});
 	}
@@ -111,7 +111,7 @@ void FSkeletalMeshCurvesTrack::UpdateSeries(FGameplayGraphSeries& InSeries, cons
 
 	if(AnimationProvider)
 	{
-		Trace::FAnalysisSessionReadScope SessionReadScope(SharedData.GetAnalysisSession());
+		TraceServices::FAnalysisSessionReadScope SessionReadScope(SharedData.GetAnalysisSession());
 
 		FSkeletalMeshCurveSeries& CurveSeries = *static_cast<FSkeletalMeshCurveSeries*>(&InSeries);
 
@@ -135,7 +135,7 @@ void FSkeletalMeshCurvesTrack::UpdateSeries(FGameplayGraphSeries& InSeries, cons
 						LastFrameWithCurve = FrameCounter;
 					}
 				});
-				return Trace::EEventEnumerate::Continue;
+				return TraceServices::EEventEnumerate::Continue;
 			});
 		});
 	}
@@ -154,7 +154,7 @@ void FSkeletalMeshCurvesTrack::InitTooltip(FTooltipDrawState& Tooltip, const ITi
 		Tooltip.AddTitle(GetName());
 
 		{
-			Trace::FAnalysisSessionReadScope SessionReadScope(SharedData.GetAnalysisSession());
+			TraceServices::FAnalysisSessionReadScope SessionReadScope(SharedData.GetAnalysisSession());
 
 			const FAnimationProvider* AnimationProvider = SharedData.GetAnalysisSession().ReadProvider<FAnimationProvider>(FAnimationProvider::ProviderName);
 			if(AnimationProvider)
@@ -195,14 +195,14 @@ void FSkeletalMeshCurvesTrack::FindSkeletalMeshPoseMessage(const FTimingEventSea
 
 			if(AnimationProvider)
 			{
-				Trace::FAnalysisSessionReadScope SessionReadScope(SharedData.GetAnalysisSession());
+				TraceServices::FAnalysisSessionReadScope SessionReadScope(SharedData.GetAnalysisSession());
 
 				AnimationProvider->ReadSkeletalMeshPoseTimeline(GetGameplayTrack().GetObjectId(), [&InContext](const FAnimationProvider::SkeletalMeshPoseTimeline& InTimeline, bool bInHasCurves)
 				{
 					InTimeline.EnumerateEvents(InContext.GetParameters().StartTime, InContext.GetParameters().EndTime, [&InContext](double InEventStartTime, double InEventEndTime, uint32 InDepth, const FSkeletalMeshPoseMessage& InMessage)
 					{
 						InContext.Check(InEventStartTime, InEventEndTime, 0, InMessage);
-						return Trace::EEventEnumerate::Continue;
+						return TraceServices::EEventEnumerate::Continue;
 					});
 				});
 			}
@@ -214,13 +214,13 @@ void FSkeletalMeshCurvesTrack::FindSkeletalMeshPoseMessage(const FTimingEventSea
 		});
 }
 
-void FSkeletalMeshCurvesTrack::GetVariantsAtFrame(const Trace::FFrame& InFrame, TArray<TSharedRef<FVariantTreeNode>>& OutVariants) const
+void FSkeletalMeshCurvesTrack::GetVariantsAtFrame(const TraceServices::FFrame& InFrame, TArray<TSharedRef<FVariantTreeNode>>& OutVariants) const
 {
 	const FGameplayProvider* GameplayProvider = SharedData.GetAnalysisSession().ReadProvider<FGameplayProvider>(FGameplayProvider::ProviderName);
 	const FAnimationProvider* AnimationProvider = SharedData.GetAnalysisSession().ReadProvider<FAnimationProvider>(FAnimationProvider::ProviderName);
 	if(GameplayProvider && AnimationProvider)
 	{
-		Trace::FAnalysisSessionReadScope SessionReadScope(SharedData.GetAnalysisSession());
+		TraceServices::FAnalysisSessionReadScope SessionReadScope(SharedData.GetAnalysisSession());
 
 		TSharedRef<FVariantTreeNode> Header = OutVariants.Add_GetRef(FVariantTreeNode::MakeHeader(FText::FromString(GetName()), 0));
 
@@ -236,7 +236,7 @@ void FSkeletalMeshCurvesTrack::GetVariantsAtFrame(const Trace::FFrame& InFrame,
 						Header->AddChild(FVariantTreeNode::MakeFloat(FText::FromString(CurveName), InCurve.Value));
 					});
 				}
-				return Trace::EEventEnumerate::Continue;
+				return TraceServices::EEventEnumerate::Continue;
 			});
 		});
 	}
diff --git a/Engine/Plugins/Animation/GameplayInsights/Source/GameplayInsights/Private/SkeletalMeshCurvesTrack.h b/Engine/Plugins/Animation/GameplayInsights/Source/GameplayInsights/Private/SkeletalMeshCurvesTrack.h
--- a/Engine/Plugins/Animation/GameplayInsights/Source/GameplayInsights/Private/SkeletalMeshCurvesTrack.h
+++ b/Engine/Plugins/Animation/GameplayInsights/Source/GameplayInsights/Private/SkeletalMeshCurvesTrack.h
@@ -27,7 +27,7 @@ public:
 	virtual void AddAllSeries() override;
 	virtual bool UpdateSeriesBounds(FGameplayGraphSeries& InSeries, const FTimingTrackViewport& InViewport) override;
 	virtual void UpdateSeries(FGameplayGraphSeries& InSeries, const FTimingTrackViewport& InViewport) override;
-	virtual void GetVariantsAtFrame(const Trace::FFrame& InFrame, TArray<TSharedRef<FVariantTreeNode>>& OutVariants) const override;
+	virtual void GetVariantsAtFrame(const TraceServices::FFrame& InFrame, TArray<TSharedRef<FVariantTreeNode>>& OutVariants) const override;
 
 private:
 	// Helper function used to find a skeletal mesh pose
diff --git a/Engine/Plugins/Animation/GameplayInsights/Source/GameplayInsights/Private/SkeletalMeshPoseTrack.cpp b/Engine/Plugins/Animation/GameplayInsights/Source/GameplayInsights/Private/SkeletalMeshPoseTrack.cpp
--- a/Engine/Plugins/Animation/GameplayInsights/Source/GameplayInsights/Private/SkeletalMeshPoseTrack.cpp
+++ b/Engine/Plugins/Animation/GameplayInsights/Source/GameplayInsights/Private/SkeletalMeshPoseTrack.cpp
@@ -71,14 +71,14 @@ void FSkeletalMeshPoseTrack::BuildDrawState(ITimingEventsTrackDrawStateBuilder&
 
 	if(AnimationProvider)
 	{
-		Trace::FAnalysisSessionReadScope SessionReadScope(SharedData.GetAnalysisSession());
+		TraceServices::FAnalysisSessionReadScope SessionReadScope(SharedData.GetAnalysisSession());
 
 		AnimationProvider->ReadSkeletalMeshPoseTimeline(GetGameplayTrack().GetObjectId(), [&Context, &Builder](const FAnimationProvider::SkeletalMeshPoseTimeline& InTimeline, bool bInHasCurves)
 		{
 			InTimeline.EnumerateEvents(Context.GetViewport().GetStartTime(), Context.GetViewport().GetEndTime(), [&Builder](double InStartTime, double InEndTime, uint32 InDepth, const FSkeletalMeshPoseMessage& InMessage)
 			{
 				Builder.AddEvent(InStartTime, InEndTime, 0, InMessage.MeshName);
-				return Trace::EEventEnumerate::Continue;
+				return TraceServices::EEventEnumerate::Continue;
 			});
 		});
 	}
@@ -103,7 +103,7 @@ void FSkeletalMeshPoseTrack::InitTooltip(FTooltipDrawState& Tooltip, const ITimi
 		Tooltip.AddNameValueTextLine(LOCTEXT("EventTime", "Time").ToString(), FText::AsNumber(InFoundStartTime).ToString());
 
 		{
-			Trace::FAnalysisSessionReadScope SessionReadScope(SharedData.GetAnalysisSession());
+			TraceServices::FAnalysisSessionReadScope SessionReadScope(SharedData.GetAnalysisSession());
 
 			const FGameplayProvider* GameplayProvider = SharedData.GetAnalysisSession().ReadProvider<FGameplayProvider>(FGameplayProvider::ProviderName);
 			const FObjectInfo* SkeletalMeshObjectInfo = GameplayProvider->FindObjectInfo(InMessage.MeshId);
@@ -151,14 +151,14 @@ void FSkeletalMeshPoseTrack::FindSkeletalMeshPoseMessage(const FTimingEventSearc
 
 			if(AnimationProvider)
 			{
-				Trace::FAnalysisSessionReadScope SessionReadScope(SharedData.GetAnalysisSession());
+				TraceServices::FAnalysisSessionReadScope SessionReadScope(SharedData.GetAnalysisSession());
 
 				AnimationProvider->ReadSkeletalMeshPoseTimeline(GetGameplayTrack().GetObjectId(), [&InContext](const FAnimationProvider::SkeletalMeshPoseTimeline& InTimeline, bool bInHasCurves)
 				{
 					InTimeline.EnumerateEvents(InContext.GetParameters().StartTime, InContext.GetParameters().EndTime, [&InContext](double InEventStartTime, double InEventEndTime, uint32 InDepth, const FSkeletalMeshPoseMessage& InMessage)
 					{
 						InContext.Check(InEventStartTime, InEventEndTime, 0, InMessage);
-						return Trace::EEventEnumerate::Continue;
+						return TraceServices::EEventEnumerate::Continue;
 					});
 				});
 			}
@@ -214,7 +214,7 @@ void FSkeletalMeshPoseTrack::BuildContextMenu(FMenuBuilder& MenuBuilder)
 
 	if(GameplayProvider)
 	{
-		Trace::FAnalysisSessionReadScope SessionReadScope(SharedData.GetAnalysisSession());
+		TraceServices::FAnalysisSessionReadScope SessionReadScope(SharedData.GetAnalysisSession());
 	
 		const FObjectInfo* ComponentObjectInfo = GameplayProvider->FindObjectInfo(GetGameplayTrack().GetObjectId());
 		if(ComponentObjectInfo != nullptr)
@@ -232,7 +232,7 @@ void FSkeletalMeshPoseTrack::BuildContextMenu(FMenuBuilder& MenuBuilder)
 						FUIAction(
 							FExecuteAction::CreateLambda([this, ActorObjectInfo, GameplayProvider]()
 							{
-								Trace::FAnalysisSessionReadScope SessionReadScope(SharedData.GetAnalysisSession());
+								TraceServices::FAnalysisSessionReadScope SessionReadScope(SharedData.GetAnalysisSession());
 
 								bool bSetDrawPose = true;
 								SharedData.EnumerateSkeletalMeshPoseTracks([GameplayProvider, ActorObjectInfo, &bSetDrawPose](const TSharedRef<FSkeletalMeshPoseTrack>& InTrack)
@@ -259,7 +259,7 @@ void FSkeletalMeshPoseTrack::BuildContextMenu(FMenuBuilder& MenuBuilder)
 							FCanExecuteAction(),
 							FIsActionChecked::CreateLambda([this, ActorObjectInfo, GameplayProvider]()
 							{
-								Trace::FAnalysisSessionReadScope SessionReadScope(SharedData.GetAnalysisSession());
+								TraceServices::FAnalysisSessionReadScope SessionReadScope(SharedData.GetAnalysisSession());
 
 								bool bDrawPoseSet = true;
 								SharedData.EnumerateSkeletalMeshPoseTracks([GameplayProvider, ActorObjectInfo, &bDrawPoseSet](const TSharedRef<FSkeletalMeshPoseTrack>& InTrack)
@@ -284,7 +284,7 @@ void FSkeletalMeshPoseTrack::BuildContextMenu(FMenuBuilder& MenuBuilder)
 						FUIAction(
 							FExecuteAction::CreateLambda([this, ActorObjectInfo, GameplayProvider]()
 							{
-								Trace::FAnalysisSessionReadScope SessionReadScope(SharedData.GetAnalysisSession());
+								TraceServices::FAnalysisSessionReadScope SessionReadScope(SharedData.GetAnalysisSession());
 
 								bool bSetDrawSkeleton = true;
 								SharedData.EnumerateSkeletalMeshPoseTracks([GameplayProvider, ActorObjectInfo, &bSetDrawSkeleton](const TSharedRef<FSkeletalMeshPoseTrack>& InTrack)
@@ -311,7 +311,7 @@ void FSkeletalMeshPoseTrack::BuildContextMenu(FMenuBuilder& MenuBuilder)
 							FCanExecuteAction(),
 							FIsActionChecked::CreateLambda([this, ActorObjectInfo, GameplayProvider]()
 							{
-								Trace::FAnalysisSessionReadScope SessionReadScope(SharedData.GetAnalysisSession());
+								TraceServices::FAnalysisSessionReadScope SessionReadScope(SharedData.GetAnalysisSession());
 
 								bool bDrawSkeletonSet = true;
 								SharedData.EnumerateSkeletalMeshPoseTracks([GameplayProvider, ActorObjectInfo, &bDrawSkeletonSet](const TSharedRef<FSkeletalMeshPoseTrack>& InTrack)
@@ -376,7 +376,7 @@ void FSkeletalMeshPoseTrack::DrawPoses(UWorld* InWorld, double InTime, double In
 
 		if(AnimationProvider && GameplayProvider)
 		{
-			Trace::FAnalysisSessionReadScope SessionReadScope(SharedData.GetAnalysisSession());
+			TraceServices::FAnalysisSessionReadScope SessionReadScope(SharedData.GetAnalysisSession());
 
 			FWorldComponentCache& CacheForWorld = GetWorldCache(InWorld);
 			if(CacheForWorld.Component)
@@ -414,7 +414,7 @@ void FSkeletalMeshPoseTrack::DrawPoses(UWorld* InWorld, double InTime, double In
 							Component->SetDebugDrawColor(Color);
 						}
 					}
-					return Trace::EEventEnumerate::Continue;
+					return TraceServices::EEventEnumerate::Continue;
 				});
 			});
 		}
diff --git a/Engine/Plugins/Animation/GameplayInsights/Source/GameplayInsights/Private/VariantTreeNode.cpp b/Engine/Plugins/Animation/GameplayInsights/Source/GameplayInsights/Private/VariantTreeNode.cpp
--- a/Engine/Plugins/Animation/GameplayInsights/Source/GameplayInsights/Private/VariantTreeNode.cpp
+++ b/Engine/Plugins/Animation/GameplayInsights/Source/GameplayInsights/Private/VariantTreeNode.cpp
@@ -6,7 +6,7 @@
 
 #define LOCTEXT_NAMESPACE "FVariantTreeNode"
 
-FString FVariantTreeNode::GetValueAsString(const Trace::IAnalysisSession& InAnalysisSession) const
+FString FVariantTreeNode::GetValueAsString(const TraceServices::IAnalysisSession& InAnalysisSession) const
 {
 	switch(Value.Type)
 	{
@@ -27,7 +27,7 @@ FString FVariantTreeNode::GetValueAsString(const Trace::IAnalysisSession& InAnal
 			const FGameplayProvider* GameplayProvider = InAnalysisSession.ReadProvider<FGameplayProvider>(FGameplayProvider::ProviderName);
 			if(GameplayProvider)
 			{
-				Trace::FAnalysisSessionReadScope SessionReadScope(InAnalysisSession);
+				TraceServices::FAnalysisSessionReadScope SessionReadScope(InAnalysisSession);
 
 				const FObjectInfo& ObjectInfo = GameplayProvider->GetObjectInfo(Value.Object.Value);
 				return ObjectInfo.PathName;
@@ -40,7 +40,7 @@ FString FVariantTreeNode::GetValueAsString(const Trace::IAnalysisSession& InAnal
 			const FGameplayProvider* GameplayProvider = InAnalysisSession.ReadProvider<FGameplayProvider>(FGameplayProvider::ProviderName);
 			if(GameplayProvider)
 			{
-				Trace::FAnalysisSessionReadScope SessionReadScope(InAnalysisSession);
+				TraceServices::FAnalysisSessionReadScope SessionReadScope(InAnalysisSession);
 
 				const FClassInfo& ClassInfo = GameplayProvider->GetClassInfo(Value.Class.Value);
 				return ClassInfo.PathName;
@@ -53,4 +53,4 @@ FString FVariantTreeNode::GetValueAsString(const Trace::IAnalysisSession& InAnal
 	return FString();
 }
 
-#undef LOCTEXT_NAMESPACE
\ No newline at end of file
+#undef LOCTEXT_NAMESPACE
diff --git a/Engine/Plugins/Animation/GameplayInsights/Source/GameplayInsights/Private/VariantTreeNode.h b/Engine/Plugins/Animation/GameplayInsights/Source/GameplayInsights/Private/VariantTreeNode.h
--- a/Engine/Plugins/Animation/GameplayInsights/Source/GameplayInsights/Private/VariantTreeNode.h
+++ b/Engine/Plugins/Animation/GameplayInsights/Source/GameplayInsights/Private/VariantTreeNode.h
@@ -5,7 +5,7 @@
 #include "CoreMinimal.h"
 #include "IAnimationProvider.h"
 
-namespace Trace { class IAnalysisSession; }
+namespace TraceServices { class IAnalysisSession; }
 
 enum class EVariantTreeNodeFilterState
 {
@@ -46,7 +46,7 @@ struct FVariantTreeNode : TSharedFromThis<FVariantTreeNode>
 
 	EVariantTreeNodeFilterState GetFilterState() const { return FilterState; }
 
-	FString GetValueAsString(const Trace::IAnalysisSession& InAnalysisSession) const;
+	FString GetValueAsString(const TraceServices::IAnalysisSession& InAnalysisSession) const;
 
 	static TSharedRef<FVariantTreeNode> MakeHeader(const FText& InName, uint64 InId)
 	{
@@ -141,4 +141,4 @@ private:
 
 	// Filter state given current filter text
 	EVariantTreeNodeFilterState FilterState;
-};
\ No newline at end of file
+};
diff --git a/Engine/Plugins/Animation/GameplayInsights/Source/GameplayInsights/Public/IAnimationProvider.h b/Engine/Plugins/Animation/GameplayInsights/Source/GameplayInsights/Public/IAnimationProvider.h
--- a/Engine/Plugins/Animation/GameplayInsights/Source/GameplayInsights/Public/IAnimationProvider.h
+++ b/Engine/Plugins/Animation/GameplayInsights/Source/GameplayInsights/Public/IAnimationProvider.h
@@ -227,19 +227,19 @@ struct FAnimMontageMessage
 	uint16 FrameCounter = 0;
 };
 
-class IAnimationProvider : public Trace::IProvider
+class IAnimationProvider : public TraceServices::IProvider
 {
 public:
-	typedef Trace::ITimeline<FTickRecordMessage> TickRecordTimeline;
-	typedef Trace::ITimeline<FSkeletalMeshPoseMessage> SkeletalMeshPoseTimeline;
-	typedef Trace::ITimeline<FAnimGraphMessage> AnimGraphTimeline;
-	typedef Trace::ITimeline<FAnimNodeMessage> AnimNodesTimeline;
-	typedef Trace::ITimeline<FAnimNodeValueMessage> AnimNodeValuesTimeline;
-	typedef Trace::ITimeline<FAnimSequencePlayerMessage> AnimSequencePlayersTimeline;
-	typedef Trace::ITimeline<FAnimStateMachineMessage> StateMachinesTimeline;
-	typedef Trace::ITimeline<FBlendSpacePlayerMessage> BlendSpacePlayersTimeline;
-	typedef Trace::ITimeline<FAnimNotifyMessage> AnimNotifyTimeline;
-	typedef Trace::ITimeline<FAnimMontageMessage> AnimMontageTimeline;
+	typedef TraceServices::ITimeline<FTickRecordMessage> TickRecordTimeline;
+	typedef TraceServices::ITimeline<FSkeletalMeshPoseMessage> SkeletalMeshPoseTimeline;
+	typedef TraceServices::ITimeline<FAnimGraphMessage> AnimGraphTimeline;
+	typedef TraceServices::ITimeline<FAnimNodeMessage> AnimNodesTimeline;
+	typedef TraceServices::ITimeline<FAnimNodeValueMessage> AnimNodeValuesTimeline;
+	typedef TraceServices::ITimeline<FAnimSequencePlayerMessage> AnimSequencePlayersTimeline;
+	typedef TraceServices::ITimeline<FAnimStateMachineMessage> StateMachinesTimeline;
+	typedef TraceServices::ITimeline<FBlendSpacePlayerMessage> BlendSpacePlayersTimeline;
+	typedef TraceServices::ITimeline<FAnimNotifyMessage> AnimNotifyTimeline;
+	typedef TraceServices::ITimeline<FAnimMontageMessage> AnimMontageTimeline;
 
 	virtual bool ReadSkeletalMeshPoseTimeline(uint64 InObjectId, TFunctionRef<void(const SkeletalMeshPoseTimeline&, bool)> Callback) const = 0;
 	virtual void GetSkeletalMeshComponentSpacePose(const FSkeletalMeshPoseMessage& InMessage, const FSkeletalMeshInfo& InMeshInfo, FTransform& OutComponentToWorld, TArray<FTransform>& OutTransforms) const = 0;
@@ -261,4 +261,4 @@ public:
 	virtual const TCHAR* GetName(uint32 InId) const = 0;
 	virtual FText FormatNodeKeyValue(const FAnimNodeValueMessage& InMessage) const = 0;
 	virtual FText FormatNodeValue(const FAnimNodeValueMessage& InMessage) const = 0;
-};
\ No newline at end of file
+};
diff --git a/Engine/Plugins/Animation/GameplayInsights/Source/GameplayInsights/Public/IGameplayProvider.h b/Engine/Plugins/Animation/GameplayInsights/Source/GameplayInsights/Public/IGameplayProvider.h
--- a/Engine/Plugins/Animation/GameplayInsights/Source/GameplayInsights/Public/IGameplayProvider.h
+++ b/Engine/Plugins/Animation/GameplayInsights/Source/GameplayInsights/Public/IGameplayProvider.h
@@ -87,11 +87,11 @@ struct FWorldInfo
 // Delegate fired when an object receives an end play event
 DECLARE_MULTICAST_DELEGATE_ThreeParams(FOnObjectEndPlay, uint64 /*ObjectId*/, double /*Time*/, const FObjectInfo& /*ObjectInfo*/);
 
-class IGameplayProvider : public Trace::IProvider
+class IGameplayProvider : public TraceServices::IProvider
 {
 public:
-	typedef Trace::ITimeline<FObjectEventMessage> ObjectEventsTimeline;
-	typedef Trace::ITimeline<FObjectPropertiesMessage> ObjectPropertiesTimeline;
+	typedef TraceServices::ITimeline<FObjectEventMessage> ObjectEventsTimeline;
+	typedef TraceServices::ITimeline<FObjectPropertiesMessage> ObjectPropertiesTimeline;
 
 	virtual bool ReadObjectEventsTimeline(uint64 InObjectId, TFunctionRef<void(const ObjectEventsTimeline&)> Callback) const = 0;
 	virtual bool ReadObjectEvent(uint64 InObjectId, uint64 InMessageId, TFunctionRef<void(const FObjectEventMessage&)> Callback) const = 0;
@@ -109,4 +109,4 @@ public:
 	virtual const FObjectInfo& GetObjectInfo(uint64 InObjectId) const = 0;
 	virtual FOnObjectEndPlay& OnObjectEndPlay() = 0;
 	virtual const TCHAR* GetPropertyName(uint32 InPropertyStringId) const = 0;
-};
\ No newline at end of file
+};
diff --git a/Engine/Plugins/Developer/TraceDataFiltering/Source/TraceDataFiltering/Private/BaseSessionFilterService.cpp b/Engine/Plugins/Developer/TraceDataFiltering/Source/TraceDataFiltering/Private/BaseSessionFilterService.cpp
--- a/Engine/Plugins/Developer/TraceDataFiltering/Source/TraceDataFiltering/Private/BaseSessionFilterService.cpp
+++ b/Engine/Plugins/Developer/TraceDataFiltering/Source/TraceDataFiltering/Private/BaseSessionFilterService.cpp
@@ -13,7 +13,7 @@
 #include <IPAddress.h>
 #include <SocketSubsystem.h>
 
-FBaseSessionFilterService::FBaseSessionFilterService(Trace::FSessionHandle InHandle, TSharedPtr<const Trace::IAnalysisSession> InSession) : Session(InSession), Handle(InHandle)
+FBaseSessionFilterService::FBaseSessionFilterService(TraceServices::FSessionHandle InHandle, TSharedPtr<const TraceServices::IAnalysisSession> InSession) : Session(InSession), Handle(InHandle)
 {
 	FCoreDelegates::OnEndFrame.AddRaw(this, &FBaseSessionFilterService::OnApplyChannelChanges);
 }
@@ -25,11 +25,11 @@ FBaseSessionFilterService::~FBaseSessionFilterService()
 
 void FBaseSessionFilterService::GetRootObjects(TArray<FTraceObjectInfo>& OutObjects) const
 {
-	const Trace::IChannelProvider* ChannelProvider = Session->ReadProvider<Trace::IChannelProvider>("ChannelProvider");
+	const TraceServices::IChannelProvider* ChannelProvider = Session->ReadProvider<TraceServices::IChannelProvider>("ChannelProvider");
 	if (ChannelProvider)
 	{
 		const uint64 ChannelCount = ChannelProvider->GetChannelCount();		
-		const TArray<Trace::FChannelEntry>& Channels = ChannelProvider->GetChannels();
+		const TArray<TraceServices::FChannelEntry>& Channels = ChannelProvider->GetChannels();
 		for (uint64 ChannelIndex = 0; ChannelIndex < Channels.Num(); ++ChannelIndex)
 		{
 			FTraceObjectInfo& EventInfo = OutObjects.AddDefaulted_GetRef();
@@ -49,7 +49,7 @@ void FBaseSessionFilterService::GetChildObjects(uint32 InObjectHash, TArray<FTra
 
 const FDateTime& FBaseSessionFilterService::GetTimestamp()
 {
-	const Trace::IChannelProvider* ChannelProvider = Session->ReadProvider<Trace::IChannelProvider>("ChannelProvider");
+	const TraceServices::IChannelProvider* ChannelProvider = Session->ReadProvider<TraceServices::IChannelProvider>("ChannelProvider");
 	if (ChannelProvider)
 	{
 		TimeStamp = ChannelProvider->GetTimeStamp();
@@ -88,11 +88,11 @@ void FBaseSessionFilterService::UpdateFilterPreset(const TSharedPtr<IFilterPrese
 
 void FBaseSessionFilterService::DisableAllChannels()
 {
-	const Trace::IChannelProvider* ChannelProvider = Session->ReadProvider<Trace::IChannelProvider>("ChannelProvider");
+	const TraceServices::IChannelProvider* ChannelProvider = Session->ReadProvider<TraceServices::IChannelProvider>("ChannelProvider");
 	if (ChannelProvider)
 	{
 		const uint64 ChannelCount = ChannelProvider->GetChannelCount();
-		const TArray<Trace::FChannelEntry>& Channels = ChannelProvider->GetChannels();
+		const TArray<TraceServices::FChannelEntry>& Channels = ChannelProvider->GetChannels();
 		for (uint64 ChannelIndex = 0; ChannelIndex < Channels.Num(); ++ChannelIndex)
 		{
 			SetObjectFilterState(Channels[ChannelIndex].Name, false);
diff --git a/Engine/Plugins/Developer/TraceDataFiltering/Source/TraceDataFiltering/Private/BaseSessionFilterService.h b/Engine/Plugins/Developer/TraceDataFiltering/Source/TraceDataFiltering/Private/BaseSessionFilterService.h
--- a/Engine/Plugins/Developer/TraceDataFiltering/Source/TraceDataFiltering/Private/BaseSessionFilterService.h
+++ b/Engine/Plugins/Developer/TraceDataFiltering/Source/TraceDataFiltering/Private/BaseSessionFilterService.h
@@ -5,7 +5,7 @@
 #include "ISessionTraceFilterService.h"
 #include "Misc/DateTime.h"
 
-namespace Trace
+namespace TraceServices
 {
 	class IAnalysisSession;
 	typedef uint64 FSessionHandle;
@@ -13,12 +13,12 @@ namespace Trace
 
 class IEventInfoProvider;
 
-/** Implementation of ISessionTraceFilterService specifically to use as a base for the Trace and Editor implementations, using Trace::IChannelProvider to provide information about Channels available on the running application 
- and Trace::ISessionService to change the Channel state(s). */
+/** Implementation of ISessionTraceFilterService specifically to use as a base for the Trace and Editor implementations, using TraceServices::IChannelProvider to provide information about Channels available on the running application 
+ and TraceServices::ISessionService to change the Channel state(s). */
 class FBaseSessionFilterService : public ISessionTraceFilterService
 {
 public:
-	FBaseSessionFilterService(Trace::FSessionHandle InHandle, TSharedPtr<const Trace::IAnalysisSession> InSession);
+	FBaseSessionFilterService(TraceServices::FSessionHandle InHandle, TSharedPtr<const TraceServices::IAnalysisSession> InSession);
 	virtual ~FBaseSessionFilterService();
 
 	/** Begin ISessionTraceFilterService overrides */
@@ -39,9 +39,9 @@ protected:
 	void RetrieveAndStoreStartupChannels();
 protected:
 	/** Session this instance represents the filtering service for */
-	TSharedPtr<const Trace::IAnalysisSession> Session;
+	TSharedPtr<const TraceServices::IAnalysisSession> Session;
 	/** Session handle for AnalyisSession*/
-	Trace::FSessionHandle Handle;
+	TraceServices::FSessionHandle Handle;
 
 	/** Names of channels that were either enabled or disabled during the duration of this frame */
 	TArray<FString> FrameEnabledChannels;
@@ -49,4 +49,4 @@ protected:
 
 	/** Timestamp at which contained data (including provider) was last updated */
 	FDateTime TimeStamp;	
-};
\ No newline at end of file
+};
diff --git a/Engine/Plugins/Developer/TraceDataFiltering/Source/TraceDataFiltering/Private/ISessionTraceFilterService.h b/Engine/Plugins/Developer/TraceDataFiltering/Source/TraceDataFiltering/Private/ISessionTraceFilterService.h
--- a/Engine/Plugins/Developer/TraceDataFiltering/Source/TraceDataFiltering/Private/ISessionTraceFilterService.h
+++ b/Engine/Plugins/Developer/TraceDataFiltering/Source/TraceDataFiltering/Private/ISessionTraceFilterService.h
@@ -16,7 +16,7 @@ struct FTraceObjectInfo
 	uint32 OwnerHash;
 };
 
-/** Filtering service, representing the state and data for a specific Trace::IAnalysisSession */
+/** Filtering service, representing the state and data for a specific TraceServices::IAnalysisSession */
 class ISessionTraceFilterService : public TSharedFromThis<ISessionTraceFilterService>
 {
 public:
diff --git a/Engine/Plugins/Developer/TraceDataFiltering/Source/TraceDataFiltering/Private/STraceDataFilterWidget.cpp b/Engine/Plugins/Developer/TraceDataFiltering/Source/TraceDataFiltering/Private/STraceDataFilterWidget.cpp
--- a/Engine/Plugins/Developer/TraceDataFiltering/Source/TraceDataFiltering/Private/STraceDataFilterWidget.cpp
+++ b/Engine/Plugins/Developer/TraceDataFiltering/Source/TraceDataFiltering/Private/STraceDataFilterWidget.cpp
@@ -566,7 +566,7 @@ void STraceDataFilterWidget::RestoreItemSelection()
 	SelectedObjectNames.Empty();
 }
 
-void STraceDataFilterWidget::SetCurrentAnalysisSession(uint32 SessionHandle, TSharedRef<const Trace::IAnalysisSession> AnalysisSession)
+void STraceDataFilterWidget::SetCurrentAnalysisSession(uint32 SessionHandle, TSharedRef<const TraceServices::IAnalysisSession> AnalysisSession)
 {
 #if WITH_EDITOR
 	SessionFilterService = MakeShareable(new FSessionEditorFilterService(SessionHandle, AnalysisSession));
@@ -717,7 +717,7 @@ void STraceDataFilterWidget::RefreshTreeviewData()
 void STraceDataFilterWidget::Tick(const FGeometry& AllottedGeometry, const double InCurrentTime, const float InDeltaTime)
 {
 	IUnrealInsightsModule& InsightsModule = FModuleManager::LoadModuleChecked<IUnrealInsightsModule>("TraceInsights");
-	TSharedPtr<const Trace::IAnalysisSession> AnalysisSession = InsightsModule.GetAnalysisSession();
+	TSharedPtr<const TraceServices::IAnalysisSession> AnalysisSession = InsightsModule.GetAnalysisSession();
 
 	if (SessionFilterService.IsValid() )
 	{
diff --git a/Engine/Plugins/Developer/TraceDataFiltering/Source/TraceDataFiltering/Private/STraceDataFilterWidget.h b/Engine/Plugins/Developer/TraceDataFiltering/Source/TraceDataFiltering/Private/STraceDataFilterWidget.h
--- a/Engine/Plugins/Developer/TraceDataFiltering/Source/TraceDataFiltering/Private/STraceDataFilterWidget.h
+++ b/Engine/Plugins/Developer/TraceDataFiltering/Source/TraceDataFiltering/Private/STraceDataFilterWidget.h
@@ -54,7 +54,7 @@ protected:
 	void ConstructFilterHandler();
 
 	/** Callback for whenever a different analysis session (store) has been retrieved */
-	void SetCurrentAnalysisSession(uint32 SessionHandle, TSharedRef<const Trace::IAnalysisSession> AnalysisSession);
+	void SetCurrentAnalysisSession(uint32 SessionHandle, TSharedRef<const TraceServices::IAnalysisSession> AnalysisSession);
 	bool HasValidFilterSession() const;
 		
 	void RefreshTreeviewData();
@@ -131,4 +131,4 @@ protected:
 
 	/** Inter-window cached state of last expansion state */
 	static TSet<FString> LastExpandedObjectNames;
-};
\ No newline at end of file
+};
diff --git a/Engine/Plugins/Developer/TraceDataFiltering/Source/TraceDataFiltering/Private/SessionEditorFilterService.cpp b/Engine/Plugins/Developer/TraceDataFiltering/Source/TraceDataFiltering/Private/SessionEditorFilterService.cpp
--- a/Engine/Plugins/Developer/TraceDataFiltering/Source/TraceDataFiltering/Private/SessionEditorFilterService.cpp
+++ b/Engine/Plugins/Developer/TraceDataFiltering/Source/TraceDataFiltering/Private/SessionEditorFilterService.cpp
@@ -6,7 +6,7 @@
 #include "Misc/CoreDelegates.h"
 #include "Trace/Trace.h"
 
-FSessionEditorFilterService::FSessionEditorFilterService(Trace::FSessionHandle InHandle, TSharedPtr<const Trace::IAnalysisSession> InSession) : FBaseSessionFilterService(InHandle, InSession)
+FSessionEditorFilterService::FSessionEditorFilterService(TraceServices::FSessionHandle InHandle, TSharedPtr<const TraceServices::IAnalysisSession> InSession) : FBaseSessionFilterService(InHandle, InSession)
 {	
 }
 
diff --git a/Engine/Plugins/Developer/TraceDataFiltering/Source/TraceDataFiltering/Private/SessionEditorFilterService.h b/Engine/Plugins/Developer/TraceDataFiltering/Source/TraceDataFiltering/Private/SessionEditorFilterService.h
--- a/Engine/Plugins/Developer/TraceDataFiltering/Source/TraceDataFiltering/Private/SessionEditorFilterService.h
+++ b/Engine/Plugins/Developer/TraceDataFiltering/Source/TraceDataFiltering/Private/SessionEditorFilterService.h
@@ -5,7 +5,7 @@
 #include "ISessionTraceFilterService.h"
 #include "BaseSessionFilterService.h"
 
-namespace Trace
+namespace TraceServices
 {
 	class IAnalysisSession;
 	typedef uint64 FSessionHandle;
@@ -15,10 +15,10 @@ namespace Trace
 class FSessionEditorFilterService : public FBaseSessionFilterService
 {
 public:
-	FSessionEditorFilterService(Trace::FSessionHandle InHandle, TSharedPtr<const Trace::IAnalysisSession> InSession);
+	FSessionEditorFilterService(TraceServices::FSessionHandle InHandle, TSharedPtr<const TraceServices::IAnalysisSession> InSession);
 	virtual ~FSessionEditorFilterService() {}
 
 	/** Begin FBaseSessionFilterService overrides */
 	virtual void OnApplyChannelChanges() override;
 	/** End OnApplyChannelChanges overrides */
-};
\ No newline at end of file
+};
diff --git a/Engine/Plugins/Developer/TraceDataFiltering/Source/TraceDataFiltering/Private/Trace/SessionTraceFilterService.cpp b/Engine/Plugins/Developer/TraceDataFiltering/Source/TraceDataFiltering/Private/Trace/SessionTraceFilterService.cpp
--- a/Engine/Plugins/Developer/TraceDataFiltering/Source/TraceDataFiltering/Private/Trace/SessionTraceFilterService.cpp
+++ b/Engine/Plugins/Developer/TraceDataFiltering/Source/TraceDataFiltering/Private/Trace/SessionTraceFilterService.cpp
@@ -13,7 +13,7 @@
 #include <IPAddress.h>
 #include <SocketSubsystem.h>
 
-FSessionTraceFilterService::FSessionTraceFilterService(Trace::FSessionHandle InHandle, TSharedPtr<const Trace::IAnalysisSession> InSession) : FBaseSessionFilterService(InHandle, InSession)
+FSessionTraceFilterService::FSessionTraceFilterService(TraceServices::FSessionHandle InHandle, TSharedPtr<const TraceServices::IAnalysisSession> InSession) : FBaseSessionFilterService(InHandle, InSession)
 {
 
 }
diff --git a/Engine/Plugins/Developer/TraceDataFiltering/Source/TraceDataFiltering/Private/Trace/SessionTraceFilterService.h b/Engine/Plugins/Developer/TraceDataFiltering/Source/TraceDataFiltering/Private/Trace/SessionTraceFilterService.h
--- a/Engine/Plugins/Developer/TraceDataFiltering/Source/TraceDataFiltering/Private/Trace/SessionTraceFilterService.h
+++ b/Engine/Plugins/Developer/TraceDataFiltering/Source/TraceDataFiltering/Private/Trace/SessionTraceFilterService.h
@@ -5,18 +5,18 @@
 #include "ISessionTraceFilterService.h"
 #include "BaseSessionFilterService.h"
 
-namespace Trace
+namespace TraceServices
 {
 	class IAnalysisSession;
 	typedef uint64 FSessionHandle;
 }
 
-/** Implementation of ISessionTraceFilterService specifically to use with Trace, using Trace::IChannelProvider to provide information about Channels available on the running application 
- and Trace::ISessionService to change the Channel state(s). */
+/** Implementation of ISessionTraceFilterService specifically to use with Trace, using TraceServices::IChannelProvider to provide information about Channels available on the running application 
+ and TraceServices::ISessionService to change the Channel state(s). */
 class FSessionTraceFilterService : public FBaseSessionFilterService
 {
 public:
-	FSessionTraceFilterService(Trace::FSessionHandle InHandle, TSharedPtr<const Trace::IAnalysisSession> InSession);
+	FSessionTraceFilterService(TraceServices::FSessionHandle InHandle, TSharedPtr<const TraceServices::IAnalysisSession> InSession);
 	virtual ~FSessionTraceFilterService() {}
 
 	/** Begin FBaseSessionFilterService overrides */
diff --git a/Engine/Plugins/Developer/TraceSourceFiltering/Source/SourceFilteringEditor/Private/EditorSessionSourceFilterService.h b/Engine/Plugins/Developer/TraceSourceFiltering/Source/SourceFilteringEditor/Private/EditorSessionSourceFilterService.h
--- a/Engine/Plugins/Developer/TraceSourceFiltering/Source/SourceFilteringEditor/Private/EditorSessionSourceFilterService.h
+++ b/Engine/Plugins/Developer/TraceSourceFiltering/Source/SourceFilteringEditor/Private/EditorSessionSourceFilterService.h
@@ -14,12 +14,6 @@
 #include "DataSourceFilter.h"
 #include "TraceSourceFilteringSettings.h"
 
-namespace Trace
-{
-	class IAnalysisSession;
-	typedef uint64 FSessionHandle;
-}
-
 class UDataSourceFilter;
 class UDataSourceFilterSet;
 
diff --git a/Engine/Plugins/Developer/TraceSourceFiltering/Source/SourceFilteringEditor/Private/STraceSourceFilteringWidget.cpp b/Engine/Plugins/Developer/TraceSourceFiltering/Source/SourceFilteringEditor/Private/STraceSourceFilteringWidget.cpp
--- a/Engine/Plugins/Developer/TraceSourceFiltering/Source/SourceFilteringEditor/Private/STraceSourceFilteringWidget.cpp
+++ b/Engine/Plugins/Developer/TraceSourceFiltering/Source/SourceFilteringEditor/Private/STraceSourceFilteringWidget.cpp
@@ -127,7 +127,7 @@ void STraceSourceFilteringWidget::ConstructMenuBox()
 	];
 }
 
-void STraceSourceFilteringWidget::SetCurrentAnalysisSession(uint32 SessionHandle, TSharedRef<const Trace::IAnalysisSession> AnalysisSession)
+void STraceSourceFilteringWidget::SetCurrentAnalysisSession(uint32 SessionHandle, TSharedRef<const TraceServices::IAnalysisSession> AnalysisSession)
 {
 	if (SessionFilterService.IsValid())
 	{
@@ -180,7 +180,7 @@ void STraceSourceFilteringWidget::Tick(const FGeometry& AllottedGeometry, const
 	if (!SessionFilterService.IsValid() )
 	{
 		IUnrealInsightsModule& InsightsModule = FModuleManager::LoadModuleChecked<IUnrealInsightsModule>("TraceInsights");
-		TSharedPtr<const Trace::IAnalysisSession> AnalysisSession = InsightsModule.GetAnalysisSession();
+		TSharedPtr<const TraceServices::IAnalysisSession> AnalysisSession = InsightsModule.GetAnalysisSession();
 		if (AnalysisSession.IsValid())
 		{
 			Trace::FStoreClient* StoreClient = InsightsModule.GetStoreClient();
diff --git a/Engine/Plugins/Developer/TraceSourceFiltering/Source/SourceFilteringEditor/Private/STraceSourceFilteringWidget.h b/Engine/Plugins/Developer/TraceSourceFiltering/Source/SourceFilteringEditor/Private/STraceSourceFilteringWidget.h
--- a/Engine/Plugins/Developer/TraceSourceFiltering/Source/SourceFilteringEditor/Private/STraceSourceFilteringWidget.h
+++ b/Engine/Plugins/Developer/TraceSourceFiltering/Source/SourceFilteringEditor/Private/STraceSourceFilteringWidget.h
@@ -42,7 +42,7 @@ protected:
 	TSharedRef<SWidget> OnGetOptionsMenu();
 
 	/** Callback for whenever a different analysis session (store) has been retrieved */
-	void SetCurrentAnalysisSession(uint32 SessionHandle, TSharedRef<const Trace::IAnalysisSession> AnalysisSession);
+	void SetCurrentAnalysisSession(uint32 SessionHandle, TSharedRef<const TraceServices::IAnalysisSession> AnalysisSession);
 
 	/** Returns whether or not a valid ISessionSourceFilterService is available / set */
 	bool HasValidFilterSession() const;
diff --git a/Engine/Plugins/Developer/TraceSourceFiltering/Source/SourceFilteringEditor/Private/SourceFilterService.cpp b/Engine/Plugins/Developer/TraceSourceFiltering/Source/SourceFilteringEditor/Private/SourceFilterService.cpp
--- a/Engine/Plugins/Developer/TraceSourceFiltering/Source/SourceFilteringEditor/Private/SourceFilterService.cpp
+++ b/Engine/Plugins/Developer/TraceSourceFiltering/Source/SourceFilteringEditor/Private/SourceFilterService.cpp
@@ -11,9 +11,9 @@
 #include "EditorSessionSourceFilterService.h"
 #endif
 
-TSharedRef<ISessionSourceFilterService> FSourceFilterService::GetFilterServiceForSession(uint32 InHandle, TSharedRef<const Trace::IAnalysisSession> AnalysisSession)
+TSharedRef<ISessionSourceFilterService> FSourceFilterService::GetFilterServiceForSession(uint32 InHandle, TSharedRef<const TraceServices::IAnalysisSession> AnalysisSession)
 {
 	TSharedPtr<ISessionSourceFilterService> NewService = nullptr;		
 	NewService = MakeShareable(new FEditorSessionSourceFilterService());
 	return NewService->AsShared();
-}
\ No newline at end of file
+}
diff --git a/Engine/Plugins/Developer/TraceSourceFiltering/Source/SourceFilteringEditor/Private/SourceFilterService.h b/Engine/Plugins/Developer/TraceSourceFiltering/Source/SourceFilteringEditor/Private/SourceFilterService.h
--- a/Engine/Plugins/Developer/TraceSourceFiltering/Source/SourceFilteringEditor/Private/SourceFilterService.h
+++ b/Engine/Plugins/Developer/TraceSourceFiltering/Source/SourceFilteringEditor/Private/SourceFilterService.h
@@ -6,7 +6,7 @@
 #include "Containers/Map.h"
 #include "ISessionSourceFilterService.h"
 
-namespace Trace
+namespace TraceServices
 {
 	class IAnalysisSession;
 }
@@ -15,5 +15,5 @@ class FSourceFilterService
 {
 public:
 	/** Retrieves the ISessionSourceFilterService for the provided Trace Session Handle */
-	static TSharedRef<ISessionSourceFilterService> GetFilterServiceForSession(uint32 InHandle, TSharedRef<const Trace::IAnalysisSession> AnalysisSession);
-};
\ No newline at end of file
+	static TSharedRef<ISessionSourceFilterService> GetFilterServiceForSession(uint32 InHandle, TSharedRef<const TraceServices::IAnalysisSession> AnalysisSession);
+};
diff --git a/Engine/Plugins/Runtime/NetworkPredictionInsights/Source/NetworkPredictionInsights/Private/NetworkPredictionAnalyzer.cpp b/Engine/Plugins/Runtime/NetworkPredictionInsights/Source/NetworkPredictionInsights/Private/NetworkPredictionAnalyzer.cpp
--- a/Engine/Plugins/Runtime/NetworkPredictionInsights/Source/NetworkPredictionInsights/Private/NetworkPredictionAnalyzer.cpp
+++ b/Engine/Plugins/Runtime/NetworkPredictionInsights/Source/NetworkPredictionInsights/Private/NetworkPredictionAnalyzer.cpp
@@ -5,7 +5,7 @@
 #include "Containers/StringView.h"
 
 
-FNetworkPredictionAnalyzer::FNetworkPredictionAnalyzer(Trace::IAnalysisSession& InSession, FNetworkPredictionProvider& InNetworkPredictionProvider)
+FNetworkPredictionAnalyzer::FNetworkPredictionAnalyzer(TraceServices::IAnalysisSession& InSession, FNetworkPredictionProvider& InNetworkPredictionProvider)
 	: Session(InSession)
 	, NetworkPredictionProvider(InNetworkPredictionProvider)
 {
@@ -52,7 +52,7 @@ void FNetworkPredictionAnalyzer::OnAnalysisEnd()
 
 bool FNetworkPredictionAnalyzer::OnEvent(uint16 RouteId, EStyle Style, const FOnEventContext& Context)
 {
-	Trace::FAnalysisSessionEditScope _(Session);
+	TraceServices::FAnalysisSessionEditScope _(Session);
 	const auto& EventData = Context.EventData;
 
 	auto ParseUserState = [&EventData, this](ENP_UserState Type)
@@ -274,4 +274,4 @@ bool FNetworkPredictionAnalyzer::OnEvent(uint16 RouteId, EStyle Style, const FOn
 	}
 
 	return true;
-}
\ No newline at end of file
+}
diff --git a/Engine/Plugins/Runtime/NetworkPredictionInsights/Source/NetworkPredictionInsights/Private/NetworkPredictionAnalyzer.h b/Engine/Plugins/Runtime/NetworkPredictionInsights/Source/NetworkPredictionInsights/Private/NetworkPredictionAnalyzer.h
--- a/Engine/Plugins/Runtime/NetworkPredictionInsights/Source/NetworkPredictionInsights/Private/NetworkPredictionAnalyzer.h
+++ b/Engine/Plugins/Runtime/NetworkPredictionInsights/Source/NetworkPredictionInsights/Private/NetworkPredictionAnalyzer.h
@@ -6,7 +6,7 @@
 #include "TraceServices/Model/AnalysisSession.h"
 
 class FNetworkPredictionProvider;
-namespace Trace { class IAnalysisSession; }
+namespace TraceServices { class IAnalysisSession; }
 
 // Analyzes events that are contained in a trace,
 // Works by subscribing to events by name along with user-provided "route" identifiers
@@ -19,7 +19,7 @@ class FNetworkPredictionAnalyzer : public Trace::IAnalyzer
 {
 public:
 
-	FNetworkPredictionAnalyzer(Trace::IAnalysisSession& InSession, FNetworkPredictionProvider& InNetworkPredictionProvider);
+	FNetworkPredictionAnalyzer(TraceServices::IAnalysisSession& InSession, FNetworkPredictionProvider& InNetworkPredictionProvider);
 
 	virtual void OnAnalysisBegin(const FOnAnalysisContext& Context) override;
 	virtual void OnAnalysisEnd() override;
@@ -54,7 +54,7 @@ private:
 	};
 
 
-	Trace::IAnalysisSession& Session;
+	TraceServices::IAnalysisSession& Session;
 	FNetworkPredictionProvider& NetworkPredictionProvider;
 
 	// Current values
@@ -69,4 +69,4 @@ private:
 	bool bLocalOffsetFrameChanged = false;
 
 	int32 PendingWriteFrame;
-};
\ No newline at end of file
+};
diff --git a/Engine/Plugins/Runtime/NetworkPredictionInsights/Source/NetworkPredictionInsights/Private/NetworkPredictionInsightsModule.cpp b/Engine/Plugins/Runtime/NetworkPredictionInsights/Source/NetworkPredictionInsights/Private/NetworkPredictionInsightsModule.cpp
--- a/Engine/Plugins/Runtime/NetworkPredictionInsights/Source/NetworkPredictionInsights/Private/NetworkPredictionInsightsModule.cpp
+++ b/Engine/Plugins/Runtime/NetworkPredictionInsights/Source/NetworkPredictionInsights/Private/NetworkPredictionInsightsModule.cpp
@@ -33,7 +33,7 @@ const FName FNetworkPredictionInsightsModule::InsightsTabName("NetworkPrediction
 
 void FNetworkPredictionInsightsModule::StartupModule()
 {
-	IModularFeatures::Get().RegisterModularFeature(Trace::ModuleFeatureName, &NetworkPredictionTraceModule);
+	IModularFeatures::Get().RegisterModularFeature(TraceServices::ModuleFeatureName, &NetworkPredictionTraceModule);
 
 	FNetworkPredictionInsightsManager::Initialize();
 	IUnrealInsightsModule& UnrealInsightsModule = FModuleManager::LoadModuleChecked<IUnrealInsightsModule>("TraceInsights");
@@ -49,7 +49,7 @@ void FNetworkPredictionInsightsModule::StartupModule()
 		auto SessionPtr = UnrealInsightsModule.GetAnalysisSession();
 		if (SessionPtr.IsValid())
 		{
-			Trace::FAnalysisSessionReadScope SessionReadScope(*SessionPtr.Get());
+			TraceServices::FAnalysisSessionReadScope SessionReadScope(*SessionPtr.Get());
 			if (const INetworkPredictionProvider* NetworkPredictionProvider = ReadNetworkPredictionProvider(*SessionPtr.Get()))
 			{
 				auto NetworkPredictionTraceVersion = NetworkPredictionProvider->GetNetworkPredictionTraceVersion();
@@ -133,7 +133,7 @@ void FNetworkPredictionInsightsModule::ShutdownModule()
 
 	FTicker::GetCoreTicker().RemoveTicker(TickerHandle);
 
-	IModularFeatures::Get().UnregisterModularFeature(Trace::ModuleFeatureName, &NetworkPredictionTraceModule);
+	IModularFeatures::Get().UnregisterModularFeature(TraceServices::ModuleFeatureName, &NetworkPredictionTraceModule);
 }
 
 IMPLEMENT_MODULE(FNetworkPredictionInsightsModule, NetworkPredictionInsights);
diff --git a/Engine/Plugins/Runtime/NetworkPredictionInsights/Source/NetworkPredictionInsights/Private/NetworkPredictionProvider.cpp b/Engine/Plugins/Runtime/NetworkPredictionInsights/Source/NetworkPredictionInsights/Private/NetworkPredictionProvider.cpp
--- a/Engine/Plugins/Runtime/NetworkPredictionInsights/Source/NetworkPredictionInsights/Private/NetworkPredictionProvider.cpp
+++ b/Engine/Plugins/Runtime/NetworkPredictionInsights/Source/NetworkPredictionInsights/Private/NetworkPredictionProvider.cpp
@@ -8,7 +8,7 @@ DEFINE_LOG_CATEGORY_STATIC(LogNetworkPredictionTraceProvider, Log, All);
 
 // -----------------------------------------------------------------------------
 
-FNetworkPredictionProvider::FNetworkPredictionProvider(Trace::IAnalysisSession& InSession)
+FNetworkPredictionProvider::FNetworkPredictionProvider(TraceServices::IAnalysisSession& InSession)
 	: Session(InSession)
 {
 
@@ -407,7 +407,7 @@ TSharedRef<FSimulationData>& FNetworkPredictionProvider::FindChecked(int32 Trace
 //
 // -----------------------------------------------------------------------------
 
-const INetworkPredictionProvider* ReadNetworkPredictionProvider(const Trace::IAnalysisSession& Session)
+const INetworkPredictionProvider* ReadNetworkPredictionProvider(const TraceServices::IAnalysisSession& Session)
 {
 	Session.ReadAccessCheck();
 	return Session.ReadProvider<INetworkPredictionProvider>(FNetworkPredictionProvider::ProviderName);
@@ -502,4 +502,4 @@ const TCHAR* LexToString(ENP_TickingPolicy Policy)
 		ensure(false);
 		return TEXT("???");
 	}
-}
\ No newline at end of file
+}
diff --git a/Engine/Plugins/Runtime/NetworkPredictionInsights/Source/NetworkPredictionInsights/Private/NetworkPredictionProvider.h b/Engine/Plugins/Runtime/NetworkPredictionInsights/Source/NetworkPredictionInsights/Private/NetworkPredictionProvider.h
--- a/Engine/Plugins/Runtime/NetworkPredictionInsights/Source/NetworkPredictionInsights/Private/NetworkPredictionProvider.h
+++ b/Engine/Plugins/Runtime/NetworkPredictionInsights/Source/NetworkPredictionInsights/Private/NetworkPredictionProvider.h
@@ -9,7 +9,7 @@ class FNetworkPredictionProvider : public INetworkPredictionProvider
 public:
 	static FName ProviderName;
 
-	FNetworkPredictionProvider(Trace::IAnalysisSession& InSession);
+	FNetworkPredictionProvider(TraceServices::IAnalysisSession& InSession);
 
 	// -----------------------------------------------------
 	virtual uint32 GetNetworkPredictionTraceVersion() const override { return NetworkPredictionTraceVersion; }
@@ -60,5 +60,5 @@ private:
 	uint64 DataCounter = 0;
 	uint32 NetworkPredictionTraceVersion = 0;
 	int32 PIESessionCounter = 0;
-	Trace::IAnalysisSession& Session;
-};
\ No newline at end of file
+	TraceServices::IAnalysisSession& Session;
+};
diff --git a/Engine/Plugins/Runtime/NetworkPredictionInsights/Source/NetworkPredictionInsights/Private/NetworkPredictionTraceModule.cpp b/Engine/Plugins/Runtime/NetworkPredictionInsights/Source/NetworkPredictionInsights/Private/NetworkPredictionTraceModule.cpp
--- a/Engine/Plugins/Runtime/NetworkPredictionInsights/Source/NetworkPredictionInsights/Private/NetworkPredictionTraceModule.cpp
+++ b/Engine/Plugins/Runtime/NetworkPredictionInsights/Source/NetworkPredictionInsights/Private/NetworkPredictionTraceModule.cpp
@@ -6,13 +6,13 @@
 
 FName FNetworkPredictionTraceModule::ModuleName("NetworkPredictionTrace");
 
-void FNetworkPredictionTraceModule::GetModuleInfo(Trace::FModuleInfo& OutModuleInfo)
+void FNetworkPredictionTraceModule::GetModuleInfo(TraceServices::FModuleInfo& OutModuleInfo)
 {
 	OutModuleInfo.Name = ModuleName;
 	OutModuleInfo.DisplayName = TEXT("NetworkPrediction");
 }
 
-void FNetworkPredictionTraceModule::OnAnalysisBegin(Trace::IAnalysisSession& InSession)
+void FNetworkPredictionTraceModule::OnAnalysisBegin(TraceServices::IAnalysisSession& InSession)
 {
 	FNetworkPredictionProvider* NetworkPredictionProvider = new FNetworkPredictionProvider(InSession);
 	InSession.AddProvider(FNetworkPredictionProvider::ProviderName, NetworkPredictionProvider);
@@ -25,7 +25,7 @@ void FNetworkPredictionTraceModule::GetLoggers(TArray<const TCHAR *>& OutLoggers
 	OutLoggers.Add(TEXT("NetworkPrediction"));
 }
 
-void FNetworkPredictionTraceModule::GenerateReports(const Trace::IAnalysisSession& Session, const TCHAR* CmdLine, const TCHAR* OutputDirectory)
+void FNetworkPredictionTraceModule::GenerateReports(const TraceServices::IAnalysisSession& Session, const TCHAR* CmdLine, const TCHAR* OutputDirectory)
 {
 
 }
diff --git a/Engine/Plugins/Runtime/NetworkPredictionInsights/Source/NetworkPredictionInsights/Private/NetworkPredictionTraceModule.h b/Engine/Plugins/Runtime/NetworkPredictionInsights/Source/NetworkPredictionInsights/Private/NetworkPredictionTraceModule.h
--- a/Engine/Plugins/Runtime/NetworkPredictionInsights/Source/NetworkPredictionInsights/Private/NetworkPredictionTraceModule.h
+++ b/Engine/Plugins/Runtime/NetworkPredictionInsights/Source/NetworkPredictionInsights/Private/NetworkPredictionTraceModule.h
@@ -8,16 +8,16 @@
 class IAnimationProvider;
 namespace Insights { class ITimingViewSession; }
 
-class FNetworkPredictionTraceModule : public Trace::IModule
+class FNetworkPredictionTraceModule : public TraceServices::IModule
 {
 public:
-	// Trace::IModule interface
-	virtual void GetModuleInfo(Trace::FModuleInfo& OutModuleInfo) override;
-	virtual void OnAnalysisBegin(Trace::IAnalysisSession& Session) override;
+	// TraceServices::IModule interface
+	virtual void GetModuleInfo(TraceServices::FModuleInfo& OutModuleInfo) override;
+	virtual void OnAnalysisBegin(TraceServices::IAnalysisSession& Session) override;
 	virtual void GetLoggers(TArray<const TCHAR *>& OutLoggers) override;
-	virtual void GenerateReports(const Trace::IAnalysisSession& Session, const TCHAR* CmdLine, const TCHAR* OutputDirectory) override;
+	virtual void GenerateReports(const TraceServices::IAnalysisSession& Session, const TCHAR* CmdLine, const TCHAR* OutputDirectory) override;
 	virtual const TCHAR* GetCommandLineArgument() override { return TEXT("objecttrace"); }
 
 private:
 	static FName ModuleName;
-};
\ No newline at end of file
+};
diff --git a/Engine/Plugins/Runtime/NetworkPredictionInsights/Source/NetworkPredictionInsights/Private/UI/SNPSimFrameView.cpp b/Engine/Plugins/Runtime/NetworkPredictionInsights/Source/NetworkPredictionInsights/Private/UI/SNPSimFrameView.cpp
--- a/Engine/Plugins/Runtime/NetworkPredictionInsights/Source/NetworkPredictionInsights/Private/UI/SNPSimFrameView.cpp
+++ b/Engine/Plugins/Runtime/NetworkPredictionInsights/Source/NetworkPredictionInsights/Private/UI/SNPSimFrameView.cpp
@@ -353,11 +353,11 @@ void SNPSimFrameView::Tick(const FGeometry& AllottedGeometry, const double InCur
 		const uint64 WaitTime = static_cast<uint64>(0.1 / FPlatformTime::GetSecondsPerCycle64()); // 100ms
 		AnalysisSyncNextTimestamp = Time + WaitTime;
 
-		TSharedPtr<const Trace::IAnalysisSession> Session = FInsightsManager::Get()->GetSession();
+		TSharedPtr<const TraceServices::IAnalysisSession> Session = FInsightsManager::Get()->GetSession();
 		if (Session.IsValid())
 		{
-			Trace::FAnalysisSessionReadScope SessionReadScope(*Session.Get());
-			const Trace::INetProfilerProvider& NetProfilerProvider = Trace::ReadNetProfilerProvider(*Session.Get());
+			TraceServices::FAnalysisSessionReadScope SessionReadScope(*Session.Get());
+			const TraceServices::INetProfilerProvider& NetProfilerProvider = TraceServices::ReadNetProfilerProvider(*Session.Get());
 
 			const uint32 NewConnectionChangeCount = NetProfilerProvider.GetConnectionChangeCount();
 			if (NewConnectionChangeCount != ConnectionChangeCount)
@@ -403,10 +403,10 @@ void SNPSimFrameView::UpdateState()
 
 	TArray<FSimulationActorGroup>& ActorGroups = SimulationFrameView.ActorGroups;
 
-	TSharedPtr<const Trace::IAnalysisSession> Session = InsightsModule->GetAnalysisSession();
+	TSharedPtr<const TraceServices::IAnalysisSession> Session = InsightsModule->GetAnalysisSession();
 	if (Session.IsValid() && NetworkPredictionWindow.IsValid())
 	{
-		Trace::FAnalysisSessionReadScope SessionReadScope(*Session.Get());
+		TraceServices::FAnalysisSessionReadScope SessionReadScope(*Session.Get());
 		const INetworkPredictionProvider* NetworkPredictionProvider = ReadNetworkPredictionProvider(*Session.Get());
 		if (NetworkPredictionProvider)
 		{
@@ -1743,4 +1743,4 @@ void SNPSimFrameView::SortActorGroupToTop(const FSimulationActorGroup* Group)
 	bIsStateDirty = true;
 }
 
-#undef LOCTEXT_NAMESPACE
\ No newline at end of file
+#undef LOCTEXT_NAMESPACE
diff --git a/Engine/Plugins/Runtime/NetworkPredictionInsights/Source/NetworkPredictionInsights/Private/UI/SNPSimFrameView.h b/Engine/Plugins/Runtime/NetworkPredictionInsights/Source/NetworkPredictionInsights/Private/UI/SNPSimFrameView.h
--- a/Engine/Plugins/Runtime/NetworkPredictionInsights/Source/NetworkPredictionInsights/Private/UI/SNPSimFrameView.h
+++ b/Engine/Plugins/Runtime/NetworkPredictionInsights/Source/NetworkPredictionInsights/Private/UI/SNPSimFrameView.h
@@ -259,7 +259,7 @@ public:
 	void DrawSelection(int32 StartPacketIndex, int32 EndPacketIndex) const;
 	
 
-	static FLinearColor GetColorByStatus(Trace::ENetProfilerDeliveryStatus Status);
+	static FLinearColor GetColorByStatus(TraceServices::ENetProfilerDeliveryStatus Status);
 
 	int32 GetNumPackets() const { return NumPackets; }
 	int32 GetNumDrawSamples() const { return NumDrawSamples; }
@@ -458,4 +458,4 @@ private:
 
 	void SetFrameContentView_FrameNumber() { SimulationFrameView.ContentType = ESimFrameContentType::FrameNumber; bIsStateDirty = true;}
 	void SetFrameContentView_NumBufferdInputCmds() { SimulationFrameView.ContentType = ESimFrameContentType::NumBufferedInputCmds; bIsStateDirty = true; }
-};
\ No newline at end of file
+};
diff --git a/Engine/Plugins/Runtime/NetworkPredictionInsights/Source/NetworkPredictionInsights/Private/UI/SNPWindow.cpp b/Engine/Plugins/Runtime/NetworkPredictionInsights/Source/NetworkPredictionInsights/Private/UI/SNPWindow.cpp
--- a/Engine/Plugins/Runtime/NetworkPredictionInsights/Source/NetworkPredictionInsights/Private/UI/SNPWindow.cpp
+++ b/Engine/Plugins/Runtime/NetworkPredictionInsights/Source/NetworkPredictionInsights/Private/UI/SNPWindow.cpp
@@ -207,10 +207,10 @@ void SNPWindow::Construct(const FArguments& InArgs, const TSharedRef<SDockTab>&
 
 void SNPWindow::Tick(const FGeometry& AllottedGeometry, const double InCurrentTime, const float InDeltaTime)
 {
-	TSharedPtr<const Trace::IAnalysisSession> Session = InsightsModule->GetAnalysisSession();
+	TSharedPtr<const TraceServices::IAnalysisSession> Session = InsightsModule->GetAnalysisSession();
 	if (Session.IsValid())
 	{
-		Trace::FAnalysisSessionReadScope SessionReadScope(*Session.Get());
+		TraceServices::FAnalysisSessionReadScope SessionReadScope(*Session.Get());
 		if (const INetworkPredictionProvider* NetworkPredictionProvider = ReadNetworkPredictionProvider(*Session.Get()))
 		{
 			uint64 NewDataCounter = NetworkPredictionProvider->GetNetworkPredictionDataCounter();
@@ -518,10 +518,10 @@ void SNPWindow::PopulateFilteredDataView()
 	UnfilteredDataView.FirstEngineFrame = TNumericLimits<uint64>::Max();
 	UnfilteredDataView.LastEngineFrame = 0;
 
-	TSharedPtr<const Trace::IAnalysisSession> Session = InsightsModule->GetAnalysisSession();
+	TSharedPtr<const TraceServices::IAnalysisSession> Session = InsightsModule->GetAnalysisSession();
 	if (Session.IsValid())
 	{
-		Trace::FAnalysisSessionReadScope SessionReadScope(*Session.Get());
+		TraceServices::FAnalysisSessionReadScope SessionReadScope(*Session.Get());
 		const INetworkPredictionProvider* NetworkPredictionProvider = ReadNetworkPredictionProvider(*Session.Get());
 		if (NetworkPredictionProvider)
 		{
@@ -667,4 +667,4 @@ TSharedRef<SWidget> SNPWindow::PIESessionComboBox_OnGenerateWidget(TSharedPtr<in
 
 ////////////////////////////////////////////////////////////////////////////////////////////////
 
-#undef LOCTEXT_NAMESPACE
\ No newline at end of file
+#undef LOCTEXT_NAMESPACE
diff --git a/Engine/Plugins/Runtime/NetworkPredictionInsights/Source/NetworkPredictionInsights/Private/UI/SNPWindow.h b/Engine/Plugins/Runtime/NetworkPredictionInsights/Source/NetworkPredictionInsights/Private/UI/SNPWindow.h
--- a/Engine/Plugins/Runtime/NetworkPredictionInsights/Source/NetworkPredictionInsights/Private/UI/SNPWindow.h
+++ b/Engine/Plugins/Runtime/NetworkPredictionInsights/Source/NetworkPredictionInsights/Private/UI/SNPWindow.h
@@ -19,7 +19,7 @@
 #include "Containers/ArrayView.h"
 #include "Framework/Commands/UICommandList.h"
 
-namespace Trace
+namespace TraceServices
 {
 	class IAnalysisService;
 	class IAnalysisSession;
diff --git a/Engine/Plugins/Runtime/NetworkPredictionInsights/Source/NetworkPredictionInsights/Public/INetworkPredictionProvider.h b/Engine/Plugins/Runtime/NetworkPredictionInsights/Source/NetworkPredictionInsights/Public/INetworkPredictionProvider.h
--- a/Engine/Plugins/Runtime/NetworkPredictionInsights/Source/NetworkPredictionInsights/Public/INetworkPredictionProvider.h
+++ b/Engine/Plugins/Runtime/NetworkPredictionInsights/Source/NetworkPredictionInsights/Public/INetworkPredictionProvider.h
@@ -125,14 +125,14 @@ private:
 	bool Valid() const { return CurrentItemIdx >= PageArrayView.StartItemIdx && CurrentItemIdx < PageArrayView.EndItemIdx; }
 
 	const TRestrictedPageArrayView<T>& PageArrayView;
-	mutable TPagedArrayIterator<T, typename TPagedArray<T>::PageType> It;
+	mutable TraceServices::TPagedArrayIterator<T, typename TraceServices::TPagedArray<T>::PageType> It;
 	uint64 CurrentItemIdx; // Kind of unfortunate but we store this ourselves to avoid calculating it off the iter each time
 };
 
 template<typename T>
 struct TRestrictedPageArrayView
 {
-	TRestrictedPageArrayView(const TPagedArray<T>& InPageArray, uint64 InStart, uint64 InEnd)
+	TRestrictedPageArrayView(const TraceServices::TPagedArray<T>& InPageArray, uint64 InStart, uint64 InEnd)
 		: PageArray(InPageArray), StartItemIdx(InStart), EndItemIdx(InEnd) { }
 
 	TRestrictedPageViewIterator<T> GetIterator() const
@@ -171,7 +171,7 @@ private:
 	template<typename>
 	friend struct TRestrictedPageViewIterator;
 
-	const TPagedArray<T>& PageArray;
+	const TraceServices::TPagedArray<T>& PageArray;
 	uint64 StartItemIdx;
 	uint64 EndItemIdx;
 };
@@ -191,12 +191,12 @@ typename TEnableIf<TIsPointer<T>::Value, uint64>::Type GetEngineFrame(const T& E
 
 // Searches through PageArray for index of first element >= MinFrameNumber
 template<typename T>
-uint64 FindMinIndexPagedArray(const TPagedArray<T>& PageArray, uint64 MinFrameNumber)
+uint64 FindMinIndexPagedArray(const TraceServices::TPagedArray<T>& PageArray, uint64 MinFrameNumber)
 {
 	// Flip through pages first before iterating individual items
 	for (uint64 PageIdx=0; PageIdx < PageArray.NumPages(); ++PageIdx)
 	{
-		auto* Page = const_cast<TPagedArray<T>&>(PageArray).GetPage(PageIdx);
+		auto* Page = const_cast<TraceServices::TPagedArray<T>&>(PageArray).GetPage(PageIdx);
 		check(Page->Count >0);
 
 		//if (Page->Items[Page->Count-1].EngineFrame >= MinFrameNumber)
@@ -218,7 +218,7 @@ uint64 FindMinIndexPagedArray(const TPagedArray<T>& PageArray, uint64 MinFrameNu
 // Searches through PageArray for index+1 of first element > MinFrameNumber. E.g, index to iterate up to (<) for valid EngineFrameNumber.
 // Fixme: it probably makes more sense to just give this a start index and pickup where FindMinIndex left off
 template<typename T>
-uint64 FindMaxIndexPagedArray(const TPagedArray<T>& PageArray, uint64 MaxFrameNumber)
+uint64 FindMaxIndexPagedArray(const TraceServices::TPagedArray<T>& PageArray, uint64 MaxFrameNumber)
 {
 	if (PageArray.Num() == 0)
 	{
@@ -233,7 +233,7 @@ uint64 FindMaxIndexPagedArray(const TPagedArray<T>& PageArray, uint64 MaxFrameNu
 	// Flip through pages first before iterating individual items
 	for (uint64 PageIdx=PageArray.NumPages()-1; ; --PageIdx)
 	{
-		auto* Page = const_cast<TPagedArray<T>&>(PageArray).GetPage(PageIdx);
+		auto* Page = const_cast<TraceServices::TPagedArray<T>&>(PageArray).GetPage(PageIdx);
 		check(Page->Count >0);
 
 		if (GetEngineFrame(Page->Items[0]) <= MaxFrameNumber)
@@ -336,7 +336,7 @@ private:
 // Holds all data we traced for a given simulation.
 struct FSimulationData
 {
-	FSimulationData(int32 InTraceID, Trace::ILinearAllocator& Allocator)
+	FSimulationData(int32 InTraceID, TraceServices::ILinearAllocator& Allocator)
 		: TraceID(InTraceID)
 		, Ticks(Allocator, 1024)
 		, NetRecv(Allocator, 1024)
@@ -427,7 +427,7 @@ struct FSimulationData
 
 	struct FUserStateStore
 	{
-		FUserStateStore(Trace::ILinearAllocator& Allocator)
+		FUserStateStore(TraceServices::ILinearAllocator& Allocator)
 			: UserStates(Allocator, 1024) { }
 		
 		FORCENOINLINE FUserState& Push(int32 Frame, uint64 EngineFrame)
@@ -509,14 +509,14 @@ struct FSimulationData
 		//	-User states should be pushed in ascending EngineFrame order but NOT necessarily in accessing SimulationFram 
 		//	-We don't need to hyper optimize this but we should avoid linear searching 
 
-		TPagedArray<FUserState> UserStates;	// Holds the actual FUserState. Pointers to these elements are stable
+		TraceServices::TPagedArray<FUserState> UserStates;	// Holds the actual FUserState. Pointers to these elements are stable
 		TMap<int32, FUserState*> FrameMap;	// Acceleration map to find head FUserState for given simulation frame
 		TBitArray<> PopulatedFrames;		// Bit array for which simulation frames are populated. Accelerates finding "last valid" state in sparse cases.
 	};
 
 	struct FUserData
 	{
-		FUserData(Trace::ILinearAllocator& Allocator)
+		FUserData(TraceServices::ILinearAllocator& Allocator)
 			: Store{Allocator, Allocator, Allocator, Allocator}
 		{ }
 
@@ -598,8 +598,8 @@ struct FSimulationData
 
 	int32 TraceID;
 	
-	TPagedArray<FTick> Ticks;
-	TPagedArray<FNetSerializeRecv> NetRecv;	// Actually holds the allocated recv records
+	TraceServices::TPagedArray<FTick> Ticks;
+	TraceServices::TPagedArray<FNetSerializeRecv> NetRecv;	// Actually holds the allocated recv records
 	
 	TSparseFrameData<FSparse> SparseData;
 	FConst ConstData;
@@ -632,7 +632,7 @@ struct FSimulationData
 	} Analysis;
 };
 
-class INetworkPredictionProvider : public Trace::IProvider
+class INetworkPredictionProvider : public TraceServices::IProvider
 {
 public:
 
@@ -647,4 +647,4 @@ public:
 	virtual TArrayView<const TSharedRef<FSimulationData>> ReadSimulationData() const = 0;
 };
 
-NETWORKPREDICTIONINSIGHTS_API const INetworkPredictionProvider* ReadNetworkPredictionProvider(const Trace::IAnalysisSession& Session);
\ No newline at end of file
+NETWORKPREDICTIONINSIGHTS_API const INetworkPredictionProvider* ReadNetworkPredictionProvider(const TraceServices::IAnalysisSession& Session);
diff --git a/Engine/Plugins/Slate/SlateInsights/Source/SlateInsights/Private/SSlateFrameSchematicView.cpp b/Engine/Plugins/Slate/SlateInsights/Source/SlateInsights/Private/SSlateFrameSchematicView.cpp
--- a/Engine/Plugins/Slate/SlateInsights/Source/SlateInsights/Private/SSlateFrameSchematicView.cpp
+++ b/Engine/Plugins/Slate/SlateInsights/Source/SlateInsights/Private/SSlateFrameSchematicView.cpp
@@ -169,14 +169,14 @@ namespace Private
 		}
 	};
 
-	FText GetWidgetName(const Trace::IAnalysisSession* AnalysisSession, Message::FWidgetId WidgetId)
+	FText GetWidgetName(const TraceServices::IAnalysisSession* AnalysisSession, Message::FWidgetId WidgetId)
 	{
 		if (AnalysisSession)
 		{
 			const FSlateProvider* SlateProvider = AnalysisSession->ReadProvider<FSlateProvider>(FSlateProvider::ProviderName);
 			if (SlateProvider)
 			{
-				Trace::FAnalysisSessionReadScope SessionReadScope(*AnalysisSession);
+				TraceServices::FAnalysisSessionReadScope SessionReadScope(*AnalysisSession);
 
 				if (const Message::FWidgetInfo* WidgetInfo = SlateProvider->FindWidget(WidgetId))
 				{
@@ -321,7 +321,7 @@ SSlateFrameSchematicView::~SSlateFrameSchematicView()
 	}
 }
 
-void SSlateFrameSchematicView::SetSession(Insights::ITimingViewSession* InTimingViewSession, const Trace::IAnalysisSession* InAnalysisSession)
+void SSlateFrameSchematicView::SetSession(Insights::ITimingViewSession* InTimingViewSession, const TraceServices::IAnalysisSession* InAnalysisSession)
 {
 	if (TimingViewSession)
 	{
@@ -576,7 +576,7 @@ void SSlateFrameSchematicView::RefreshNodes()
 			const FSlateProvider* SlateProvider = AnalysisSession->ReadProvider<FSlateProvider>(FSlateProvider::ProviderName);
 			if (SlateProvider)
 			{
-				Trace::FAnalysisSessionReadScope SessionReadScope(*AnalysisSession);
+				TraceServices::FAnalysisSessionReadScope SessionReadScope(*AnalysisSession);
 
 				if (FMath::IsNearlyEqual(StartTime, EndTime))
 				{
@@ -589,7 +589,7 @@ void SSlateFrameSchematicView::RefreshNodes()
 						{
 							this->StartTime = FMath::Min(StartTime, EventStartTime);
 							this->EndTime = FMath::Max(EndTime, EventEndTime+Message.DeltaTime);
-							return Trace::EEventEnumerate::Continue;
+							return TraceServices::EEventEnumerate::Continue;
 						});
 				}
 
@@ -670,7 +670,7 @@ void SSlateFrameSchematicView::RefreshNodes_Invalidation(const FSlateProvider* S
 				WidgetInfo->Callstack = Message.Callstack;
 			}
 
-			return Trace::EEventEnumerate::Continue;
+			return TraceServices::EEventEnumerate::Continue;
 		});
 
 	for (const auto& Itt : InvalidationMap)
@@ -702,7 +702,7 @@ void SSlateFrameSchematicView::RefreshNodes_Update(const FSlateProvider* SlatePr
 					Message.WidgetId,
 					MakeShared<Private::FWidgetUpdateInfo>(Message.WidgetId, Message.UpdateFlags));
 			}
-			return Trace::EEventEnumerate::Continue;
+			return TraceServices::EEventEnumerate::Continue;
 		});
 	WidgetUpdateInfosMap.GenerateValueArray(WidgetUpdateInfos);
 
diff --git a/Engine/Plugins/Slate/SlateInsights/Source/SlateInsights/Private/SSlateFrameSchematicView.h b/Engine/Plugins/Slate/SlateInsights/Source/SlateInsights/Private/SSlateFrameSchematicView.h
--- a/Engine/Plugins/Slate/SlateInsights/Source/SlateInsights/Private/SSlateFrameSchematicView.h
+++ b/Engine/Plugins/Slate/SlateInsights/Source/SlateInsights/Private/SSlateFrameSchematicView.h
@@ -10,7 +10,7 @@
 
 namespace Insights { class ITimingViewSession; }
 namespace Insights { enum class ETimeChangedFlags; }
-namespace Trace { class IAnalysisSession; }
+namespace TraceServices { class IAnalysisSession; }
 class FAnimationSharedData;
 class IInsightsManager;
 class ITableRow;
@@ -37,7 +37,7 @@ public:
 
 	void Construct(const FArguments& InArgs);
 
-	void SetSession(Insights::ITimingViewSession* InTimingViewSession, const Trace::IAnalysisSession* InAnalysisSession);
+	void SetSession(Insights::ITimingViewSession* InTimingViewSession, const TraceServices::IAnalysisSession* InAnalysisSession);
 
 private:
 	TSharedRef<ITableRow> HandleUniqueInvalidatedMakeTreeRowWidget(TSharedPtr<Private::FWidgetUniqueInvalidatedInfo> InInfo, const TSharedRef<STableViewBase>& OwnerTable);
@@ -60,7 +60,7 @@ private:
 	void RefreshNodes_Update(const FSlateProvider* SlateProvider);
 
 private:
-	const Trace::IAnalysisSession* AnalysisSession;
+	const TraceServices::IAnalysisSession* AnalysisSession;
 	Insights::ITimingViewSession* TimingViewSession;
 
 	TSharedPtr<STreeView<TSharedPtr<Private::FWidgetUniqueInvalidatedInfo>>> WidgetInvalidateInfoListView;
diff --git a/Engine/Plugins/Slate/SlateInsights/Source/SlateInsights/Private/SlateAnalyzer.cpp b/Engine/Plugins/Slate/SlateInsights/Source/SlateInsights/Private/SlateAnalyzer.cpp
--- a/Engine/Plugins/Slate/SlateInsights/Source/SlateInsights/Private/SlateAnalyzer.cpp
+++ b/Engine/Plugins/Slate/SlateInsights/Source/SlateInsights/Private/SlateAnalyzer.cpp
@@ -13,7 +13,7 @@ namespace UE
 namespace SlateInsights
 {
 
-FSlateAnalyzer::FSlateAnalyzer(Trace::IAnalysisSession& InSession, FSlateProvider& InSlateProvider)
+FSlateAnalyzer::FSlateAnalyzer(TraceServices::IAnalysisSession& InSession, FSlateProvider& InSlateProvider)
 	: Session(InSession)
 	, SlateProvider(InSlateProvider)
 {
@@ -35,7 +35,7 @@ void FSlateAnalyzer::OnAnalysisBegin(const FOnAnalysisContext& Context)
 
 bool FSlateAnalyzer::OnEvent(uint16 RouteId, EStyle Style, const FOnEventContext& Context)
 {
-	Trace::FAnalysisSessionEditScope _(Session);
+	TraceServices::FAnalysisSessionEditScope _(Session);
 
 	const Trace::IAnalyzer::FEventData& EventData = Context.EventData;
 	switch (RouteId)
diff --git a/Engine/Plugins/Slate/SlateInsights/Source/SlateInsights/Private/SlateAnalyzer.h b/Engine/Plugins/Slate/SlateInsights/Source/SlateInsights/Private/SlateAnalyzer.h
--- a/Engine/Plugins/Slate/SlateInsights/Source/SlateInsights/Private/SlateAnalyzer.h
+++ b/Engine/Plugins/Slate/SlateInsights/Source/SlateInsights/Private/SlateAnalyzer.h
@@ -4,7 +4,7 @@
 
 #include "Trace/Analyzer.h"
 
-namespace Trace { class IAnalysisSession; }
+namespace TraceServices { class IAnalysisSession; }
 
 namespace UE
 {
@@ -16,7 +16,7 @@ class FSlateProvider;
 class FSlateAnalyzer : public Trace::IAnalyzer
 {
 public:
-	FSlateAnalyzer(Trace::IAnalysisSession& InSession, FSlateProvider& InSlateProvider);
+	FSlateAnalyzer(TraceServices::IAnalysisSession& InSession, FSlateProvider& InSlateProvider);
 
 	virtual void OnAnalysisBegin(const FOnAnalysisContext& Context) override;
 	virtual void OnAnalysisEnd() override {}
@@ -35,7 +35,7 @@ private:
 		RouteId_RootChildOrderInvalidated,
 	};
 
-	Trace::IAnalysisSession& Session;
+	TraceServices::IAnalysisSession& Session;
 	FSlateProvider& SlateProvider;
 };
 
diff --git a/Engine/Plugins/Slate/SlateInsights/Source/SlateInsights/Private/SlateFrameGraphTrack.cpp b/Engine/Plugins/Slate/SlateInsights/Source/SlateInsights/Private/SlateFrameGraphTrack.cpp
--- a/Engine/Plugins/Slate/SlateInsights/Source/SlateInsights/Private/SlateFrameGraphTrack.cpp
+++ b/Engine/Plugins/Slate/SlateInsights/Source/SlateInsights/Private/SlateFrameGraphTrack.cpp
@@ -387,7 +387,7 @@ void FSlateFrameGraphTrack::AddAllSeries(const ITimingTrackUpdateContext& Contex
 	bool bFirstSeries = AllSeries.Num() == 0;
 	if(SlateProvider && bFirstSeries)
 	{
-		Trace::FAnalysisSessionReadScope SessionReadScope(SharedData.GetAnalysisSession());
+		TraceServices::FAnalysisSessionReadScope SessionReadScope(SharedData.GetAnalysisSession());
 		FSlateStyleSet& StyleSet = FSlateInsightsStyle::Get();
 		{
 			
@@ -462,7 +462,7 @@ void FSlateFrameGraphTrack::UpdateSeries(const FTimingTrackViewport& InViewport,
 {
 	if (const FSlateProvider* SlateProvider = SharedData.GetAnalysisSession().ReadProvider<FSlateProvider>(FSlateProvider::ProviderName))
 	{
-		Trace::FAnalysisSessionReadScope SessionReadScope(SharedData.GetAnalysisSession());
+		TraceServices::FAnalysisSessionReadScope SessionReadScope(SharedData.GetAnalysisSession());
 
 		for (TSharedPtr<Private::FSlateFrameGraphSeries>& Serie : Series)
 		{
@@ -480,7 +480,7 @@ void FSlateFrameGraphTrack::UpdateSeries(const FTimingTrackViewport& InViewport,
 					{
 						Serie->AddEvent(StartTime, EndTime, Message);
 					}
-					return Trace::EEventEnumerate::Continue;
+					return TraceServices::EEventEnumerate::Continue;
 				});
 		}
 
diff --git a/Engine/Plugins/Slate/SlateInsights/Source/SlateInsights/Private/SlateInsightsModule.cpp b/Engine/Plugins/Slate/SlateInsights/Source/SlateInsights/Private/SlateInsightsModule.cpp
--- a/Engine/Plugins/Slate/SlateInsights/Source/SlateInsights/Private/SlateInsightsModule.cpp
+++ b/Engine/Plugins/Slate/SlateInsights/Source/SlateInsights/Private/SlateInsightsModule.cpp
@@ -36,7 +36,7 @@ FSlateInsightsModule& FSlateInsightsModule::Get()
 
 void FSlateInsightsModule::StartupModule()
 {
-	IModularFeatures::Get().RegisterModularFeature(Trace::ModuleFeatureName, &TraceModule);
+	IModularFeatures::Get().RegisterModularFeature(TraceServices::ModuleFeatureName, &TraceModule);
 	IModularFeatures::Get().RegisterModularFeature(Insights::TimingViewExtenderFeatureName, &TimingViewExtender);
 
 	IUnrealInsightsModule& UnrealInsightsModule = FModuleManager::LoadModuleChecked<IUnrealInsightsModule>("TraceInsights");
@@ -52,7 +52,7 @@ void FSlateInsightsModule::ShutdownModule()
 		TimingProfilerLayoutExtension.RemoveAll(this);
 	}
 
-	IModularFeatures::Get().UnregisterModularFeature(Trace::ModuleFeatureName, &TraceModule);
+	IModularFeatures::Get().UnregisterModularFeature(TraceServices::ModuleFeatureName, &TraceModule);
 	IModularFeatures::Get().UnregisterModularFeature(Insights::TimingViewExtenderFeatureName, &TimingViewExtender);
 }
 
diff --git a/Engine/Plugins/Slate/SlateInsights/Source/SlateInsights/Private/SlateProvider.cpp b/Engine/Plugins/Slate/SlateInsights/Source/SlateInsights/Private/SlateProvider.cpp
--- a/Engine/Plugins/Slate/SlateInsights/Source/SlateInsights/Private/SlateProvider.cpp
+++ b/Engine/Plugins/Slate/SlateInsights/Source/SlateInsights/Private/SlateProvider.cpp
@@ -111,7 +111,7 @@ namespace Message
 	}
 } //namespace Message
 
-FSlateProvider::FSlateProvider(Trace::IAnalysisSession& InSession)
+FSlateProvider::FSlateProvider(TraceServices::IAnalysisSession& InSession)
 	: Session(InSession)
 	, WidgetTimelines(Session.GetLinearAllocator())
 	, ApplicationTickedTimeline(Session.GetLinearAllocator())
@@ -181,4 +181,4 @@ void FSlateProvider::AddWidgetInvalidatedEvent(double Seconds, Message::FWidgetI
 } //namespace SlateInsights
 } //namespace UE
 
-#undef LOCTEXT_NAMESPACE
\ No newline at end of file
+#undef LOCTEXT_NAMESPACE
diff --git a/Engine/Plugins/Slate/SlateInsights/Source/SlateInsights/Private/SlateProvider.h b/Engine/Plugins/Slate/SlateInsights/Source/SlateInsights/Private/SlateProvider.h
--- a/Engine/Plugins/Slate/SlateInsights/Source/SlateInsights/Private/SlateProvider.h
+++ b/Engine/Plugins/Slate/SlateInsights/Source/SlateInsights/Private/SlateProvider.h
@@ -16,7 +16,7 @@
 #include "Trace/Analyzer.h"
 #include "Trace/SlateTrace.h"
 
-namespace Trace { class IAnalysisSession; }
+namespace TraceServices { class IAnalysisSession; }
 
 namespace UE
 {
@@ -96,12 +96,12 @@ struct FApplicationTickedMessage
 
 } //namespace Message
 
-class FSlateProvider : public Trace::IProvider
+class FSlateProvider : public TraceServices::IProvider
 {
 public:
 	static FName ProviderName;
 
-	FSlateProvider(Trace::IAnalysisSession& InSession);
+	FSlateProvider(TraceServices::IAnalysisSession& InSession);
 
 	/** */
 	void AddWidget(double Seconds, uint64 WidgetId);
@@ -114,7 +114,7 @@ public:
 	void AddWidgetInvalidatedEvent(double Seconds, Message::FWidgetInvalidatedMessage InvalidatedMessage);
 
 	/** */
-	using TApplicationTickedTimeline = Trace::TPointTimeline<Message::FApplicationTickedMessage>;
+	using TApplicationTickedTimeline = TraceServices::TPointTimeline<Message::FApplicationTickedMessage>;
 	const TApplicationTickedTimeline& GetApplicationTickedTimeline() const
 	{
 		Session.ReadAccessCheck();
@@ -122,7 +122,7 @@ public:
 	}
 
 	/** */
-	using TWidgetUpdatedTimeline = Trace::TPointTimeline<Message::FWidgetUpdatedMessage>;
+	using TWidgetUpdatedTimeline = TraceServices::TPointTimeline<Message::FWidgetUpdatedMessage>;
 	const TWidgetUpdatedTimeline& GetWidgetUpdatedTimeline() const
 	{
 		Session.ReadAccessCheck();
@@ -130,7 +130,7 @@ public:
 	}
 	
 	/** */
-	using TWidgetInvalidatedTimeline = Trace::TPointTimeline<Message::FWidgetInvalidatedMessage>;
+	using TWidgetInvalidatedTimeline = TraceServices::TPointTimeline<Message::FWidgetInvalidatedMessage>;
 	const TWidgetInvalidatedTimeline& GetWidgetInvalidatedTimeline() const
 	{
 		Session.ReadAccessCheck();
@@ -138,7 +138,7 @@ public:
 	}
 
 	/** */
-	using TWidgetTimeline = Trace::TIntervalTimeline<Message::FWidgetId>;
+	using TWidgetTimeline = TraceServices::TIntervalTimeline<Message::FWidgetId>;
 	const TWidgetTimeline& GetWidgetTimeline() const
 	{
 		Session.ReadAccessCheck();
@@ -171,7 +171,7 @@ public:
 	}
 
 private:
-	Trace::IAnalysisSession& Session;
+	TraceServices::IAnalysisSession& Session;
 
 	TMap<Message::FWidgetId, Message::FWidgetInfo> WidgetInfos;
 
diff --git a/Engine/Plugins/Slate/SlateInsights/Source/SlateInsights/Private/SlateTimingViewExtender.cpp b/Engine/Plugins/Slate/SlateInsights/Source/SlateInsights/Private/SlateTimingViewExtender.cpp
--- a/Engine/Plugins/Slate/SlateInsights/Source/SlateInsights/Private/SlateTimingViewExtender.cpp
+++ b/Engine/Plugins/Slate/SlateInsights/Source/SlateInsights/Private/SlateTimingViewExtender.cpp
@@ -39,7 +39,7 @@ void FSlateTimingViewExtender::OnEndSession(Insights::ITimingViewSession& InSess
 	PerSessionDataMap.Remove(&InSession);
 }
 
-void FSlateTimingViewExtender::Tick(Insights::ITimingViewSession& InSession, const Trace::IAnalysisSession& InAnalysisSession)
+void FSlateTimingViewExtender::Tick(Insights::ITimingViewSession& InSession, const TraceServices::IAnalysisSession& InAnalysisSession)
 {
 	FPerSessionData* PerSessionData = PerSessionDataMap.Find(&InSession);
 	if(PerSessionData != nullptr)
diff --git a/Engine/Plugins/Slate/SlateInsights/Source/SlateInsights/Private/SlateTimingViewExtender.h b/Engine/Plugins/Slate/SlateInsights/Source/SlateInsights/Private/SlateTimingViewExtender.h
--- a/Engine/Plugins/Slate/SlateInsights/Source/SlateInsights/Private/SlateTimingViewExtender.h
+++ b/Engine/Plugins/Slate/SlateInsights/Source/SlateInsights/Private/SlateTimingViewExtender.h
@@ -6,7 +6,7 @@
 #include "SlateTimingViewSession.h"
 
 namespace Insights { class ITimingViewSession; }
-namespace Insights { class IAnalysisSession; }
+namespace TraceServices { class IAnalysisSession; }
 
 class FMenuBuilder;
 
@@ -21,7 +21,7 @@ public:
 	//~ Begin Insights::ITimingViewExtender interface
 	virtual void OnBeginSession(Insights::ITimingViewSession& InSession) override;
 	virtual void OnEndSession(Insights::ITimingViewSession& InSession) override;
-	virtual void Tick(Insights::ITimingViewSession& InSession, const Trace::IAnalysisSession& InAnalysisSession) override;
+	virtual void Tick(Insights::ITimingViewSession& InSession, const TraceServices::IAnalysisSession& InAnalysisSession) override;
 	virtual void ExtendFilterMenu(Insights::ITimingViewSession& InSession, FMenuBuilder& InMenuBuilder) override;
 	//~ End Insights::ITimingViewExtender interface
 
diff --git a/Engine/Plugins/Slate/SlateInsights/Source/SlateInsights/Private/SlateTimingViewSession.cpp b/Engine/Plugins/Slate/SlateInsights/Source/SlateInsights/Private/SlateTimingViewSession.cpp
--- a/Engine/Plugins/Slate/SlateInsights/Source/SlateInsights/Private/SlateTimingViewSession.cpp
+++ b/Engine/Plugins/Slate/SlateInsights/Source/SlateInsights/Private/SlateTimingViewSession.cpp
@@ -47,7 +47,7 @@ void FSlateTimingViewSession::OnEndSession(Insights::ITimingViewSession& InTimin
 	TimingViewSession = nullptr;
 }
 
-void FSlateTimingViewSession::Tick(Insights::ITimingViewSession& InTimingViewSession, const Trace::IAnalysisSession& InAnalysisSession)
+void FSlateTimingViewSession::Tick(Insights::ITimingViewSession& InTimingViewSession, const TraceServices::IAnalysisSession& InAnalysisSession)
 {
 	AnalysisSession = &InAnalysisSession;
 
@@ -55,7 +55,7 @@ void FSlateTimingViewSession::Tick(Insights::ITimingViewSession& InTimingViewSes
 
 	if (SlateProvider)
 	{
-		Trace::FAnalysisSessionReadScope SessionReadScope(GetAnalysisSession());
+		TraceServices::FAnalysisSessionReadScope SessionReadScope(GetAnalysisSession());
 
 		// Add "stat" track
 		const FSlateProvider::TApplicationTickedTimeline& TickedTimeline = SlateProvider->GetApplicationTickedTimeline();
diff --git a/Engine/Plugins/Slate/SlateInsights/Source/SlateInsights/Private/SlateTimingViewSession.h b/Engine/Plugins/Slate/SlateInsights/Source/SlateInsights/Private/SlateTimingViewSession.h
--- a/Engine/Plugins/Slate/SlateInsights/Source/SlateInsights/Private/SlateTimingViewSession.h
+++ b/Engine/Plugins/Slate/SlateInsights/Source/SlateInsights/Private/SlateTimingViewSession.h
@@ -5,7 +5,7 @@
 #include "Insights/ViewModels/TimingEventsTrack.h"
 #include "TraceServices/Model/Frames.h"
 
-namespace Trace { class IAnalysisSession; }
+namespace TraceServices { class IAnalysisSession; }
 namespace Insights { class ITimingViewSession; }
 namespace Insights { enum class ETimeChangedFlags : int32; }
 namespace UE { namespace SlateInsights { class FSlateFrameGraphTrack; } }
@@ -24,11 +24,11 @@ public:
 
 	void OnBeginSession(Insights::ITimingViewSession& InTimingViewSession);
 	void OnEndSession(Insights::ITimingViewSession& InTimingViewSession);
-	void Tick(Insights::ITimingViewSession& InTimingViewSession, const Trace::IAnalysisSession& InAnalysisSession);
+	void Tick(Insights::ITimingViewSession& InTimingViewSession, const TraceServices::IAnalysisSession& InAnalysisSession);
 	void ExtendFilterMenu(FMenuBuilder& InMenuBuilder);
 
 	/** Get the last cached analysis session */
-	const Trace::IAnalysisSession& GetAnalysisSession() const { return *AnalysisSession; }
+	const TraceServices::IAnalysisSession& GetAnalysisSession() const { return *AnalysisSession; }
 
 	/** Check whether the analysis session is valid */
 	bool IsAnalysisSessionValid() const { return AnalysisSession != nullptr; }
@@ -41,7 +41,7 @@ public:
 
 private:
 	// Cached analysis session, set in Tick()
-	const Trace::IAnalysisSession* AnalysisSession;
+	const TraceServices::IAnalysisSession* AnalysisSession;
 
 	// Cached timing view session, set in OnBeginSession/OnEndSession
 	Insights::ITimingViewSession* TimingViewSession;
