Add montage position details to animation insights

diff --git a/Engine/Plugins/Animation/GameplayInsights/Source/GameplayInsights/Private/AnimationAnalyzer.cpp b/Engine/Plugins/Animation/GameplayInsights/Source/GameplayInsights/Private/AnimationAnalyzer.cpp
--- a/Engine/Plugins/Animation/GameplayInsights/Source/GameplayInsights/Private/AnimationAnalyzer.cpp
+++ b/Engine/Plugins/Animation/GameplayInsights/Source/GameplayInsights/Private/AnimationAnalyzer.cpp
@@ -327,8 +327,9 @@ bool FAnimationAnalyzer::OnEvent(uint16 RouteId, EStyle Style, const FOnEventCon
 		uint32 NextSectionNameId = EventData.GetValue<uint32>("NextSectionNameId");
 		float Weight = EventData.GetValue<float>("Weight");
 		float DesiredWeight = EventData.GetValue<float>("DesiredWeight");
+		float Position = EventData.GetValue<float>("Position");
 		uint16 FrameCounter = EventData.GetValue<uint16>("FrameCounter");
-		AnimationProvider.AppendMontage(AnimInstanceId, Context.EventTime.AsSeconds(Cycle), MontageId, CurrentSectionNameId, NextSectionNameId, Weight, DesiredWeight, FrameCounter);
+		AnimationProvider.AppendMontage(AnimInstanceId, Context.EventTime.AsSeconds(Cycle), MontageId, CurrentSectionNameId, NextSectionNameId, Weight, DesiredWeight, Position, FrameCounter);
 		break;
 	}
 	case RouteId_Sync:
diff --git a/Engine/Plugins/Animation/GameplayInsights/Source/GameplayInsights/Private/AnimationProvider.cpp b/Engine/Plugins/Animation/GameplayInsights/Source/GameplayInsights/Private/AnimationProvider.cpp
--- a/Engine/Plugins/Animation/GameplayInsights/Source/GameplayInsights/Private/AnimationProvider.cpp
+++ b/Engine/Plugins/Animation/GameplayInsights/Source/GameplayInsights/Private/AnimationProvider.cpp
@@ -1077,7 +1077,7 @@ void FAnimationProvider::AppendNotify(uint64 InAnimInstanceId, double InTime, ui
 	Session.UpdateDurationSeconds(InTime);
 }
 
-void FAnimationProvider::AppendMontage(uint64 InAnimInstanceId, double InTime, uint64 InMontageId, uint32 InCurrentSectionNameId, uint32 InNextSectionNameId, float InWeight, float InDesiredWeight, uint16 InFrameCounter)
+void FAnimationProvider::AppendMontage(uint64 InAnimInstanceId, double InTime, uint64 InMontageId, uint32 InCurrentSectionNameId, uint32 InNextSectionNameId, float InWeight, float InDesiredWeight, float InPosition, uint16 InFrameCounter)
 {
 	Session.WriteAccessCheck();
 
@@ -1107,6 +1107,7 @@ void FAnimationProvider::AppendMontage(uint64 InAnimInstanceId, double InTime, u
 	Message.NextSectionNameId = InNextSectionNameId;
 	Message.Weight = InWeight;
 	Message.DesiredWeight = InDesiredWeight;
+	Message.Position = InPosition;
 	Message.FrameCounter = InFrameCounter;
 
 	TimelineStorage->Timeline->AppendEvent(InTime, Message);
diff --git a/Engine/Plugins/Animation/GameplayInsights/Source/GameplayInsights/Private/AnimationProvider.h b/Engine/Plugins/Animation/GameplayInsights/Source/GameplayInsights/Private/AnimationProvider.h
--- a/Engine/Plugins/Animation/GameplayInsights/Source/GameplayInsights/Private/AnimationProvider.h
+++ b/Engine/Plugins/Animation/GameplayInsights/Source/GameplayInsights/Private/AnimationProvider.h
@@ -99,7 +99,7 @@ public:
 	void AppendNotify(uint64 InAnimInstanceId, double InTime, uint64 InAssetId, uint64 InNotifyId, uint32 InNameId, float InNotifyTime, float InNotifyDuration, EAnimNotifyMessageType InNotifyEventType);
 
 	/** Append montage data */
-	void AppendMontage(uint64 InAnimInstanceId, double InTime, uint64 InMontageId, uint32 InCurrentSectionNameId, uint32 InNextSectionNameId, float InWeight, float InDesiredWeight, uint16 InFrameCounter);
+	void AppendMontage(uint64 InAnimInstanceId, double InTime, uint64 InMontageId, uint32 InCurrentSectionNameId, uint32 InNextSectionNameId, float InWeight, float InDesiredWeight, float InPosition, uint16 InFrameCounter);
 
 	/** Append sync data */
 	void AppendSync(uint64 InAnimInstanceId, double InTime, int32 InSourceNodeId, uint32 InGroupNameId);
diff --git a/Engine/Plugins/Animation/GameplayInsights/Source/GameplayInsights/Private/MontageTrack.cpp b/Engine/Plugins/Animation/GameplayInsights/Source/GameplayInsights/Private/MontageTrack.cpp
--- a/Engine/Plugins/Animation/GameplayInsights/Source/GameplayInsights/Private/MontageTrack.cpp
+++ b/Engine/Plugins/Animation/GameplayInsights/Source/GameplayInsights/Private/MontageTrack.cpp
@@ -9,6 +9,7 @@
 #include "Framework/MultiBox/MultiBoxBuilder.h"
 #include "Insights/ViewModels/TimingEventSearch.h"
 #include "Insights/ViewModels/TooltipDrawState.h"
+#include "Templates/Invoke.h"
 #include "Modules/ModuleManager.h"
 #include "Insights/ViewModels/GraphTrackBuilder.h"
 #include "VariantTreeNode.h"
@@ -28,6 +29,14 @@ FMontageTrack::FMontageTrack(const FAnimationSharedData& InSharedData, uint64 In
 
 void FMontageTrack::AddAllSeries()
 {
+	struct FSeriesDescription
+	{
+		FText Name;
+		FText Description;
+		FMontageSeries::ESeriesType Type;
+		bool bEnabled;
+	};
+
 	const FAnimationProvider* AnimationProvider = SharedData.GetAnalysisSession().ReadProvider<FAnimationProvider>(FAnimationProvider::ProviderName);
 	const FGameplayProvider* GameplayProvider = SharedData.GetAnalysisSession().ReadProvider<FGameplayProvider>(FGameplayProvider::ProviderName);
 
@@ -39,59 +48,85 @@ void FMontageTrack::AddAllSeries()
 		{
 			if(!AllSeries.ContainsByPredicate([InMontageId](const TSharedPtr<FGraphSeries>& InSeries){ return StaticCastSharedPtr<FMontageSeries>(InSeries)->MontageId == InMontageId; }))
 			{
-				auto MakeCurveSeriesColor = [](uint64 InSeed, bool bInLine)
+				static const FSeriesDescription SeriesDescriptions[] =
+				{
+					{
+						LOCTEXT("SeriesNameWeight", "Weight"),
+						LOCTEXT("SeriesDescWeight", "The final effective weight that this montage was played at"),
+						FMontageSeries::ESeriesType::Weight,
+						true
+					},
+					{
+						LOCTEXT("SeriesNamePosition", "Position"),
+						LOCTEXT("SeriesDescPosition", "The playback position of this montage"),
+						FMontageSeries::ESeriesType::Position,
+						false
+					},
+				};
+				auto MakeCurveSeriesColor = [](uint32 Hash, bool bInLine)
 				{
-					FRandomStream Stream(GetTypeHash(InSeed));
-					const uint8 Hue = (uint8)(Stream.FRand() * 255.0f);
+					FRandomStream Stream(Hash);
+					const uint8 Hue = (uint8)Stream.RandHelper(255);
 					const uint8 SatVal = bInLine ? 196 : 128;
 					return FLinearColor::MakeFromHSV8(Hue, SatVal, SatVal);
 				};
 
-				TSharedRef<FMontageSeries> Series = MakeShared<FMontageSeries>();
-
 				const FObjectInfo& MontageInfo = GameplayProvider->GetObjectInfo(InMontageId);
 
-				Series->SetName(MontageInfo.Name);
-				Series->SetDescription(FText::Format(LOCTEXT("MontageTooltipFormat", "Weight for montage '{0}'"), FText::FromString(MontageInfo.Name)));
-
-				const FLinearColor LineColor = MakeCurveSeriesColor(InMontageId, true);
-				const FLinearColor FillColor = MakeCurveSeriesColor(InMontageId, false);
-				Series->SetColor(LineColor, LineColor, FillColor);
+				auto AddSeries = [this, GameplayProvider, MakeCurveSeriesColor](const FSeriesDescription& InSeriesDescription, const TCHAR* InAssetName, uint64 InMontageId)
+				{
+					FString SeriesName = FText::Format(LOCTEXT("SeriesNameFormat", "{0} - {1}"), FText::FromString(InAssetName), InSeriesDescription.Name).ToString();
+
+					TSharedRef<FMontageSeries> Series = MakeShared<FMontageSeries>();
+
+					Series->SetName(SeriesName);
+					Series->SetDescription(InSeriesDescription.Description.ToString());
+					uint32 Hash = GetTypeHash(SeriesName) % 255;
+					const FLinearColor LineColor = MakeCurveSeriesColor(Hash, true);
+					const FLinearColor FillColor = MakeCurveSeriesColor(Hash, false);
+					Series->SetColor(LineColor, LineColor, FillColor);
+
+					Series->MontageId = InMontageId;
+					Series->Type = InSeriesDescription.Type;
+					Series->SetVisibility(true);
+					Series->SetBaselineY(25.0f);
+					Series->SetScaleY(20.0f);
+					Series->EnableAutoZoom();
+					AllSeries.Add(Series);
+				};
 
-				Series->MontageId = InMontageId;
-				Series->SetVisibility(true);
-				Series->SetBaselineY(25.0f);
-				Series->SetScaleY(20.0f);
-				Series->EnableAutoZoom();
-				AllSeries.Add(Series);
+				for (const FSeriesDescription& SeriesDescription : SeriesDescriptions)
+				{
+					AddSeries(SeriesDescription, MontageInfo.Name, InMontageId);
+				}
 			}
 		});
 	}
 }
 
-bool FMontageTrack::UpdateSeriesBounds(FGameplayGraphSeries& InSeries, const FTimingTrackViewport& InViewport)
+template<typename ProjectionType>
+bool FMontageTrack::UpdateSeriesBoundsHelper(FMontageSeries& InSeries, const FTimingTrackViewport& InViewport, ProjectionType Projection)
 {
 	bool bFoundEvents = false;
 
 	const FAnimationProvider* AnimationProvider = SharedData.GetAnalysisSession().ReadProvider<FAnimationProvider>(FAnimationProvider::ProviderName);
 
-	if(AnimationProvider)
+	if (AnimationProvider)
 	{
 		TraceServices::FAnalysisSessionReadScope SessionReadScope(SharedData.GetAnalysisSession());
 
-		FMontageSeries& MontageSeries = *static_cast<FMontageSeries*>(&InSeries);
-
-		MontageSeries.CurrentMin = 0.0f;
-		MontageSeries.CurrentMax = 0.0f;
+		InSeries.CurrentMin = 0.0;
+		InSeries.CurrentMax = 0.0;
 
-		AnimationProvider->ReadMontageTimeline(GetGameplayTrack().GetObjectId(), [&bFoundEvents, &InViewport, &MontageSeries, &AnimationProvider](const FAnimationProvider::AnimMontageTimeline& InTimeline)
+		AnimationProvider->ReadMontageTimeline(GetGameplayTrack().GetObjectId(), [&bFoundEvents, &InViewport, &InSeries, &Projection](const FAnimationProvider::AnimMontageTimeline& InTimeline)
 		{
-			InTimeline.EnumerateEvents(InViewport.GetStartTime(), InViewport.GetEndTime(), [&bFoundEvents, &MontageSeries, &AnimationProvider](double InStartTime, double InEndTime, uint32 InDepth, const FAnimMontageMessage& InMessage)
+			InTimeline.EnumerateEvents(InViewport.GetStartTime(), InViewport.GetEndTime(), [&bFoundEvents, &InSeries, &Projection](double InStartTime, double InEndTime, uint32 InDepth, const FAnimMontageMessage& InMessage)
 			{
-				if(InMessage.MontageId == MontageSeries.MontageId)
+				if (InMessage.MontageId == InSeries.MontageId )
 				{
-					MontageSeries.CurrentMin = FMath::Min(MontageSeries.CurrentMin, InMessage.Weight);
-					MontageSeries.CurrentMax = FMath::Max(MontageSeries.CurrentMax, InMessage.Weight);
+					const float Value = Invoke(Projection, InMessage);
+					InSeries.CurrentMin = FMath::Min(InSeries.CurrentMin, Value);
+					InSeries.CurrentMax = FMath::Max(InSeries.CurrentMax, Value);
 					bFoundEvents = true;
 				}
 				return TraceServices::EEventEnumerate::Continue;
@@ -102,32 +137,45 @@ bool FMontageTrack::UpdateSeriesBounds(FGameplayGraphSeries& InSeries, const FTi
 	return bFoundEvents;
 }
 
-void FMontageTrack::UpdateSeries(FGameplayGraphSeries& InSeries, const FTimingTrackViewport& InViewport)
+
+bool FMontageTrack::UpdateSeriesBounds(FGameplayGraphSeries& InSeries, const FTimingTrackViewport& InViewport)
+{
+	FMontageSeries& MontageSeries = *static_cast<FMontageSeries*>(&InSeries);
+	switch (MontageSeries.Type)
+	{
+	case FMontageSeries::ESeriesType::Weight:
+		return UpdateSeriesBoundsHelper(MontageSeries, InViewport, &FAnimMontageMessage::Weight);
+	case FMontageSeries::ESeriesType::Position:
+		return UpdateSeriesBoundsHelper(MontageSeries, InViewport, &FAnimMontageMessage::Position);
+	}
+	return false;
+}
+
+template<typename ProjectionType>
+void FMontageTrack::UpdateSeriesHelper(FMontageSeries& InSeries, const FTimingTrackViewport& InViewport, ProjectionType Projection)
 {
 	const FAnimationProvider* AnimationProvider = SharedData.GetAnalysisSession().ReadProvider<FAnimationProvider>(FAnimationProvider::ProviderName);
 
-	if(AnimationProvider)
+	if (AnimationProvider)
 	{
 		TraceServices::FAnalysisSessionReadScope SessionReadScope(SharedData.GetAnalysisSession());
 
-		FMontageSeries& MontageSeries = *static_cast<FMontageSeries*>(&InSeries);
+		FGraphTrackBuilder Builder(*this, InSeries, InViewport);
 
-		FGraphTrackBuilder Builder(*this, MontageSeries, InViewport);
-
-		AnimationProvider->ReadMontageTimeline(GetGameplayTrack().GetObjectId(), [this, &Builder, &InViewport, &MontageSeries, &AnimationProvider](const FAnimationProvider::AnimMontageTimeline& InTimeline)
+		AnimationProvider->ReadMontageTimeline(GetGameplayTrack().GetObjectId(), [this, &AnimationProvider, &Builder, &InViewport, &InSeries, &Projection](const FAnimationProvider::AnimMontageTimeline& InTimeline)
 		{
 			uint16 FrameCounter = 0;
 			uint16 LastFrameWithMontage = 0;
 
-			InTimeline.EnumerateEvents(InViewport.GetStartTime(), InViewport.GetEndTime(), [this, &FrameCounter, &LastFrameWithMontage, &Builder, &MontageSeries, &AnimationProvider](double InStartTime, double InEndTime, uint32 InDepth, const FAnimMontageMessage& InMessage)
+			InTimeline.EnumerateEvents(InViewport.GetStartTime(), InViewport.GetEndTime(), [this, &FrameCounter, &LastFrameWithMontage, &Builder, &InSeries, &Projection](double InStartTime, double InEndTime, uint32 InDepth, const FAnimMontageMessage& InMessage)
 			{
 				FrameCounter = InMessage.FrameCounter;
 
-				if(InMessage.MontageId == MontageSeries.MontageId)
+				if (InMessage.MontageId == InSeries.MontageId)
 				{
-					Builder.AddEvent(InStartTime, InEndTime - InStartTime, InMessage.Weight, LastFrameWithMontage == FrameCounter - 1);
+					Builder.AddEvent(InStartTime, InEndTime - InStartTime, Invoke(Projection, InMessage), LastFrameWithMontage == FrameCounter - 1);
 
-					LastFrameWithMontage = FrameCounter;
+					LastFrameWithMontage = InMessage.FrameCounter;
 				}
 				return TraceServices::EEventEnumerate::Continue;
 			});
@@ -135,6 +183,20 @@ void FMontageTrack::UpdateSeries(FGameplayGraphSeries& InSeries, const FTimingTr
 	}
 }
 
+void FMontageTrack::UpdateSeries(FGameplayGraphSeries& InSeries, const FTimingTrackViewport& InViewport)
+{
+	FMontageSeries& MontageSeries = *static_cast<FMontageSeries*>(&InSeries);
+	switch (MontageSeries.Type)
+	{
+	case FMontageSeries::ESeriesType::Weight:
+		UpdateSeriesHelper(MontageSeries, InViewport, &FAnimMontageMessage::Weight);
+		break;
+	case FMontageSeries::ESeriesType::Position:
+		UpdateSeriesHelper(MontageSeries, InViewport, &FAnimMontageMessage::Position);
+		break;
+	}
+}
+
 void FMontageTrack::InitTooltip(FTooltipDrawState& Tooltip, const ITimingEvent& HoveredTimingEvent) const
 {
 	const FGraphTrackEvent& GraphTrackEvent = *static_cast<const FGraphTrackEvent*>(&HoveredTimingEvent);
@@ -150,6 +212,7 @@ void FMontageTrack::InitTooltip(FTooltipDrawState& Tooltip, const ITimingEvent&
 		Tooltip.AddNameValueTextLine(LOCTEXT("EventTime", "Time").ToString(), FText::AsNumber(InFoundStartTime).ToString());
 		Tooltip.AddNameValueTextLine(LOCTEXT("EventWeight", "Weight").ToString(), FText::AsNumber(GraphTrackEvent.GetValue()).ToString());
 		Tooltip.AddNameValueTextLine(LOCTEXT("EventDesiredWeight", "Desired Weight").ToString(), FText::AsNumber(InMessage.DesiredWeight).ToString());
+		Tooltip.AddNameValueTextLine(LOCTEXT("EventPosition", "Position").ToString(), FText::AsNumber(InMessage.Position).ToString()); 
 
 		{
 			TraceServices::FAnalysisSessionReadScope SessionReadScope(SharedData.GetAnalysisSession());
diff --git a/Engine/Plugins/Animation/GameplayInsights/Source/GameplayInsights/Private/MontageTrack.h b/Engine/Plugins/Animation/GameplayInsights/Source/GameplayInsights/Private/MontageTrack.h
--- a/Engine/Plugins/Animation/GameplayInsights/Source/GameplayInsights/Private/MontageTrack.h
+++ b/Engine/Plugins/Animation/GameplayInsights/Source/GameplayInsights/Private/MontageTrack.h
@@ -12,7 +12,16 @@ class FMontageTrack;
 class FMontageSeries : public FGameplayGraphSeries
 {
 public:
+	enum class ESeriesType : uint32
+	{
+		Weight,
+		Position,
+
+		Count,
+	};
+
 	uint64 MontageId;
+	ESeriesType Type;
 };
 
 class FMontageTrack : public FGameplayGraphTrack
@@ -33,6 +42,14 @@ private:
 	// Helper function used to find a montage message
 	void FindMontageMessage(const FTimingEventSearchParameters& InParameters, TFunctionRef<void(double, double, uint32, const FAnimMontageMessage&)> InFoundPredicate) const;
 
+	// Helper function for series bounds update
+	template<typename ProjectionType>
+	bool UpdateSeriesBoundsHelper(FMontageSeries& InSeries, const FTimingTrackViewport & InViewport, ProjectionType Projection);
+
+	// Helper function for series update
+	template<typename ProjectionType>
+	void UpdateSeriesHelper(FMontageSeries& InSeries, const FTimingTrackViewport & InViewport, ProjectionType Projection);
+
 private:
 	/** The shared data */
 	const FAnimationSharedData& SharedData;
diff --git a/Engine/Plugins/Animation/GameplayInsights/Source/GameplayInsights/Public/IAnimationProvider.h b/Engine/Plugins/Animation/GameplayInsights/Source/GameplayInsights/Public/IAnimationProvider.h
--- a/Engine/Plugins/Animation/GameplayInsights/Source/GameplayInsights/Public/IAnimationProvider.h
+++ b/Engine/Plugins/Animation/GameplayInsights/Source/GameplayInsights/Public/IAnimationProvider.h
@@ -224,6 +224,7 @@ struct FAnimMontageMessage
 	uint32 NextSectionNameId = 0;
 	float Weight = 0.0f;
 	float DesiredWeight = 0.0f;
+	float Position = 0.0f;
 	uint16 FrameCounter = 0;
 };
 
diff --git a/Engine/Source/Runtime/Engine/Private/Animation/AnimTrace.cpp b/Engine/Source/Runtime/Engine/Private/Animation/AnimTrace.cpp
--- a/Engine/Source/Runtime/Engine/Private/Animation/AnimTrace.cpp
+++ b/Engine/Source/Runtime/Engine/Private/Animation/AnimTrace.cpp
@@ -229,6 +229,7 @@ UE_TRACE_EVENT_BEGIN(Animation, Montage)
 	UE_TRACE_EVENT_FIELD(uint32, NextSectionNameId)
 	UE_TRACE_EVENT_FIELD(float, Weight)
 	UE_TRACE_EVENT_FIELD(float, DesiredWeight)
+	UE_TRACE_EVENT_FIELD(float, Position)
 	UE_TRACE_EVENT_FIELD(uint16, FrameCounter)
 UE_TRACE_EVENT_END()
 
@@ -1146,6 +1147,7 @@ void FAnimTrace::OutputMontage(UAnimInstance* InAnimInstance, const FAnimMontage
 			<< Montage.NextSectionNameId(NextSectionNameId)
 			<< Montage.Weight(InMontageInstance.GetWeight())
 			<< Montage.DesiredWeight(InMontageInstance.GetDesiredWeight())
+			<< Montage.Position(InMontageInstance.GetPosition())
 			<< Montage.FrameCounter(FObjectTrace::GetObjectWorldTickCounter(InAnimInstance));
 	}
 }
