Rewind Debugger initial checkin

diff --git a/Engine/Plugins/Animation/GameplayInsights/GameplayInsights.uplugin b/Engine/Plugins/Animation/GameplayInsights/GameplayInsights.uplugin
--- a/Engine/Plugins/Animation/GameplayInsights/GameplayInsights.uplugin
+++ b/Engine/Plugins/Animation/GameplayInsights/GameplayInsights.uplugin
@@ -28,6 +28,12 @@
 			"Name": "GameplayInsightsEditor",
 			"Type": "Editor",
 			"LoadingPhase": "Default"
+		},
+		{
+			"Name": "RewindDebugger",
+			"Type": "Editor",
+			"LoadingPhase": "Default"
 		}
+
 	]
 }
diff --git a/Engine/Plugins/Animation/GameplayInsights/Source/GameplayInsights/Private/AnimationSharedData.cpp b/Engine/Plugins/Animation/GameplayInsights/Source/GameplayInsights/Private/AnimationSharedData.cpp
--- a/Engine/Plugins/Animation/GameplayInsights/Source/GameplayInsights/Private/AnimationSharedData.cpp
+++ b/Engine/Plugins/Animation/GameplayInsights/Source/GameplayInsights/Private/AnimationSharedData.cpp
@@ -669,8 +669,18 @@ void FAnimationSharedData::OpenAnimGraphTab(uint64 InAnimInstanceId) const
 				});
 			})
 		);
+		
+		TSharedPtr<SAnimGraphSchematicView> AnimGraphView = SNew(SAnimGraphSchematicView, InAnimInstanceId, TimingViewSession->GetTimeMarker(), *AnalysisSession);
+		TimingViewSession->OnTimeMarkerChanged().AddLambda(
+			[AnimGraphViewWeakPtr = TWeakPtr<SAnimGraphSchematicView>(AnimGraphView)](Insights::ETimeChangedFlags InFlags, double TimeMarker)
+			{
+				if (AnimGraphViewWeakPtr.IsValid())
+				{
+					AnimGraphViewWeakPtr.Pin()->SetTimeMarker(TimeMarker);
+				}
+			});
 
-		Tab->SetContent(SNew(SAnimGraphSchematicView, InAnimInstanceId, *TimingViewSession, *AnalysisSession));
+		Tab->SetContent(AnimGraphView.ToSharedRef());
 		WeakAnimGraphDocumentTabs.Add(Tab);
 
 		const FGameplayProvider* GameplayProvider = AnalysisSession->ReadProvider<FGameplayProvider>(FGameplayProvider::ProviderName);
diff --git a/Engine/Plugins/Animation/GameplayInsights/Source/GameplayInsights/Private/GameplayAnalyzer.cpp b/Engine/Plugins/Animation/GameplayInsights/Source/GameplayInsights/Private/GameplayAnalyzer.cpp
--- a/Engine/Plugins/Animation/GameplayInsights/Source/GameplayInsights/Private/GameplayAnalyzer.cpp
+++ b/Engine/Plugins/Animation/GameplayInsights/Source/GameplayInsights/Private/GameplayAnalyzer.cpp
@@ -14,6 +14,7 @@ void FGameplayAnalyzer::OnAnalysisBegin(const FOnAnalysisContext& Context)
 {
 	auto& Builder = Context.InterfaceBuilder;
 
+	Builder.RouteEvent(RouteId_RecordingInfo, "Object", "RecordingInfo");
 	Builder.RouteEvent(RouteId_Class, "Object", "Class");
 	Builder.RouteEvent(RouteId_Object, "Object", "Object");
 	Builder.RouteEvent(RouteId_ObjectEvent, "Object", "ObjectEvent");
@@ -42,6 +43,16 @@ bool FGameplayAnalyzer::OnEvent(uint16 RouteId, EStyle Style, const FOnEventCont
 		GameplayProvider.AppendWorld(Id, PIEInstanceId, Type, NetMode, bIsSimulating);
 		break;
 	}
+	case RouteId_RecordingInfo:
+	{
+		uint64 WorldId = EventData.GetValue<uint64>("WorldId");
+		uint64 Cycle = EventData.GetValue<uint64>("Cycle");
+		uint32 RecordingIndex = EventData.GetValue<uint32>("RecordingIndex");
+		uint32 FrameIndex = EventData.GetValue<uint32>("FrameIndex");
+		double ElapsedTime = EventData.GetValue<double>("ElapsedTime");
+		GameplayProvider.AppendRecordingInfo(WorldId, Context.EventTime.AsSeconds(Cycle), RecordingIndex, FrameIndex, ElapsedTime);
+		break;
+	}
 	case RouteId_Class:
 	{
 		const TCHAR* ClassNameAndPathName = reinterpret_cast<const TCHAR*>(EventData.GetAttachment());
diff --git a/Engine/Plugins/Animation/GameplayInsights/Source/GameplayInsights/Private/GameplayAnalyzer.h b/Engine/Plugins/Animation/GameplayInsights/Source/GameplayInsights/Private/GameplayAnalyzer.h
--- a/Engine/Plugins/Animation/GameplayInsights/Source/GameplayInsights/Private/GameplayAnalyzer.h
+++ b/Engine/Plugins/Animation/GameplayInsights/Source/GameplayInsights/Private/GameplayAnalyzer.h
@@ -28,6 +28,7 @@ private:
 		RouteId_PropertiesStart,
 		RouteId_PropertiesEnd,
 		RouteId_PropertyValue,
+		RouteId_RecordingInfo,
 	};
 
 	TraceServices::IAnalysisSession& Session;
diff --git a/Engine/Plugins/Animation/GameplayInsights/Source/GameplayInsights/Private/GameplayInsightsModule.cpp b/Engine/Plugins/Animation/GameplayInsights/Source/GameplayInsights/Private/GameplayInsightsModule.cpp
--- a/Engine/Plugins/Animation/GameplayInsights/Source/GameplayInsights/Private/GameplayInsightsModule.cpp
+++ b/Engine/Plugins/Animation/GameplayInsights/Source/GameplayInsights/Private/GameplayInsightsModule.cpp
@@ -17,6 +17,7 @@
 #include "Trace/StoreClient.h"
 #include "Stats/Stats.h"
 #include "ObjectPropertyTrace.h"
+#include "SAnimGraphSchematicView.h"
 
 #if WITH_EDITOR
 #include "IAnimationBlueprintEditorModule.h"
@@ -37,6 +38,11 @@
 
 const FName GameplayInsightsTabs::DocumentTab("DocumentTab");
 
+TSharedRef<IAnimGraphSchematicView> FGameplayInsightsModule::CreateAnimGraphSchematicView(uint64 InAnimInstanceId, double InTimeMarker, const TraceServices::IAnalysisSession& InAnalysisSession)
+{
+	return SNew(SAnimGraphSchematicView, InAnimInstanceId, InTimeMarker, InAnalysisSession);
+}
+
 void FGameplayInsightsModule::StartupModule()
 {
 	IModularFeatures::Get().RegisterModularFeature(TraceServices::ModuleFeatureName, &GameplayTraceModule);
diff --git a/Engine/Plugins/Animation/GameplayInsights/Source/GameplayInsights/Private/GameplayInsightsModule.h b/Engine/Plugins/Animation/GameplayInsights/Source/GameplayInsights/Private/GameplayInsightsModule.h
--- a/Engine/Plugins/Animation/GameplayInsights/Source/GameplayInsights/Private/GameplayInsightsModule.h
+++ b/Engine/Plugins/Animation/GameplayInsights/Source/GameplayInsights/Private/GameplayInsightsModule.h
@@ -20,7 +20,7 @@ namespace Trace
 }
 }
 
-class FGameplayInsightsModule : public IModuleInterface
+class FGameplayInsightsModule : public IGameplayInsightsModule
 {
 public:
 	// IModuleInterface interface
@@ -30,6 +30,8 @@ public:
 	// Spawn a document tab
 	TSharedRef<SDockTab> SpawnTimingProfilerDocumentTab(const FTabManager::FSearchPreference& InSearchPreference);
 
+	virtual TSharedRef<IAnimGraphSchematicView> CreateAnimGraphSchematicView(uint64 InAnimInstanceId, double InTimeMarker, const TraceServices::IAnalysisSession& InAnalysisSession) override;
+
 protected:
 	void RegisterTimingProfilerLayoutExtensions(FInsightsMajorTabExtender& InOutExtender);
 #if WITH_EDITOR
diff --git a/Engine/Plugins/Animation/GameplayInsights/Source/GameplayInsights/Private/GameplayProvider.cpp b/Engine/Plugins/Animation/GameplayInsights/Source/GameplayInsights/Private/GameplayProvider.cpp
--- a/Engine/Plugins/Animation/GameplayInsights/Source/GameplayInsights/Private/GameplayProvider.cpp
+++ b/Engine/Plugins/Animation/GameplayInsights/Source/GameplayInsights/Private/GameplayProvider.cpp
@@ -336,6 +336,42 @@ void FGameplayProvider::AppendWorld(uint64 InObjectId, int32 InPIEInstanceId, ui
 	}
 }
 
+void FGameplayProvider::AppendRecordingInfo(uint64 InWorldId, double InProfileTime, uint32 InRecordingIndex, uint32 InFrameIndex, double InElapsedTime)
+{
+	Session.WriteAccessCheck();
+
+	FRecordingInfoMessage NewRecordingInfo;
+	NewRecordingInfo.WorldId = InWorldId;
+	NewRecordingInfo.ProfileTime = InProfileTime;
+	NewRecordingInfo.RecordingIndex = InRecordingIndex;
+	NewRecordingInfo.FrameIndex = InFrameIndex;
+	NewRecordingInfo.ElapsedTime = InElapsedTime;
+
+	if(TSharedRef<TraceServices::TPointTimeline<FRecordingInfoMessage>>* ExistingRecording = Recordings.Find(InRecordingIndex))
+	{
+		(*ExistingRecording)->AppendEvent(InProfileTime, NewRecordingInfo);
+	}
+	else
+	{
+		TSharedPtr<TraceServices::TPointTimeline<FRecordingInfoMessage>> NewRecording = MakeShared<TraceServices::TPointTimeline<FRecordingInfoMessage>>(Session.GetLinearAllocator());
+		NewRecording->AppendEvent(InProfileTime, NewRecordingInfo);
+		Recordings.Add(InRecordingIndex, NewRecording.ToSharedRef());
+	}
+}
+
+
+const FGameplayProvider::RecordingInfoTimeline* FGameplayProvider::GetRecordingInfo(uint32 RecordingId) const  
+{
+	Session.ReadAccessCheck();
+
+	if(const TSharedRef<TraceServices::TPointTimeline<FRecordingInfoMessage>>* Recording = Recordings.Find(RecordingId))
+	{
+		return &(*Recording).Get();
+	}
+
+	return nullptr;
+}
+
 void FGameplayProvider::AppendClassPropertyStringId(uint32 InStringId, const FStringView& InString)
 {
 	Session.WriteAccessCheck();
diff --git a/Engine/Plugins/Animation/GameplayInsights/Source/GameplayInsights/Private/GameplayProvider.h b/Engine/Plugins/Animation/GameplayInsights/Source/GameplayInsights/Private/GameplayProvider.h
--- a/Engine/Plugins/Animation/GameplayInsights/Source/GameplayInsights/Private/GameplayProvider.h
+++ b/Engine/Plugins/Animation/GameplayInsights/Source/GameplayInsights/Private/GameplayProvider.h
@@ -35,6 +35,7 @@ public:
 	virtual const FObjectInfo& GetObjectInfo(uint64 InObjectId) const override;
 	virtual FOnObjectEndPlay& OnObjectEndPlay() override { return OnObjectEndPlayDelegate; }
 	virtual const TCHAR* GetPropertyName(uint32 InPropertyStringId) const override;
+	virtual const RecordingInfoTimeline* GetRecordingInfo(uint32 RecordingId) const override; 
 
 	/** Add a class message */
 	void AppendClass(uint64 InClassId, uint64 InSuperId, const TCHAR* InClassName, const TCHAR* InClassPathName);
@@ -48,6 +49,9 @@ public:
 	/** Add a world message */
 	void AppendWorld(uint64 InObjectId, int32 InPIEInstanceId, uint8 InType, uint8 InNetMode, bool bInIsSimulating);
 
+	/** Add a recording info message */
+	void AppendRecordingInfo(uint64 InWorldId, double InProfileTime, uint32 InRecordingIndex, uint32 InFrameIndex, double InElapsedTime);
+
 	/** Add a class property string ID message */
 	void AppendClassPropertyStringId(uint32 InStringId, const FStringView& InString);
 
@@ -119,6 +123,9 @@ private:
 	/** Map from string ID to stored string */
 	TMap<uint32, const TCHAR*> PropertyStrings;
 
+	/** Map of RecordingInfo Timelines by RecordingId - Each timeline is a mapping from Gameplay Elapsed time to Profiler Elapsed Time for one recording session */
+	TMap<uint32, TSharedRef<TraceServices::TPointTimeline<FRecordingInfoMessage>>> Recordings;
+
 	/** Whether we have any data */
 	bool bHasAnyData;
 
diff --git a/Engine/Plugins/Animation/GameplayInsights/Source/GameplayInsights/Private/SAnimGraphSchematicView.cpp b/Engine/Plugins/Animation/GameplayInsights/Source/GameplayInsights/Private/SAnimGraphSchematicView.cpp
--- a/Engine/Plugins/Animation/GameplayInsights/Source/GameplayInsights/Private/SAnimGraphSchematicView.cpp
+++ b/Engine/Plugins/Animation/GameplayInsights/Source/GameplayInsights/Private/SAnimGraphSchematicView.cpp
@@ -4,7 +4,6 @@
 #include "AnimationProvider.h"
 #include "Widgets/Input/SSearchBox.h"
 #include "TraceServices/Model/AnalysisSession.h"
-#include "Insights/ITimingViewSession.h"
 #include "GameplayProvider.h"
 #include "TraceServices/Model/Frames.h"
 #include "Widgets/Layout/SBorder.h"
@@ -444,14 +443,11 @@ class SAnimGraphSchematicNode : public SMultiColumnTableRow<TSharedRef<FAnimGrap
 	TAttribute<FText> FilterText;
 };
 
-void SAnimGraphSchematicView::Construct(const FArguments& InArgs, uint64 InAnimInstanceId, Insights::ITimingViewSession& InTimingViewSession, const TraceServices::IAnalysisSession& InAnalysisSession)
+void SAnimGraphSchematicView::Construct(const FArguments& InArgs, uint64 InAnimInstanceId, double InTimeMarker, const TraceServices::IAnalysisSession& InAnalysisSession)
 {
 	AnimInstanceId = InAnimInstanceId;
-	TimingViewSession = &InTimingViewSession;
+	TimeMarker = InTimeMarker;
 	AnalysisSession = &InAnalysisSession;
-	TimeMarker = InTimingViewSession.GetTimeMarker();
-
-	InTimingViewSession.OnTimeMarkerChanged().AddSP(this, &SAnimGraphSchematicView::HandleTimeMarkerChanged);
 
 	// Create header row and add default columns
 	HeaderRow = SNew(SHeaderRow);
@@ -817,11 +813,22 @@ void SAnimGraphSchematicView::RefreshFilter()
 	TreeView->RequestTreeRefresh();
 }
 
-void SAnimGraphSchematicView::HandleTimeMarkerChanged(Insights::ETimeChangedFlags InFlags, double InTimeMarker)
+void SAnimGraphSchematicView::SetTimeMarker(double InTimeMarker)
 {
-	TimeMarker = InTimeMarker;
+	if (TimeMarker != InTimeMarker)
+	{
+		TimeMarker = InTimeMarker;
+		RefreshNodes();
+	}
+}
 
-	RefreshNodes();
+void SAnimGraphSchematicView::SetAnimInstanceId(uint64 InAnimInstanceId)
+{
+	if (AnimInstanceId != InAnimInstanceId)
+	{
+		AnimInstanceId = InAnimInstanceId;
+		RefreshNodes();
+	}
 }
 
 TSharedRef<SWidget> SAnimGraphSchematicView::HandleGetViewMenuContent()
diff --git a/Engine/Plugins/Animation/GameplayInsights/Source/GameplayInsights/Private/SAnimGraphSchematicView.h b/Engine/Plugins/Animation/GameplayInsights/Source/GameplayInsights/Private/SAnimGraphSchematicView.h
--- a/Engine/Plugins/Animation/GameplayInsights/Source/GameplayInsights/Private/SAnimGraphSchematicView.h
+++ b/Engine/Plugins/Animation/GameplayInsights/Source/GameplayInsights/Private/SAnimGraphSchematicView.h
@@ -3,6 +3,7 @@
 #pragma once
 
 #include "CoreMinimal.h"
+#include "IAnimGraphSchematicView.h"
 #include "Widgets/SCompoundWidget.h"
 #include "Widgets/Views/STreeView.h"
 
@@ -19,17 +20,17 @@ class SComboButton;
 class SSplitter;
 enum class EAnimGraphSchematicFilterState;
 
-class SAnimGraphSchematicView : public SCompoundWidget
+class SAnimGraphSchematicView : public IAnimGraphSchematicView
 {
-public:
 	SLATE_BEGIN_ARGS(SAnimGraphSchematicView) {}
-
 	SLATE_END_ARGS()
+public:
+	void Construct(const FArguments& InArgs, uint64 InAnimInstanceId, double InTimeMarker, const TraceServices::IAnalysisSession& InAnalysisSession);
 
-	void Construct(const FArguments& InArgs, uint64 InAnimInstanceId, Insights::ITimingViewSession& InTimingViewSession, const TraceServices::IAnalysisSession& InAnalysisSession);
-
+	virtual void SetAnimInstanceId(uint64 InAnimInstanceId) override;
 	uint64 GetAnimInstanceId() const { return AnimInstanceId; }
 
+	virtual void SetTimeMarker(double InTimeMarker) override;
 private:
 	// Generate a row widget for an item
 	TSharedRef<ITableRow> HandleGenerateRow(TSharedRef<FAnimGraphSchematicNode> Item, const TSharedRef<STableViewBase>& OwnerTable);
@@ -73,8 +74,6 @@ private:
 private:
 	const TraceServices::IAnalysisSession* AnalysisSession;
 
-	Insights::ITimingViewSession* TimingViewSession;
-
 	TSharedPtr<SSearchBox> SearchBox;
 
 	TSharedPtr<SHeaderRow> HeaderRow;
diff --git a/Engine/Plugins/Animation/GameplayInsights/Source/GameplayInsights/Public/IAnimGraphSchematicView.h b/Engine/Plugins/Animation/GameplayInsights/Source/GameplayInsights/Public/IAnimGraphSchematicView.h
new file mode 100644
--- /dev/null
+++ b/Engine/Plugins/Animation/GameplayInsights/Source/GameplayInsights/Public/IAnimGraphSchematicView.h
@@ -0,0 +1,13 @@
+// Copyright Epic Games, Inc. All Rights Reserved.
+
+#pragma once
+
+#include "CoreMinimal.h"
+#include "Widgets/SCompoundWidget.h"
+
+class IAnimGraphSchematicView : public SCompoundWidget
+{
+	public:
+		virtual void SetTimeMarker(double InTimeMarker) = 0;
+	    virtual void SetAnimInstanceId(uint64 InAnimInstanceId) = 0;
+};
diff --git a/Engine/Plugins/Animation/GameplayInsights/Source/GameplayInsights/Public/IGameplayInsightsModule.h b/Engine/Plugins/Animation/GameplayInsights/Source/GameplayInsights/Public/IGameplayInsightsModule.h
--- a/Engine/Plugins/Animation/GameplayInsights/Source/GameplayInsights/Public/IGameplayInsightsModule.h
+++ b/Engine/Plugins/Animation/GameplayInsights/Source/GameplayInsights/Public/IGameplayInsightsModule.h
@@ -3,8 +3,23 @@
 #pragma once
 
 #include "CoreMinimal.h"
+#include "IAnimGraphSchematicView.h"
 
 namespace GameplayInsightsTabs
 {
 	extern GAMEPLAYINSIGHTS_API const FName DocumentTab;
+
+};
+
+namespace TraceServices
+{
+	class IAnalysisSession;
+}
+
+class IGameplayInsightsModule
+	: public IModuleInterface
+{
+public:
+	/** Creates a widget for displaying the anim graph schematic */
+	virtual TSharedRef<IAnimGraphSchematicView> CreateAnimGraphSchematicView(uint64 InAnimInstanceId, double InTimeMarker, const TraceServices::IAnalysisSession& InAnalysisSession) = 0;
 };
diff --git a/Engine/Plugins/Animation/GameplayInsights/Source/GameplayInsights/Public/IGameplayProvider.h b/Engine/Plugins/Animation/GameplayInsights/Source/GameplayInsights/Public/IGameplayProvider.h
--- a/Engine/Plugins/Animation/GameplayInsights/Source/GameplayInsights/Public/IGameplayProvider.h
+++ b/Engine/Plugins/Animation/GameplayInsights/Source/GameplayInsights/Public/IGameplayProvider.h
@@ -50,6 +50,15 @@ struct FObjectEventMessage
 	const TCHAR* Name = nullptr;
 };
 
+struct FRecordingInfoMessage
+{
+	uint64 WorldId;
+	double ProfileTime;
+	double ElapsedTime;
+	uint32 FrameIndex;
+	uint32 RecordingIndex;
+};
+
 
 struct FWorldInfo
 {
@@ -92,6 +101,7 @@ class IGameplayProvider : public TraceServices::IProvider
 public:
 	typedef TraceServices::ITimeline<FObjectEventMessage> ObjectEventsTimeline;
 	typedef TraceServices::ITimeline<FObjectPropertiesMessage> ObjectPropertiesTimeline;
+	typedef TraceServices::ITimeline<FRecordingInfoMessage> RecordingInfoTimeline;
 
 	virtual bool ReadObjectEventsTimeline(uint64 InObjectId, TFunctionRef<void(const ObjectEventsTimeline&)> Callback) const = 0;
 	virtual bool ReadObjectEvent(uint64 InObjectId, uint64 InMessageId, TFunctionRef<void(const FObjectEventMessage&)> Callback) const = 0;
@@ -109,4 +119,6 @@ public:
 	virtual const FObjectInfo& GetObjectInfo(uint64 InObjectId) const = 0;
 	virtual FOnObjectEndPlay& OnObjectEndPlay() = 0;
 	virtual const TCHAR* GetPropertyName(uint32 InPropertyStringId) const = 0;
+	virtual const RecordingInfoTimeline* GetRecordingInfo(uint32 RecordingId) const = 0; 
 };
+
diff --git a/Engine/Plugins/Animation/GameplayInsights/Source/RewindDebugger/Private/BindableProperty.h b/Engine/Plugins/Animation/GameplayInsights/Source/RewindDebugger/Private/BindableProperty.h
new file mode 100644
--- /dev/null
+++ b/Engine/Plugins/Animation/GameplayInsights/Source/RewindDebugger/Private/BindableProperty.h
@@ -0,0 +1,193 @@
+// Copyright Epic Games, Inc. All Rights Reserved.
+
+#pragma once
+#include "CoreMinimal.h"
+
+// Event based properties for syncronizing widgets with Models
+// In future, this should be replaced by slate ViewModels, and/or Slate Attributes once those have an OnChanged delegate
+//
+// Usage
+//
+// Widget Header:
+// 	SLATE_BEGIN_ARGS(SMyWidget) { }
+//		SLATE_ARGUMENT(TBindablePropertyInitializer<float>, MyProperty);
+//	SLATE_END_ARGS()
+//
+//  TBindableProperty<float> MyProperty
+//
+// Widget Cpp:
+// void SMyWidget::Construct(const FArguments& InArgs, ...
+// {
+//      MyProperty.Initialize(InArgs._MyProperty);
+//
+//      optionally:
+//      MyProperty.OnPropertyChanged = TBindableProperty<float>::FOnPropertyChanged::Create...
+//
+//
+// Outside the Widget:
+//
+// TBindableProperty<float> MyModelsProperty;
+// ...
+// SNew(SMyWidget,...)
+//   .MyProperty(&MyModelsProperty);
+//  or: .MyProperty(0.5f); 
+
+enum EBindingType
+{
+    BindingType_In,    // default - for value to be copied in to widget on change
+    BindingType_Out,   // value copied out of widget on change
+    BindingType_TwoWay // value copied both ways on change
+};
+
+template <typename T, EBindingType BindingType>
+class TBindableProperty;
+
+template <typename T, EBindingType BindingType = BindingType_In>
+struct TBindablePropertyInitializer
+{
+    TBindablePropertyInitializer()
+    {
+        Source = nullptr;
+        Value = T();
+    }
+
+    TBindablePropertyInitializer(TBindableProperty<T, BindingType>* InSource)
+    {
+        Source = InSource;
+        Value = T();
+    }
+    
+    TBindablePropertyInitializer(const T& InValue)
+    {
+        Source = nullptr;
+        Value = InValue;
+    }
+
+    TBindableProperty<T, BindingType>* Source;
+    T Value;
+};
+
+template <typename T, EBindingType BindingType = BindingType_In>
+class TBindableProperty
+{
+    public:
+        DECLARE_DELEGATE_OneParam(FOnPropertyChanged, T NewValue);
+
+        // delete copy/move constructors and equals operator because we are storing raw pointers to other properties, so they can't move around in memory
+        // todo: handle move constructor correctly so that Bindable properties can be stored in arrays
+        TBindableProperty(const TBindableProperty<T,BindingType>&) =delete;
+        TBindableProperty(TBindableProperty<T,BindingType>&&) =delete;
+        TBindableProperty& operator=(const TBindableProperty<T,BindingType>&) =delete;
+
+        TBindableProperty()
+            : InternalValue() // make sure POD types get zero initialized
+        {
+        }
+
+        ~TBindableProperty()
+        {
+            ClearBinding();
+
+            for(TBindableProperty<T, BindingType>* BoundProperty : BoundTo)
+            {
+                BoundProperty->BoundFrom = nullptr;
+            }
+        }
+
+        bool IsBound()
+        {
+            return BoundFrom != nullptr;
+        }
+
+        void Set(const T& NewValue)
+        {
+            if (!IsBound() || BindingType == BindingType_TwoWay)
+            {
+                SetInternal(NewValue);
+            }
+        }
+
+        const T& Get() const
+        {
+            return InternalValue;
+        }
+
+        void ClearBinding()
+        {
+            if (BoundFrom)
+            {
+                BoundFrom->RemoveBinding(this);
+                BoundFrom = nullptr;
+            }
+        }
+
+        void Bind(TBindableProperty<T, BindingType>* Source)
+        {
+            if (Source)
+            {
+                if (BindingType == BindingType_Out)
+                {
+                    ClearBinding();
+                    AddBinding(Source);
+                    Source->BoundFrom = this; 
+                }
+                else
+                {
+                    ClearBinding();
+                    Source->AddBinding(this);
+                    BoundFrom = Source; 
+                }
+            }
+        }
+
+        // either set an initial value, or bind to another property (for Out properties, maybe need to support both?)
+        void Initialize(const TBindablePropertyInitializer<T, BindingType> &Initializer)
+        {
+            if (Initializer.Source)
+            {
+                Bind(Initializer.Source);
+            }
+            else
+            {
+                Set(Initializer.Value);
+            }
+        }
+
+        FOnPropertyChanged OnPropertyChanged;
+
+    private:
+        T InternalValue;
+
+        TArray<TBindableProperty<T, BindingType>*> BoundTo;     // properties I will write to on write if BindingType is TwoWay or In
+        TBindableProperty<T, BindingType>* BoundFrom = nullptr; // a Property that will write to me if BindingType is From, and I write to if BindingType is TwoWay or Out
+
+        void AddBinding(TBindableProperty<T, BindingType>* Target)
+        {
+            BoundTo.Add(Target);            
+            Target->SetInternal(InternalValue);
+        }
+
+        void RemoveBinding(TBindableProperty<T, BindingType>* Target)
+        {
+            BoundTo.Remove(Target);            
+        }
+
+        void SetInternal(const T& NewValue)
+        {
+            InternalValue = NewValue;
+            OnPropertyChanged.ExecuteIfBound(InternalValue);
+
+			for (TBindableProperty<T, BindingType>* BoundProperty : BoundTo)
+			{
+				BoundProperty->SetInternal(InternalValue);
+			}
+
+            if (BindingType == BindingType_TwoWay)
+            {
+				if (BoundFrom)
+				{
+					BoundFrom->SetInternal(InternalValue);
+				}
+            }
+        }
+};
diff --git a/Engine/Plugins/Animation/GameplayInsights/Source/RewindDebugger/Private/RewindDebugger.cpp b/Engine/Plugins/Animation/GameplayInsights/Source/RewindDebugger/Private/RewindDebugger.cpp
new file mode 100644
--- /dev/null
+++ b/Engine/Plugins/Animation/GameplayInsights/Source/RewindDebugger/Private/RewindDebugger.cpp
@@ -0,0 +1,492 @@
+// Copyright Epic Games, Inc. All Rights Reserved.
+
+#include "RewindDebugger.h"
+#include "Components/SkeletalMeshComponent.h"
+#include "Editor.h"
+#include "IAnimationProvider.h"
+#include "IGameplayProvider.h"
+#include "Insights/IUnrealInsightsModule.h"
+#include "Modules/ModuleManager.h"
+#include "ObjectTrace.h"
+#include "TraceServices/ITraceServicesModule.h"
+
+
+FRewindDebugger* FRewindDebugger::InternalInstance = nullptr;
+
+FRewindDebugger::FRewindDebugger()  :
+	ControlState(FRewindDebugger::EControlState::Pause),
+	bPIEStarted(false),
+	bPIESimulating(false),
+	bAutoRecord(false),
+	bRecording(false),
+	PlaybackRate(1),
+	CurrentScrubTime(0),
+	RecordingIndex(0)
+{
+	RecordingDuration.Set(0);
+
+	FEditorDelegates::PostPIEStarted.AddRaw(this, &FRewindDebugger::OnPIEStarted);
+	FEditorDelegates::PausePIE.AddRaw(this, &FRewindDebugger::OnPIEPaused);
+	FEditorDelegates::ResumePIE.AddRaw(this, &FRewindDebugger::OnPIEResumed);
+	FEditorDelegates::EndPIE.AddRaw(this, &FRewindDebugger::OnPIEStopped);
+
+	DebugTargetActor.OnPropertyChanged = DebugTargetActor.OnPropertyChanged.CreateLambda([this](FString Target) { RefreshDebugComponents(); });
+
+	UnrealInsightsModule = &FModuleManager::LoadModuleChecked<IUnrealInsightsModule>("TraceInsights");
+}
+
+FRewindDebugger::~FRewindDebugger() 
+{
+	FEditorDelegates::PostPIEStarted.RemoveAll(this);
+	FEditorDelegates::PausePIE.RemoveAll(this);
+	FEditorDelegates::ResumePIE.RemoveAll(this);
+	FEditorDelegates::EndPIE.RemoveAll(this);
+}
+
+void FRewindDebugger::Initialize() 
+{
+	InternalInstance = new FRewindDebugger;
+}
+
+void FRewindDebugger::Shutdown() 
+{
+	delete InternalInstance;
+}
+
+void FRewindDebugger::OnComponentListChanged(const FOnComponentListChanged& InComponentListChangedDelegate)
+{
+	ComponentListChangedDelegate = InComponentListChangedDelegate;
+}
+
+void FRewindDebugger::OnTrackCursor(const FOnTrackCursor& InTrackCursorDelegate)
+{
+	TrackCursorDelegate = InTrackCursorDelegate;
+}
+
+void FRewindDebugger::OnPIEStarted(bool bSimulating)
+{
+	bPIEStarted = true;
+	bPIESimulating = true;
+
+	if (bAutoRecord)
+	{
+		StartRecording();
+	}
+}
+
+void FRewindDebugger::OnPIEPaused(bool bSimulating)
+{
+	bPIESimulating = false;
+	ControlState = EControlState::Pause;
+}
+
+void FRewindDebugger::OnPIEResumed(bool bSimulating)
+{
+	bPIESimulating = true;
+
+	// restore all relative transforms of any meshes that may have been moved while scrubbing
+	for (TTuple<uint64, FMeshComponentResetData>& MeshData : MeshComponentsToReset)
+	{
+		if (USkeletalMeshComponent* MeshComponent = MeshData.Value.Component.Get())
+		{
+			MeshComponent->SetRelativeTransform(MeshData.Value.RelativeTransform);
+		}
+	}
+
+	MeshComponentsToReset.Empty();
+}
+void FRewindDebugger::OnPIEStopped(bool bSimulating)
+{
+	bPIEStarted = false;
+	bPIESimulating = false;
+	MeshComponentsToReset.Empty();
+
+	StopRecording();
+	// clear the current recording (until we support playback in the Editor world on spawned actors)
+	RecordingDuration.Set(0);
+	SetCurrentScrubTime(0);
+}
+
+void FRewindDebugger::RefreshDebugComponents()
+{
+	TSharedPtr<const TraceServices::IAnalysisSession> Session = UnrealInsightsModule->GetAnalysisSession();
+	TraceServices::FAnalysisSessionReadScope SessionReadScope(*Session.Get());
+	UWorld* World = GetWorldToVisualize();
+
+	const IGameplayProvider* GameplayProvider = Session->ReadProvider<IGameplayProvider>("GameplayProvider");
+	const IAnimationProvider* AnimationProvider = Session->ReadProvider<IAnimationProvider>("AnimationProvider");
+
+    TArray<TSharedPtr<FDebugObjectInfo>> NewComponentList;	
+	
+	if (DebugTargetActor.Get() == "")
+	{
+		return;
+	}
+
+	uint64 TargetActorId = 0;
+
+	// find the Id of the selected actor (move this to a function or cache it)
+	GameplayProvider->EnumerateObjects([this,&TargetActorId](const FObjectInfo& InObjectInfo)
+	{
+		if (DebugTargetActor.Get() == InObjectInfo.Name)
+		{
+			TargetActorId = InObjectInfo.Id;
+		}
+	});
+
+	// add actor (even if it isn't found in the gameplay provider)
+	if (DebugComponents.Num() > 0 && DebugComponents[0]->ObjectName == DebugTargetActor.Get() && DebugComponents[0]->ObjectId == TargetActorId)
+	{
+		// re-use the version from the old array if it exists, to maintain selection in the list view
+		NewComponentList.Add(DebugComponents[0]);
+	}
+	else
+	{
+		NewComponentList.Add(MakeShared<FDebugObjectInfo>(TargetActorId, DebugTargetActor.Get()));
+	}
+
+	if (TargetActorId != 0)
+	{
+		GameplayProvider->EnumerateObjects([this, TargetActorId, GameplayProvider, &NewComponentList](const FObjectInfo& InObjectInfo)
+		{
+			// todo: filter components based on creation/destruction frame, and the current scrubbing time (so dynamically created and destroyed components won't add up)
+
+			const FObjectInfo* ObjectInfo = &InObjectInfo;
+			uint64 ObjectId = ObjectInfo->Id;
+			
+			// find any objects that have TargetActorId as OuterId (or of any outer object recursively)
+			while (ObjectInfo && ObjectInfo->OuterId != 0)
+			{
+				if(ObjectInfo->OuterId == TargetActorId)
+				{
+					// re-use the version from the old array if it exists, to maintain selection in the list view
+					int32 FoundIndex = DebugComponents.FindLastByPredicate([ObjectId](const TSharedPtr<FDebugObjectInfo>& Info) { return Info->ObjectId == ObjectId; });
+
+					if (FoundIndex >= 0 && DebugComponents[FoundIndex]->ObjectName == InObjectInfo.Name)
+					{
+						NewComponentList.Add(DebugComponents[FoundIndex]);
+					}
+					else
+					{
+						NewComponentList.Add(MakeShared<FDebugObjectInfo>(ObjectId,InObjectInfo.Name));
+					}
+					return;
+				}
+
+				ObjectInfo = GameplayProvider->FindObjectInfo(ObjectInfo->OuterId);
+			}
+		});
+	}
+
+	if (DebugComponents != NewComponentList) // since we reused old SharedPtr for any entries that existed, array equality will tell us if anything has changed
+	{
+		DebugComponents = NewComponentList;
+		ComponentListChangedDelegate.ExecuteIfBound();
+	}
+}
+
+void FRewindDebugger::StartRecording()
+{
+	if (!CanStartRecording())
+	{
+		return;
+	}
+
+	// Enable Object and Animation Trace filters
+	UE::Trace::ToggleChannel(TEXT("Object"), true);
+	UE::Trace::ToggleChannel(TEXT("Animation"), true);
+	UE::Trace::ToggleChannel(TEXT("Frame"), true);
+
+	RecordingDuration.Set(0);
+	RecordingIndex++;
+	bRecording = true;
+
+	// setup FObjectTrace to start tracking tracing times from 0
+	// and increment the RecordingIndex so we can use it to distinguish between the latest recording and older ones
+	UWorld* World = GetWorldToVisualize();
+	FObjectTrace::ResetWorldElapsedTime(World);
+	FObjectTrace::SetWorldRecordingIndex(World, RecordingIndex);
+}
+
+void FRewindDebugger::StopRecording()
+{
+	if (bRecording)
+	{
+		// Enable Object and Animation Trace filters
+		UE::Trace::ToggleChannel(TEXT("Object"), false);
+		UE::Trace::ToggleChannel(TEXT("Animation"), false);
+		UE::Trace::ToggleChannel(TEXT("Frame"), false);
+
+		bRecording = false;
+	}
+}
+
+bool FRewindDebugger::CanPause() const
+{
+	return ControlState != EControlState::Pause;
+}
+
+void FRewindDebugger::Pause()
+{
+	if (CanPause())
+	{
+		if (bPIESimulating)
+		{
+			// pause PIE
+		}
+
+		ControlState = EControlState::Pause;
+	}
+}
+
+bool FRewindDebugger::IsPlaying() const
+{
+	return ControlState==EControlState::Play && !bPIESimulating;
+}
+
+bool FRewindDebugger::CanPlay() const
+{
+	return ControlState!=EControlState::Play && !bPIESimulating && RecordingDuration.Get() > 0;
+}
+
+void FRewindDebugger::Play()
+{
+	if (CanPlay())
+	{
+		if (CurrentScrubTime >= RecordingDuration.Get())
+		{
+			SetCurrentScrubTime(0);
+		}
+
+		ControlState = EControlState::Play;
+	}
+}
+
+bool FRewindDebugger::CanPlayReverse() const
+{
+	return ControlState!=EControlState::PlayReverse && !bPIESimulating && RecordingDuration.Get() > 0;
+}
+
+void FRewindDebugger::PlayReverse()
+{
+	if (CanPlayReverse())
+	{
+		if (CurrentScrubTime <= 0)
+		{
+			SetCurrentScrubTime(RecordingDuration.Get());
+		}
+
+		ControlState = EControlState::PlayReverse;
+	}
+}
+
+bool FRewindDebugger::CanScrub() const
+{
+	return !bPIESimulating && RecordingDuration.Get() > 0;
+}
+
+void FRewindDebugger::ScrubToStart()
+{
+	if (CanScrub())
+	{
+		Pause();
+		SetCurrentScrubTime(0);
+		TrackCursorDelegate.ExecuteIfBound(false);
+	}
+}
+
+void FRewindDebugger::ScrubToEnd()
+{
+	if (CanScrub())
+	{
+		Pause();
+		SetCurrentScrubTime(RecordingDuration.Get());
+		TrackCursorDelegate.ExecuteIfBound(false);
+	}
+}
+
+void FRewindDebugger::StepForward()
+{
+	if (CanScrub())
+	{
+		Pause();
+
+		SetCurrentScrubTime(FMath::Min(CurrentScrubTime + 1.0f/30.0f, RecordingDuration.Get()));
+		TrackCursorDelegate.ExecuteIfBound(false);
+	}
+}
+
+void FRewindDebugger::StepBackward()
+{
+	if (CanScrub())
+	{
+		Pause();
+		// todo: snap to actual frames from AnimationProvider
+		SetCurrentScrubTime(FMath::Max(CurrentScrubTime - 1.0f/30.0f, 0.0f));
+		TrackCursorDelegate.ExecuteIfBound(true);
+	}
+}
+
+
+void FRewindDebugger::ScrubToTime(float ScrubTime, bool bIsScrubbing)
+{
+	if (CanScrub())
+	{
+		Pause();
+		SetCurrentScrubTime(ScrubTime);
+	}
+}
+
+UWorld* FRewindDebugger::GetWorldToVisualize() const
+{
+	// we probably want to replace this with a world selector widget, if we are going to support tracing from anything other thn the PIE world
+
+	UWorld* World = nullptr;
+
+#if WITH_EDITOR
+	UEditorEngine* EditorEngine = Cast<UEditorEngine>(GEngine);
+	if (GIsEditor && EditorEngine != nullptr && World == nullptr)
+	{
+		// lets use PlayWorld during PIE/Simulate and regular world from editor otherwise, to draw debug information
+		World = EditorEngine->PlayWorld != nullptr ? ToRawPtr(EditorEngine->PlayWorld) : EditorEngine->GetEditorWorldContext().World();
+	}
+
+#endif
+	if (!GIsEditor && World == nullptr)
+	{
+		World = GEngine->GetWorld();
+	}
+
+	return World;
+}
+
+void  FRewindDebugger::SetCurrentScrubTime(float Time)
+{
+	CurrentScrubTime = Time;
+	UpdateTraceTime();
+}
+
+void FRewindDebugger::UpdateTraceTime()
+{
+	// find the current Insights trace time for the current scrub time.
+	TSharedPtr<const TraceServices::IAnalysisSession> Session = UnrealInsightsModule->GetAnalysisSession();
+	TraceServices::FAnalysisSessionReadScope SessionReadScope(*Session.Get());
+	UWorld* World = GetWorldToVisualize();
+
+	if (const IGameplayProvider* GameplayProvider = Session->ReadProvider<IGameplayProvider>("GameplayProvider"))
+	{
+		if (const IGameplayProvider::RecordingInfoTimeline* Recording = GameplayProvider->GetRecordingInfo(RecordingIndex))
+		{
+			uint64 EventCount = Recording->GetEventCount();
+
+			if (EventCount > 0)
+			{
+				uint64 EventIndex;
+
+				for(EventIndex = 0; EventIndex < EventCount-1; EventIndex++)
+				{
+					const FRecordingInfoMessage& Event = Recording->GetEvent(EventIndex);
+					// todo: optimize this linear search (cache current index, and pass it in as a hint)
+					//		 and do interpolation between nearest two frames
+					if (Event.ElapsedTime >= CurrentScrubTime)
+					{
+						break;
+					}
+				}
+				TraceTime.Set(Recording->GetEvent(EventIndex).ProfileTime);
+			}
+		}
+	}
+}
+
+void FRewindDebugger::Tick(float DeltaTime)
+{
+	if (bRecording)
+	{
+		// if you select a debug target before you start recording, update component list when it becomes valid
+		RefreshDebugComponents();
+	}
+
+	TSharedPtr<const TraceServices::IAnalysisSession> Session = UnrealInsightsModule->GetAnalysisSession();
+	TraceServices::FAnalysisSessionReadScope SessionReadScope(*Session.Get());
+	UWorld* World = GetWorldToVisualize();
+
+	if (const IAnimationProvider* AnimationProvider = Session->ReadProvider<IAnimationProvider>("AnimationProvider"))
+	{
+		if (bPIESimulating)
+		{
+			if (bRecording)
+			{
+				RecordingDuration.Set(FObjectTrace::GetWorldElapsedTime(World));
+				SetCurrentScrubTime(RecordingDuration.Get());
+				TrackCursorDelegate.ExecuteIfBound(false);
+			}
+		}
+		else
+		{
+			if (ControlState == EControlState::Play || ControlState == EControlState::PlayReverse)
+			{
+				float Rate = PlaybackRate * (ControlState == EControlState::Play ? 1 : -1);
+				SetCurrentScrubTime(FMath::Clamp(CurrentScrubTime + Rate * DeltaTime, 0.0f, RecordingDuration.Get()));
+				TrackCursorDelegate.ExecuteIfBound(Rate<0);
+
+				if (CurrentScrubTime == 0 || CurrentScrubTime == RecordingDuration.Get())
+				{
+					// pause at end.
+					ControlState = EControlState::Pause;
+				}
+			}
+
+			double CurrentTraceTime = TraceTime.Get();
+
+			// update pose on all SkeletalMeshComponents
+			ULevel* CurLevel = World->GetCurrentLevel();
+			for( AActor* LevelActor : CurLevel->Actors )
+			{
+				if (LevelActor)
+				{
+					TInlineComponentArray<USkeletalMeshComponent*> SkeletalMeshComponents;
+					LevelActor->GetComponents(SkeletalMeshComponents);
+
+					for (USkeletalMeshComponent* MeshComponent : SkeletalMeshComponents)
+					{
+						int64 ObjectId = FObjectTrace::GetObjectId(MeshComponent);
+
+						AnimationProvider->ReadSkeletalMeshPoseTimeline(ObjectId, [this, ObjectId, CurrentTraceTime, MeshComponent, AnimationProvider](const IAnimationProvider::SkeletalMeshPoseTimeline& TimelineData, bool bHasCurves)
+						{
+							double PrecedingPoseTime;
+							double FollowingPoseTime;
+							const FSkeletalMeshPoseMessage* PrecedingPose;
+							const FSkeletalMeshPoseMessage* FollowingPose;
+
+							TimelineData.FindNearestEvents(CurrentTraceTime, PrecedingPose, PrecedingPoseTime, FollowingPose, FollowingPoseTime);
+
+							if (FollowingPose || PrecedingPose)
+							{	
+								// Ideally we should iterpolate between the two poses here if both are valid (we would need to transform the component space transforms to local space though)
+								const FSkeletalMeshPoseMessage& PoseMessage = FollowingPose ? *FollowingPose : *PrecedingPose;
+
+								FTransform ComponentWorldTransform;
+								const FSkeletalMeshInfo* SkeletalMeshInfo = AnimationProvider->FindSkeletalMeshInfo(PoseMessage.MeshId);
+								AnimationProvider->GetSkeletalMeshComponentSpacePose(PoseMessage, *SkeletalMeshInfo, ComponentWorldTransform, MeshComponent->GetEditableComponentSpaceTransforms());
+
+								if (MeshComponentsToReset.Find(ObjectId) == nullptr)
+								{
+									FMeshComponentResetData ResetData;
+									ResetData.Component = MeshComponent;
+									ResetData.RelativeTransform = MeshComponent->GetRelativeTransform();
+									MeshComponentsToReset.Add(ObjectId, ResetData);
+								}
+
+								MeshComponent->SetWorldTransform(ComponentWorldTransform);
+								MeshComponent->SetForcedLOD(PoseMessage.LodIndex + 1);
+								MeshComponent->ApplyEditedComponentSpaceTransforms();
+							}
+						});
+					}
+				}
+			}
+
+		}
+	}
+}
diff --git a/Engine/Plugins/Animation/GameplayInsights/Source/RewindDebugger/Private/RewindDebugger.h b/Engine/Plugins/Animation/GameplayInsights/Source/RewindDebugger/Private/RewindDebugger.h
new file mode 100644
--- /dev/null
+++ b/Engine/Plugins/Animation/GameplayInsights/Source/RewindDebugger/Private/RewindDebugger.h
@@ -0,0 +1,133 @@
+// Copyright Epic Games, Inc. All Rights Reserved.
+#pragma once
+
+#include "CoreMinimal.h"
+#include "BindableProperty.h"
+#include "UObject/WeakObjectPtr.h"
+#include "RewindDebuggerModule.h"
+
+
+// Singleton class that handles the logic for the Rewind Debugger
+// handles:
+//  Playback/Scrubbing state
+//  Start/Stop recording
+//  Keeping track of the current Debug Target actor, and outputing a list of it's Components for the UI
+
+class FRewindDebugger
+{
+public:
+	FRewindDebugger();
+	~FRewindDebugger();
+
+	// create singleton instance
+	static void Initialize();
+
+	// destroy singleton instance
+	static void Shutdown();
+
+	// get singleton instance
+	static FRewindDebugger* Instance() { return InternalInstance; }
+
+	// Start a new Recording:  Start tracing Object + Animation data, increment the current recording index, and reset the recording elapsed time to 0
+	void StartRecording();
+
+	bool IsRecording() const { return bRecording; }
+	bool CanStartRecording() const { return !IsRecording() && bPIESimulating; }
+
+	bool AutoRecord() const { return bAutoRecord; }
+	void SetAutoRecord(bool value) { bAutoRecord = value; }
+
+	// Stop recording: Stop tracing Object + Animation Data.
+	void StopRecording();
+	bool CanStopRecording() const { return IsRecording(); }
+
+	// VCR controls
+
+	bool CanPause() const;
+	void Pause();
+
+	bool CanPlay() const;
+	void Play();
+	bool IsPlaying() const;
+
+	bool CanPlayReverse() const;
+	void PlayReverse();
+
+	void ScrubToStart();
+	void ScrubToEnd();
+	void StepForward();
+	void StepBackward();
+
+	bool CanScrub() const;
+	void ScrubToTime(float ScrubTime, bool bIsScrubbing);
+	float GetScrubTime() { return CurrentScrubTime; }
+
+	// Tick function: While recording, update recording duration.  While paused, and we have recorded data, update skinned mesh poses for the current frame, and handle playback.
+	void Tick(float DeltaTime);
+
+	// update the list of components for the currently selected debug target
+	void RefreshDebugComponents();
+
+	TArray<TSharedPtr<FDebugObjectInfo>>& GetDebugComponents() { return DebugComponents; };
+
+	DECLARE_DELEGATE(FOnComponentListChanged)
+	void OnComponentListChanged(const FOnComponentListChanged& ComponentListChangedCallback);
+
+	DECLARE_DELEGATE_OneParam( FOnTrackCursor, bool)
+	void OnTrackCursor(const FOnTrackCursor& TrackCursorCallback);
+
+	TBindableProperty<double>* GetTraceTimeProperty() { return &TraceTime; }
+	TBindableProperty<float>* GetRecordingDurationProperty() { return &RecordingDuration; }
+	TBindableProperty<FString, BindingType_Out>* GetDebugTargetActorProperty() { return &DebugTargetActor; }
+
+private:
+	void OnPIEStarted(bool bSimulating);
+	void OnPIEPaused(bool bSimulating);
+	void OnPIEResumed(bool bSimulating);
+	void OnPIEStopped(bool bSimulating);
+
+	UWorld* GetWorldToVisualize() const;
+
+	void SetCurrentScrubTime(float Time);
+	void UpdateTraceTime();
+
+	TBindableProperty<double> TraceTime;
+	TBindableProperty<float> RecordingDuration;
+	TBindableProperty<FString, BindingType_Out> DebugTargetActor;
+
+	enum class EControlState : int8
+	{
+		Play,
+		PlayReverse,
+		Pause
+	};
+
+	EControlState ControlState;
+
+	FOnComponentListChanged ComponentListChangedDelegate;
+	FOnTrackCursor TrackCursorDelegate;
+
+	bool bPIEStarted;
+	bool bPIESimulating;
+
+	bool bAutoRecord;
+	bool bRecording;
+
+	float PlaybackRate;
+	float CurrentScrubTime;
+	uint16 RecordingIndex;
+
+	static FRewindDebugger* InternalInstance;
+
+	TArray<TSharedPtr<FDebugObjectInfo>> DebugComponents;
+
+	struct FMeshComponentResetData
+	{
+		TWeakObjectPtr<USkeletalMeshComponent> Component;
+		FTransform RelativeTransform;
+	};
+
+	TMap<uint64, FMeshComponentResetData> MeshComponentsToReset;
+
+	class IUnrealInsightsModule *UnrealInsightsModule;
+};
diff --git a/Engine/Plugins/Animation/GameplayInsights/Source/RewindDebugger/Private/RewindDebuggerCommands.h b/Engine/Plugins/Animation/GameplayInsights/Source/RewindDebugger/Private/RewindDebuggerCommands.h
new file mode 100644
--- /dev/null
+++ b/Engine/Plugins/Animation/GameplayInsights/Source/RewindDebugger/Private/RewindDebuggerCommands.h
@@ -0,0 +1,48 @@
+// Copyright Epic Games, Inc. All Rights Reserved.
+
+#pragma once
+
+#include "CoreMinimal.h"
+#include "Framework/Commands/InputChord.h"
+#include "Framework/Commands/Commands.h"
+
+#define LOCTEXT_NAMESPACE "FRewindDebuggerCommands"
+
+class FRewindDebuggerCommands : public TCommands<FRewindDebuggerCommands>
+{
+public:
+
+	/** Default constructor. */
+	FRewindDebuggerCommands()
+		: TCommands<FRewindDebuggerCommands>("RewindDebugger", NSLOCTEXT("Contexts", "RewindDebugger", "Animation Insights 2"), NAME_None, "RewindDebuggerStyle")
+	{ }
+
+	// TCommands interface
+
+	virtual void RegisterCommands() override
+	{
+		UI_COMMAND(StartRecording, "Start Recording", "Start recording Animation Insights data", EUserInterfaceActionType::Button, FInputChord());
+		UI_COMMAND(StopRecording, "Stop Recording", "Stop recording Animation Insights data", EUserInterfaceActionType::Button, FInputChord());
+		UI_COMMAND(FirstFrame, "First Frame", "Jump to first recorded frame", EUserInterfaceActionType::Button, FInputChord());
+		UI_COMMAND(PreviousFrame, "Previous Frame", "Step one frame back", EUserInterfaceActionType::Button, FInputChord());
+		UI_COMMAND(ReversePlay, "Reverse Play", "Playback recorded data in reverse", EUserInterfaceActionType::Button, FInputChord());
+		UI_COMMAND(Pause, "Pause", "Pause playback", EUserInterfaceActionType::Button, FInputChord());
+		UI_COMMAND(Play, "Play", "Playback recorded data", EUserInterfaceActionType::Button, FInputChord());
+		UI_COMMAND(NextFrame, "Next Frame", "Step one frame forward", EUserInterfaceActionType::Button, FInputChord());
+		UI_COMMAND(LastFrame, "Last Frame", "Jump to last recorded frame", EUserInterfaceActionType::Button, FInputChord());
+	}
+
+
+	TSharedPtr<FUICommandInfo> StartRecording;
+	TSharedPtr<FUICommandInfo> StopRecording;
+	TSharedPtr<FUICommandInfo> FirstFrame;
+	TSharedPtr<FUICommandInfo> PreviousFrame;
+	TSharedPtr<FUICommandInfo> ReversePlay;
+	TSharedPtr<FUICommandInfo> Pause;
+	TSharedPtr<FUICommandInfo> Play;
+	TSharedPtr<FUICommandInfo> NextFrame;
+	TSharedPtr<FUICommandInfo> LastFrame;
+};
+
+
+#undef LOCTEXT_NAMESPACE
diff --git a/Engine/Plugins/Animation/GameplayInsights/Source/RewindDebugger/Private/RewindDebuggerModule.cpp b/Engine/Plugins/Animation/GameplayInsights/Source/RewindDebugger/Private/RewindDebuggerModule.cpp
new file mode 100644
--- /dev/null
+++ b/Engine/Plugins/Animation/GameplayInsights/Source/RewindDebugger/Private/RewindDebuggerModule.cpp
@@ -0,0 +1,148 @@
+// Copyright Epic Games, Inc. All Rights Reserved.
+
+#include "RewindDebuggerModule.h"
+#include "Containers/Ticker.h"
+#include "Editor.h"
+#include "Engine/Engine.h"
+#include "Engine/Selection.h"
+#include "Features/IModularFeatures.h"
+#include "Framework/Docking/LayoutExtender.h"
+#include "IAnimationBlueprintEditorModule.h"
+#include "LevelEditorMenuContext.h"
+#include "Modules/ModuleManager.h"
+#include "RewindDebugger.h"
+#include "RewindDebuggerStyle.h"
+#include "RewindDebuggerCommands.h"
+#include "SRewindDebugger.h"
+#include "SSCSEditorMenuContext.h"
+#include "SSCSEditor.h"
+#include "ToolMenus.h"
+#include "Widgets/Docking/SDockTab.h"
+#include "WorkspaceMenuStructure.h"
+#include "WorkspaceMenuStructureModule.h"
+
+#define LOCTEXT_NAMESPACE "RewindDebuggerModule"
+
+static const FName RewindDebuggerTabName("RewindDebugger");
+
+TSharedRef<SDockTab> FRewindDebuggerModule::SpawnRewindDebuggerTab(const FSpawnTabArgs& SpawnTabArgs)
+{
+	const TSharedRef<SDockTab> MajorTab = SNew(SDockTab)
+		.TabRole(ETabRole::NomadTab);
+
+	TSharedPtr<SWidget> TabContent;
+
+	TSharedPtr<FUICommandList> CommandList = MakeShared<FUICommandList>();
+	const FRewindDebuggerCommands& Commands = FRewindDebuggerCommands::Get();
+
+	FRewindDebugger* DebuggerInstance = FRewindDebugger::Instance();
+
+	CommandList->MapAction(Commands.Play,
+							FExecuteAction::CreateRaw(DebuggerInstance, &FRewindDebugger::Play), 
+							FCanExecuteAction::CreateRaw(DebuggerInstance, &FRewindDebugger::CanPlay),
+							FIsActionChecked(),
+							FIsActionButtonVisible());
+
+	CommandList->MapAction(Commands.Pause,
+							FExecuteAction::CreateRaw(DebuggerInstance, &FRewindDebugger::Pause), 
+							FCanExecuteAction::CreateRaw(DebuggerInstance, &FRewindDebugger::CanPause),
+							FIsActionChecked(),
+							FIsActionButtonVisible());
+
+	CommandList->MapAction(Commands.ReversePlay,
+							FExecuteAction::CreateRaw(DebuggerInstance, &FRewindDebugger::PlayReverse), 
+							FCanExecuteAction::CreateRaw(DebuggerInstance, &FRewindDebugger::CanPlayReverse),
+							FIsActionChecked(),
+							FIsActionButtonVisible());
+
+	CommandList->MapAction(Commands.PreviousFrame,
+							FExecuteAction::CreateRaw(DebuggerInstance, &FRewindDebugger::StepBackward), 
+							FCanExecuteAction::CreateRaw(DebuggerInstance, &FRewindDebugger::CanScrub),
+							FIsActionChecked(),
+							FIsActionButtonVisible());
+
+
+	CommandList->MapAction(Commands.FirstFrame,
+							FExecuteAction::CreateRaw(DebuggerInstance, &FRewindDebugger::ScrubToStart), 
+							FCanExecuteAction::CreateRaw(DebuggerInstance, &FRewindDebugger::CanScrub),
+							FIsActionChecked(),
+							FIsActionButtonVisible());
+
+	CommandList->MapAction(Commands.LastFrame,
+							FExecuteAction::CreateRaw(DebuggerInstance, &FRewindDebugger::ScrubToEnd), 
+							FCanExecuteAction::CreateRaw(DebuggerInstance, &FRewindDebugger::CanScrub),
+							FIsActionChecked(),
+							FIsActionButtonVisible());
+
+	CommandList->MapAction(Commands.NextFrame,
+							FExecuteAction::CreateRaw(DebuggerInstance, &FRewindDebugger::StepForward), 
+							FCanExecuteAction::CreateRaw(DebuggerInstance, &FRewindDebugger::CanScrub),
+							FIsActionChecked(),
+							FIsActionButtonVisible());
+
+	CommandList->MapAction(Commands.StartRecording, 
+							FExecuteAction::CreateRaw(DebuggerInstance, &FRewindDebugger::StartRecording), 
+							FCanExecuteAction::CreateRaw(DebuggerInstance, &FRewindDebugger::CanStartRecording),
+							FIsActionChecked(),
+							FIsActionButtonVisible::CreateLambda([]() { return !FRewindDebugger::Instance()->IsRecording();}));
+
+	CommandList->MapAction(Commands.StopRecording,
+							 FExecuteAction::CreateRaw(DebuggerInstance, &FRewindDebugger::StopRecording),
+							 FCanExecuteAction::CreateRaw(DebuggerInstance, &FRewindDebugger::CanStopRecording),
+							 FIsActionChecked(),
+							 FIsActionButtonVisible::CreateRaw(DebuggerInstance, &FRewindDebugger::CanStopRecording));
+
+
+	RewindDebuggerWidget = SNew(SRewindDebugger, CommandList.ToSharedRef(), MajorTab, SpawnTabArgs.GetOwnerWindow())
+								.DebugTargetActor(DebuggerInstance->GetDebugTargetActorProperty())
+								.RecordingDuration(DebuggerInstance->GetRecordingDurationProperty())
+								.DebugComponents(&DebuggerInstance->GetDebugComponents())
+								.TraceTime(DebuggerInstance->GetTraceTimeProperty())
+								.OnScrubPositionChanged(SRewindDebugger::FOnScrubPositionChanged::CreateRaw(DebuggerInstance,&FRewindDebugger::ScrubToTime))
+								.ScrubTime_Lambda([]() { return FRewindDebugger::Instance()->GetScrubTime(); });
+
+	DebuggerInstance->OnTrackCursor(FRewindDebugger::FOnTrackCursor::CreateSP(RewindDebuggerWidget.Get(), &SRewindDebugger::TrackCursor));
+	DebuggerInstance->OnComponentListChanged(FRewindDebugger::FOnComponentListChanged::CreateSP(RewindDebuggerWidget.Get(), &SRewindDebugger::RefreshDebugComponents));
+
+	MajorTab->SetContent(RewindDebuggerWidget.ToSharedRef());
+
+	return MajorTab;
+}
+
+void FRewindDebuggerModule::StartupModule()
+{
+	FRewindDebugger::Initialize();
+	FRewindDebuggerStyle::Initialize();
+	FRewindDebuggerCommands::Register();
+
+	FGlobalTabmanager::Get()->RegisterNomadTabSpawner(
+		RewindDebuggerTabName, FOnSpawnTab::CreateRaw(this, &FRewindDebuggerModule::SpawnRewindDebuggerTab))
+		.SetGroup(WorkspaceMenu::GetMenuStructure().GetDeveloperToolsDebugCategory())
+		.SetDisplayName(LOCTEXT("TabTitle", "Rewind Debugger"))
+		.SetIcon(FSlateIcon("RewindDebuggerStyle", "RewindDebugger.RewindIcon"))
+		.SetTooltipText(LOCTEXT("TooltipText", "Opens Rewind Debugger."));
+
+
+	TickerHandle = FTicker::GetCoreTicker().AddTicker(TEXT("RewindDebugger"), 0.0f, [this](float DeltaTime)
+	{
+		QUICK_SCOPE_CYCLE_COUNTER(STAT_FRewindDebuggerModule_Tick);
+
+		FRewindDebugger::Instance()->Tick(DeltaTime);
+
+		return true;
+	});
+
+}
+
+void FRewindDebuggerModule::ShutdownModule()
+{
+	FTicker::GetCoreTicker().RemoveTicker(TickerHandle);
+
+	FRewindDebuggerCommands::Unregister();
+	FRewindDebuggerStyle::Shutdown();
+	FRewindDebugger::Shutdown();
+}
+
+IMPLEMENT_MODULE(FRewindDebuggerModule, RewindDebugger);
+
+#undef LOCTEXT_NAMESPACE
diff --git a/Engine/Plugins/Animation/GameplayInsights/Source/RewindDebugger/Private/RewindDebuggerModule.h b/Engine/Plugins/Animation/GameplayInsights/Source/RewindDebugger/Private/RewindDebuggerModule.h
new file mode 100644
--- /dev/null
+++ b/Engine/Plugins/Animation/GameplayInsights/Source/RewindDebugger/Private/RewindDebuggerModule.h
@@ -0,0 +1,34 @@
+// Copyright Epic Games, Inc. All Rights Reserved.
+
+#pragma once
+
+#include "CoreMinimal.h"
+#include "Modules/ModuleInterface.h"
+#include "Framework/Docking/TabManager.h"
+
+class SRewindDebugger;
+
+struct FDebugObjectInfo
+{
+	FDebugObjectInfo(uint64 Id, const FString& Name): ObjectId(Id), ObjectName(Name)
+	{
+	}
+
+	uint64 ObjectId;
+	FString ObjectName;
+};
+
+class FRewindDebuggerModule : public IModuleInterface
+{
+public:
+	// IModuleInterface interface
+	virtual void StartupModule() override;
+	virtual void ShutdownModule() override;
+
+	TSharedRef<SDockTab> SpawnRewindDebuggerTab(const FSpawnTabArgs& SpawnTabArgs);
+
+private:
+	TSharedPtr<SRewindDebugger> RewindDebuggerWidget;
+
+	FDelegateHandle TickerHandle;
+};
diff --git a/Engine/Plugins/Animation/GameplayInsights/Source/RewindDebugger/Private/RewindDebuggerStyle.cpp b/Engine/Plugins/Animation/GameplayInsights/Source/RewindDebugger/Private/RewindDebuggerStyle.cpp
new file mode 100644
--- /dev/null
+++ b/Engine/Plugins/Animation/GameplayInsights/Source/RewindDebugger/Private/RewindDebuggerStyle.cpp
@@ -0,0 +1,60 @@
+// Copyright Epic Games, Inc. All Rights Reserved.
+
+#include "RewindDebuggerStyle.h"
+#include "EditorStyleSet.h"
+#include "Styling/CoreStyle.h"
+#include "Styling/SlateStyle.h"
+#include "Styling/SlateStyleMacros.h"
+#include "Styling/SlateStyleRegistry.h"
+#include "Styling/SlateTypes.h"
+
+TSharedPtr< FRewindDebuggerStyle > FRewindDebuggerStyle::StyleInstance = nullptr;
+
+
+FRewindDebuggerStyle::FRewindDebuggerStyle() :
+    FSlateStyleSet("RewindDebuggerStyle")
+{
+	const FVector2D Icon24x24(24.0f, 24.0f);
+	const FVector2D Icon16x16(16.0f, 16.0f);
+
+	SetContentRoot(FPaths::EnginePluginsDir() / TEXT("Animation/GameplayInsights/Content"));
+	SetCoreContentRoot(FPaths::EngineContentDir() / TEXT("Editor/Slate"));
+
+	// Toolbar
+	Set("RewindDebugger.StartRecording.small", new CORE_IMAGE_BRUSH("Sequencer/Transport_Bar/Record_24x", Icon24x24));
+	Set("RewindDebugger.StopRecording.small", new CORE_IMAGE_BRUSH("Sequencer/Transport_Bar/Recording_24x", Icon24x24));
+	Set("RewindDebugger.FirstFrame.small", new CORE_IMAGE_BRUSH("Sequencer/Transport_Bar/Go_To_Front_24x", Icon24x24));
+	Set("RewindDebugger.PreviousFrame.small", new CORE_IMAGE_BRUSH("Sequencer/Transport_Bar/Step_Backwards_24x", Icon24x24));
+	Set("RewindDebugger.ReversePlay.small", new CORE_IMAGE_BRUSH("Sequencer/Transport_Bar/Backwards_24x", Icon24x24));
+	Set("RewindDebugger.Pause.small", new CORE_IMAGE_BRUSH("Sequencer/Transport_Bar/Pause_24x", Icon24x24));
+	Set("RewindDebugger.Play.small", new CORE_IMAGE_BRUSH("Sequencer/Transport_Bar/Play_24x", Icon24x24));
+	Set("RewindDebugger.NextFrame.small", new CORE_IMAGE_BRUSH("Sequencer/Transport_Bar/Step_Forward_24x", Icon24x24));
+	Set("RewindDebugger.LastFrame.small", new CORE_IMAGE_BRUSH("Sequencer/Transport_Bar/Go_To_End_24x", Icon24x24));
+
+	// Actor Picker
+	Set("RewindDebugger.SelectActor", new CORE_IMAGE_BRUSH("Icons/eyedropper_16px", Icon16x16));
+
+	// tab icon
+	Set("RewindDebugger.RewindIcon", new IMAGE_BRUSH("Rewind_24x", Icon16x16));
+}
+
+void FRewindDebuggerStyle::Initialize()
+{
+	if (!StyleInstance.IsValid())
+	{
+		StyleInstance = MakeShared<FRewindDebuggerStyle>();
+		FSlateStyleRegistry::RegisterSlateStyle(*StyleInstance);
+	}
+}
+
+void FRewindDebuggerStyle::Shutdown()
+{
+	FSlateStyleRegistry::UnRegisterSlateStyle(*StyleInstance);
+	ensure(StyleInstance.IsUnique());
+	StyleInstance.Reset();
+}
+
+const ISlateStyle& FRewindDebuggerStyle::Get()
+{
+	return *StyleInstance;
+}
diff --git a/Engine/Plugins/Animation/GameplayInsights/Source/RewindDebugger/Private/RewindDebuggerStyle.h b/Engine/Plugins/Animation/GameplayInsights/Source/RewindDebugger/Private/RewindDebuggerStyle.h
new file mode 100644
--- /dev/null
+++ b/Engine/Plugins/Animation/GameplayInsights/Source/RewindDebugger/Private/RewindDebuggerStyle.h
@@ -0,0 +1,23 @@
+// Copyright Epic Games, Inc. All Rights Reserved.
+#pragma once
+
+#include "CoreMinimal.h"
+#include "Styling/SlateStyle.h"
+
+/**
+* Implements the visual style of RewindDebugger
+*/
+class FRewindDebuggerStyle : FSlateStyleSet
+{
+public:
+    FRewindDebuggerStyle();
+	static void Initialize();
+	static void Shutdown();
+
+	static const ISlateStyle& Get();
+
+private:
+
+	static TSharedPtr< FRewindDebuggerStyle > StyleInstance;
+};
+
diff --git a/Engine/Plugins/Animation/GameplayInsights/Source/RewindDebugger/Private/RewindDebuggerTimeSliderController.cpp b/Engine/Plugins/Animation/GameplayInsights/Source/RewindDebugger/Private/RewindDebuggerTimeSliderController.cpp
new file mode 100644
--- /dev/null
+++ b/Engine/Plugins/Animation/GameplayInsights/Source/RewindDebugger/Private/RewindDebuggerTimeSliderController.cpp
@@ -0,0 +1,641 @@
+// Copyright Epic Games, Inc. All Rights Reserved.
+
+#include "RewindDebuggerTimeSliderController.h"
+#include "EditorStyleSet.h"
+#include "Fonts/FontMeasure.h"
+#include "Framework/Application/SlateApplication.h"
+#include "Rendering/DrawElements.h"
+#include "Styling/CoreStyle.h"
+
+
+#define LOCTEXT_NAMESPACE "TimeSlider"
+
+
+namespace ScrubConstants
+{
+	/** The minimum amount of pixels between each major ticks on the widget */
+	const int32 MinPixelsPerDisplayTick = 5;
+
+	/**The smallest number of units between between major tick marks */
+	const float MinDisplayTickSpacing = 0.001f;
+}
+
+
+/** Utility struct for converting between scrub range space and local/absolute screen space */
+struct FTimeSliderController::FScrubRangeToScreen
+{
+	FVector2D WidgetSize;
+
+	TRange<float> ViewInput;
+	float ViewInputRange;
+	float PixelsPerInput;
+
+	FScrubRangeToScreen(TRange<float> InViewInput, const FVector2D& InWidgetSize )
+	{
+		WidgetSize = InWidgetSize;
+
+		ViewInput = InViewInput;
+		ViewInputRange = ViewInput.Size<float>();
+		PixelsPerInput = ViewInputRange > 0 ? ( WidgetSize.X / ViewInputRange ) : 0;
+	}
+
+	/** Local Widget Space -> Curve Input domain. */
+	float LocalXToInput(float ScreenX) const
+	{
+		float LocalX = ScreenX;
+		return (LocalX/PixelsPerInput) + ViewInput.GetLowerBoundValue();
+	}
+
+	/** Curve Input domain -> local Widget Space */
+	float InputToLocalX(float Input) const
+	{
+		return (Input - ViewInput.GetLowerBoundValue()) * PixelsPerInput;
+	}
+};
+
+
+/**
+ * Gets the the next spacing value in the series 
+ * to determine a good spacing value
+ * E.g, .001,.005,.010,.050,.100,.500,1.000,etc
+ */
+static float GetNextSpacing( uint32 CurrentStep )
+{
+	if(CurrentStep & 0x01) 
+	{
+		// Odd numbers
+		return FMath::Pow( 10.f, 0.5f*((float)(CurrentStep-1)) + 1.f );
+	}
+	else 
+	{
+		// Even numbers
+		return 0.5f * FMath::Pow( 10.f, 0.5f*((float)(CurrentStep)) + 1.f );
+	}
+}
+
+FTimeSliderController::FTimeSliderController(const FTimeSliderArgs& InArgs)
+	: TimeSliderArgs( InArgs )
+	, DistanceDragged( 0.0f )
+	, bDraggingScrubber( false )
+	, bPanning( false )
+{
+	ScrubHandleUp = FEditorStyle::GetBrush( TEXT( "Sequencer.Timeline.VanillaScrubHandleUp" ) ); 
+	ScrubHandleDown = FEditorStyle::GetBrush( TEXT( "Sequencer.Timeline.VanillaScrubHandleDown" ) );
+	CursorBackground = FEditorStyle::GetBrush("Sequencer.SectionArea.Background");
+}
+
+float FTimeSliderController::DetermineOptimalSpacing(float InPixelsPerInput, uint32 MinTick, float MinTickSpacing) const
+{
+	if (InPixelsPerInput == 0.0f)
+		return MinTickSpacing;
+
+	uint32 CurStep = 0;
+
+	// Start with the smallest spacing
+	float Spacing = MinTickSpacing;
+
+	while( Spacing * InPixelsPerInput < MinTick )
+	{
+		Spacing = MinTickSpacing * GetNextSpacing( CurStep );
+		CurStep++;
+	}
+
+	return Spacing;
+}
+
+void FTimeSliderController::SetTimesliderArgs(const FTimeSliderArgs& InArgs)
+{
+	TimeSliderArgs = InArgs;
+}
+
+struct FTimeSliderController::FDrawTickArgs
+{
+	/** Geometry of the area */
+	FGeometry AllottedGeometry;
+	/** Clipping rect of the area */
+	FSlateRect ClippingRect;
+	/** Color of each tick */
+	FLinearColor TickColor;
+	/** Offset in Y where to start the tick */
+	float TickOffset;
+	/** Height in of major ticks */
+	float MajorTickHeight;
+	/** Start layer for elements */
+	int32 StartLayer;
+	/** Draw effects to apply */
+	ESlateDrawEffect DrawEffects;
+	/** Whether or not to only draw major ticks */
+	bool bOnlyDrawMajorTicks;
+	/** Whether or not to mirror labels */
+	bool bMirrorLabels;
+	
+};
+
+void FTimeSliderController::DrawTicks( FSlateWindowElementList& OutDrawElements, const struct FScrubRangeToScreen& RangeToScreen, FDrawTickArgs& InArgs ) const
+{
+	const float Spacing = DetermineOptimalSpacing( RangeToScreen.PixelsPerInput, ScrubConstants::MinPixelsPerDisplayTick, ScrubConstants::MinDisplayTickSpacing );
+
+	// Sub divisions
+	// @todo Sequencer may need more robust calculation
+	const int32 Divider = 10;
+	// For slightly larger halfway tick mark
+	const int32 HalfDivider = Divider / 2;
+	// Find out where to start from
+	int32 OffsetNum = FMath::FloorToInt(RangeToScreen.ViewInput.GetLowerBoundValue() / Spacing);
+	
+	FSlateFontInfo SmallLayoutFont = FCoreStyle::GetDefaultFontStyle("Regular", 8);
+
+	TArray<FVector2D> LinePoints;
+	LinePoints.AddUninitialized(2);
+
+	float Seconds = 0;
+	while( (Seconds = OffsetNum*Spacing) < RangeToScreen.ViewInput.GetUpperBoundValue() )
+	{
+		// X position local to start of the widget area
+		float XPos = RangeToScreen.InputToLocalX( Seconds );
+		uint32 AbsOffsetNum = FMath::Abs(OffsetNum);
+
+		if ( AbsOffsetNum % Divider == 0 )
+		{
+			FVector2D Offset( XPos, InArgs.TickOffset );
+			FVector2D TickSize( 1.0f, InArgs.MajorTickHeight );
+
+			LinePoints[0] = FVector2D(1.0f,1.0f);
+			LinePoints[1] = TickSize;
+
+			// lines should not need anti-aliasing
+			const bool bAntiAliasLines = false;
+
+			// Draw each tick mark
+			FSlateDrawElement::MakeLines(
+				OutDrawElements,
+				InArgs.StartLayer,
+				InArgs.AllottedGeometry.ToPaintGeometry( Offset, TickSize ),
+				LinePoints,
+				InArgs.DrawEffects,
+				InArgs.TickColor,
+				false
+				);
+
+			if( !InArgs.bOnlyDrawMajorTicks )
+			{
+				FString FrameString = Spacing == ScrubConstants::MinDisplayTickSpacing ? FString::Printf( TEXT("%.3f"), Seconds ) : FString::Printf( TEXT("%.2f"), Seconds );
+
+				// Space the text between the tick mark but slightly above
+				const TSharedRef< FSlateFontMeasure > FontMeasureService = FSlateApplication::Get().GetRenderer()->GetFontMeasureService();
+				FVector2D TextSize = FontMeasureService->Measure(FrameString, SmallLayoutFont);
+				FVector2D TextOffset( XPos-(TextSize.X*0.5f), InArgs.bMirrorLabels ? TextSize.Y :  FMath::Abs( InArgs.AllottedGeometry.GetLocalSize().Y - (InArgs.MajorTickHeight+TextSize.Y) ) );
+
+				FSlateDrawElement::MakeText(
+					OutDrawElements,
+					InArgs.StartLayer+1, 
+					InArgs.AllottedGeometry.ToPaintGeometry( TextOffset, TextSize ), 
+					FrameString, 
+					SmallLayoutFont, 
+					InArgs.DrawEffects,
+					InArgs.TickColor 
+				);
+			}
+		}
+		else if( !InArgs.bOnlyDrawMajorTicks )
+		{
+			// Compute the size of each tick mark.  If we are half way between to visible values display a slightly larger tick mark
+			const float MinorTickHeight = AbsOffsetNum % HalfDivider == 0 ? 7.0f : 4.0f;
+
+			FVector2D Offset(XPos, InArgs.bMirrorLabels ? 0.0f : FMath::Abs( InArgs.AllottedGeometry.GetLocalSize().Y - MinorTickHeight ) );
+			FVector2D TickSize(1, MinorTickHeight);
+
+			LinePoints[0] = FVector2D(1.0f,1.0f);
+			LinePoints[1] = TickSize;
+
+			const bool bAntiAlias = false;
+			// Draw each sub mark
+			FSlateDrawElement::MakeLines(
+				OutDrawElements,
+				InArgs.StartLayer,
+				InArgs.AllottedGeometry.ToPaintGeometry( Offset, TickSize ),
+				LinePoints,
+				InArgs.DrawEffects,
+				InArgs.TickColor,
+				bAntiAlias
+			);
+		}
+		// Advance to next tick mark
+		++OffsetNum;
+	}
+}
+
+
+int32 FTimeSliderController::OnPaintTimeSlider( bool bMirrorLabels, const FGeometry& AllottedGeometry, const FSlateRect& MyCullingRect, FSlateWindowElementList& OutDrawElements, int32 LayerId, const FWidgetStyle& InWidgetStyle, bool bParentEnabled ) const
+{
+	const bool bEnabled = bParentEnabled;
+	const ESlateDrawEffect DrawEffects = bEnabled ? ESlateDrawEffect::None : ESlateDrawEffect::DisabledEffect;
+
+	TRange<float> LocalViewRange = TimeSliderArgs.ViewRange.Get();
+	const float LocalViewRangeMin = LocalViewRange.GetLowerBoundValue();
+	const float LocalViewRangeMax = LocalViewRange.GetUpperBoundValue();
+	const float LocalSequenceLength = LocalViewRangeMax-LocalViewRangeMin;
+	
+	FVector2D Scale = FVector2D(1.0f,1.0f);
+	if ( LocalSequenceLength > 0)
+	{
+		FScrubRangeToScreen RangeToScreen( LocalViewRange, AllottedGeometry.GetLocalSize() );
+	
+		const float MajorTickHeight = 9.0f;
+	
+		FDrawTickArgs Args;
+		Args.AllottedGeometry = AllottedGeometry;
+		Args.bMirrorLabels = bMirrorLabels;
+		Args.bOnlyDrawMajorTicks = false;
+		Args.TickColor = FLinearColor::White;
+		Args.ClippingRect = MyCullingRect;
+		Args.DrawEffects = DrawEffects;
+		Args.StartLayer = LayerId;
+		Args.TickOffset = bMirrorLabels ? 0.0f : FMath::Abs( AllottedGeometry.GetLocalSize().Y - MajorTickHeight );
+		Args.MajorTickHeight = MajorTickHeight;
+
+		DrawTicks( OutDrawElements, RangeToScreen, Args );
+
+		const float HandleSize = 13.0f;
+		float HalfSize = FMath::TruncToFloat(HandleSize/2.0f);
+
+		// Draw the scrub handle
+		const float XPos = RangeToScreen.InputToLocalX( TimeSliderArgs.ScrubPosition.Get() );
+
+		// Draw cursor size
+		const float CursorHalfSize = TimeSliderArgs.CursorSize.Get() * 0.5f;
+		const int32 CursorLayer = LayerId + 2;
+		const float CursorHalfLength = AllottedGeometry.GetLocalSize().X * CursorHalfSize;
+		FPaintGeometry CursorGeometry = AllottedGeometry.ToPaintGeometry(FVector2D(XPos - CursorHalfLength, 0), FVector2D(2 * CursorHalfLength, AllottedGeometry.GetLocalSize().Y));
+
+		FLinearColor CursorColor = InWidgetStyle.GetColorAndOpacityTint();
+		CursorColor.A = CursorColor.A*0.08f;
+		CursorColor.B *= 0.1f;
+		CursorColor.G *= 0.2f;
+		FSlateDrawElement::MakeBox(
+			OutDrawElements,
+			CursorLayer,
+			CursorGeometry,
+			CursorBackground,
+			DrawEffects,
+			CursorColor
+			);
+
+		// Should draw above the text
+		const int32 ArrowLayer = LayerId + 3;
+		FPaintGeometry MyGeometry =	AllottedGeometry.ToPaintGeometry( FVector2D( XPos-HalfSize, 0 ), FVector2D( HandleSize, AllottedGeometry.GetLocalSize().Y ) );
+		FLinearColor ScrubColor = InWidgetStyle.GetColorAndOpacityTint();
+
+		// @todo Sequencer this color should be specified in the style
+		ScrubColor.A = ScrubColor.A*0.5f;
+		ScrubColor.B *= 0.1f;
+		ScrubColor.G *= 0.2f;
+		FSlateDrawElement::MakeBox( 
+			OutDrawElements,
+			ArrowLayer, 
+			MyGeometry,
+			bMirrorLabels ? ScrubHandleUp : ScrubHandleDown,
+			DrawEffects, 
+			ScrubColor
+			);
+
+		// Draw ClampRange 
+
+		float LeftClamp =  RangeToScreen.InputToLocalX(TimeSliderArgs.ClampRange.Get().GetLowerBoundValue());
+		float RightClamp =  RangeToScreen.InputToLocalX(TimeSliderArgs.ClampRange.Get().GetUpperBoundValue());
+
+		FPaintGeometry RangeGeometry = AllottedGeometry.ToPaintGeometry(FVector2D(LeftClamp, AllottedGeometry.GetLocalSize().Y - 3), FVector2D(RightClamp-LeftClamp, AllottedGeometry.GetLocalSize().Y ));
+
+		FSlateDrawElement::MakeBox(
+			OutDrawElements,
+			++LayerId,
+			RangeGeometry,
+			CursorBackground,
+			DrawEffects,
+			FLinearColor::Red.CopyWithNewOpacity(0.5f)
+			);
+
+
+		return ArrowLayer;
+	}
+
+
+
+	return LayerId;
+}
+
+FReply FTimeSliderController::OnMouseButtonDown( SWidget& WidgetOwner, const FGeometry& MyGeometry, const FPointerEvent& MouseEvent )
+{
+	bool bHandleLeftMouseButton = MouseEvent.GetEffectingButton() == EKeys::LeftMouseButton;
+	bool bHandleRightMouseButton = MouseEvent.GetEffectingButton() == EKeys::RightMouseButton && TimeSliderArgs.AllowZoom;
+	
+	DistanceDragged = 0;
+
+	if ( bHandleLeftMouseButton )
+	{
+		// Always capture mouse if we left or right click on the widget
+		FScrubRangeToScreen RangeToScreen(TimeSliderArgs.ViewRange.Get(), MyGeometry.GetLocalSize());
+		FVector2D CursorPos = MyGeometry.AbsoluteToLocal(MouseEvent.GetLastScreenSpacePosition());
+		float NewValue = RangeToScreen.LocalXToInput(CursorPos.X);
+
+		CommitScrubPosition(NewValue, /*bIsScrubbing=*/false);
+		return FReply::Handled().CaptureMouse( WidgetOwner.AsShared() ).PreventThrottling();
+	}
+	else if ( bHandleRightMouseButton )
+	{
+		return FReply::Handled().CaptureMouse(WidgetOwner.AsShared());
+	}
+
+	return FReply::Unhandled();
+}
+
+FReply FTimeSliderController::OnMouseButtonUp( SWidget& WidgetOwner, const FGeometry& MyGeometry, const FPointerEvent& MouseEvent )
+{
+	bool bHandleLeftMouseButton = MouseEvent.GetEffectingButton() == EKeys::LeftMouseButton && WidgetOwner.HasMouseCapture();
+	bool bHandleRightMouseButton = MouseEvent.GetEffectingButton() == EKeys::RightMouseButton && WidgetOwner.HasMouseCapture() && TimeSliderArgs.AllowZoom;
+	
+	if ( bHandleRightMouseButton )
+	{
+		if (!bPanning)
+		{
+			// return unhandled in case our parent wants to use our right mouse button to open a context menu
+			return FReply::Unhandled().ReleaseMouseCapture();
+		}
+		
+		bPanning = false;
+		return FReply::Handled().ReleaseMouseCapture();
+	}
+	else if ( bHandleLeftMouseButton )
+	{
+		if( bDraggingScrubber )
+		{
+			TimeSliderArgs.OnEndScrubberMovement.ExecuteIfBound();
+		}
+		else
+		{
+			FScrubRangeToScreen RangeToScreen( TimeSliderArgs.ViewRange.Get(), MyGeometry.GetLocalSize() );
+			FVector2D CursorPos = MyGeometry.AbsoluteToLocal(MouseEvent.GetLastScreenSpacePosition());
+			float NewValue = RangeToScreen.LocalXToInput(CursorPos.X);
+
+			CommitScrubPosition( NewValue, /*bIsScrubbing=*/false );
+		}
+
+		bDraggingScrubber = false;
+		return FReply::Handled().ReleaseMouseCapture();
+
+	}
+
+	return FReply::Unhandled();
+}
+
+float FTimeSliderController::GetTimeAtCursorPosition(const FGeometry& MyGeometry, const FPointerEvent& MouseEvent) const
+{
+	FScrubRangeToScreen RangeToScreen(TimeSliderArgs.ViewRange.Get(), MyGeometry.GetLocalSize());
+	FVector2D CursorPos = MyGeometry.AbsoluteToLocal(MouseEvent.GetLastScreenSpacePosition());
+	float NewValue = RangeToScreen.LocalXToInput(CursorPos.X);
+
+	float LocalClampMin = TimeSliderArgs.ClampRange.Get().GetLowerBoundValue();
+	float LocalClampMax = TimeSliderArgs.ClampRange.Get().GetUpperBoundValue();
+
+	if (NewValue < LocalClampMin)
+	{
+		NewValue = LocalClampMin;
+	}
+
+	if (NewValue > LocalClampMax)
+	{
+		NewValue = LocalClampMax;
+	}
+
+	return NewValue;
+}
+
+FReply FTimeSliderController::OnMouseMove( SWidget& WidgetOwner, const FGeometry& MyGeometry, const FPointerEvent& MouseEvent )
+{
+	if ( WidgetOwner.HasMouseCapture() )
+	{
+		if (MouseEvent.IsMouseButtonDown(EKeys::RightMouseButton))
+		{
+			if (!bPanning)
+			{
+				DistanceDragged += FMath::Abs( MouseEvent.GetCursorDelta().X );
+				if ( DistanceDragged > FSlateApplication::Get().GetDragTriggerDistance() )
+				{
+					FReply::Handled().CaptureMouse(WidgetOwner.AsShared()).UseHighPrecisionMouseMovement(WidgetOwner.AsShared());
+					SoftwareCursorPosition = MyGeometry.AbsoluteToLocal(MouseEvent.GetLastScreenSpacePosition());
+					bPanning = true;
+				}
+			}
+			else
+			{
+				SoftwareCursorPosition = MyGeometry.AbsoluteToLocal(MouseEvent.GetLastScreenSpacePosition());
+
+				TRange<float> LocalViewRange = TimeSliderArgs.ViewRange.Get();
+				float LocalViewRangeMin = LocalViewRange.GetLowerBoundValue();
+				float LocalViewRangeMax = LocalViewRange.GetUpperBoundValue();
+
+				FScrubRangeToScreen ScaleInfo( LocalViewRange, MyGeometry.GetLocalSize() );
+				FVector2D ScreenDelta = MouseEvent.GetCursorDelta();
+				FVector2D InputDelta;
+				InputDelta.X = ScreenDelta.X/ScaleInfo.PixelsPerInput;
+
+				const TRange<float> NewViewRange = TRange<float>(LocalViewRangeMin - InputDelta.X, LocalViewRangeMax - InputDelta.X);
+				TRange<float> LocalClampRange = TimeSliderArgs.ClampRange.Get();
+
+				// Do not try to pan outside the clamp range to prevent undesired zoom
+				//if (LocalClampRange.Contains(NewViewRange))
+				{
+					TimeSliderArgs.OnViewRangeChanged.ExecuteIfBound(NewViewRange);
+					if (Scrollbar.IsValid())
+					{
+						float InOffsetFraction = (NewViewRange.GetLowerBoundValue() - LocalClampRange.GetLowerBoundValue()) / LocalClampRange.Size<float>();
+						float InThumbSizeFraction = NewViewRange.Size<float>() / LocalClampRange.Size<float>();
+						Scrollbar->SetState(InOffsetFraction, InThumbSizeFraction);
+					}
+
+					if (!TimeSliderArgs.ViewRange.IsBound())
+					{
+						// The  output is not bound to a delegate so we'll manage the value ourselves
+						TimeSliderArgs.ViewRange.Set(NewViewRange);
+					}
+				}
+			}
+		}
+		else if (MouseEvent.IsMouseButtonDown( EKeys::LeftMouseButton ))
+		{
+			if ( !bDraggingScrubber )
+			{
+				DistanceDragged += FMath::Abs( MouseEvent.GetCursorDelta().X );
+				if ( DistanceDragged > 0/*FSlateApplication::Get().GetDragTriggerDistance()*/ )
+				{
+					bDraggingScrubber = true;
+					TimeSliderArgs.OnBeginScrubberMovement.ExecuteIfBound();
+				}
+			}
+			else
+			{
+				const float NewValue = GetTimeAtCursorPosition(MyGeometry, MouseEvent);
+				CommitScrubPosition(NewValue, /*bIsScrubbing=*/true);
+			}
+		}
+		return FReply::Handled();
+	}
+
+	return FReply::Unhandled();
+}
+
+void FTimeSliderController::CommitScrubPosition( float NewValue, bool bIsScrubbing )
+{
+	// Manage the scrub position ourselves if its not bound to a delegate
+	if ( !TimeSliderArgs.ScrubPosition.IsBound() )
+	{
+		TimeSliderArgs.ScrubPosition.Set( NewValue );
+	}
+
+	if (!TimeSliderArgs.ViewRange.IsBound())
+	{
+		TRange<float> LocalViewRange = TimeSliderArgs.ViewRange.Get();
+		const float RangeSize = LocalViewRange.Size<float>();
+		if (NewValue < LocalViewRange.GetLowerBoundValue())
+		{
+			SetTimeRange(NewValue, NewValue + RangeSize);
+		}
+		else if (NewValue > LocalViewRange.GetUpperBoundValue())
+		{
+			SetTimeRange(NewValue - RangeSize, NewValue);
+		}
+	}
+
+	TimeSliderArgs.OnScrubPositionChanged.ExecuteIfBound( NewValue, bIsScrubbing );
+}
+
+void FTimeSliderController::SetExternalScrollbar(TSharedRef<SScrollBar> InScrollbar) 
+{ 
+	Scrollbar = InScrollbar;
+	Scrollbar->SetOnUserScrolled(FOnUserScrolled::CreateRaw(this, &FTimeSliderController::HorizontalScrollBar_OnUserScrolled)); 
+};
+
+void FTimeSliderController::HorizontalScrollBar_OnUserScrolled(float ScrollOffset)
+{
+	if (!TimeSliderArgs.ViewRange.IsBound())
+	{
+		TRange<float> LocalViewRange = TimeSliderArgs.ViewRange.Get();
+		float LocalViewRangeMin = LocalViewRange.GetLowerBoundValue();
+		float LocalViewRangeMax = LocalViewRange.GetUpperBoundValue();
+		float LocalClampMin = TimeSliderArgs.ClampRange.Get().GetLowerBoundValue();
+		float LocalClampMax = TimeSliderArgs.ClampRange.Get().GetUpperBoundValue();
+
+		float InThumbSizeFraction = (LocalViewRangeMax - LocalViewRangeMin) / (LocalClampMax - LocalClampMin);
+
+		float NewViewOutputMin = LocalClampMin + ScrollOffset * (LocalClampMax - LocalClampMin);
+		// The  output is not bound to a delegate so we'll manage the value ourselves
+		float NewViewOutputMax = FMath::Min<float>(NewViewOutputMin + (LocalViewRangeMax - LocalViewRangeMin), LocalClampMax);
+		NewViewOutputMin = NewViewOutputMax - (LocalViewRangeMax - LocalViewRangeMin);
+
+		float InOffsetFraction = (NewViewOutputMin - LocalClampMin) / (LocalClampMax - LocalClampMin);
+		//if (InOffsetFraction + InThumbSizeFraction <= 1)
+		{
+			TimeSliderArgs.ViewRange.Set(TRange<float>(NewViewOutputMin, NewViewOutputMax));
+			Scrollbar->SetState(InOffsetFraction, InThumbSizeFraction);
+		}
+	}
+}
+
+void FTimeSliderController::SetTimeRange(float NewViewOutputMin, float NewViewOutputMax)
+{
+	TimeSliderArgs.ViewRange.Set(TRange<float>(NewViewOutputMin, NewViewOutputMax));
+
+	float LocalClampMin = TimeSliderArgs.ClampRange.Get().GetLowerBoundValue();
+	float LocalClampMax = TimeSliderArgs.ClampRange.Get().GetUpperBoundValue();
+
+	const float InOffsetFraction = (NewViewOutputMin - LocalClampMin) / (LocalClampMax - LocalClampMin);
+	const float InThumbSizeFraction = (NewViewOutputMax - NewViewOutputMin) / (LocalClampMax - LocalClampMin);
+	if(Scrollbar.IsValid())
+	{
+		Scrollbar->SetState(InOffsetFraction, InThumbSizeFraction);
+	}
+}
+
+void FTimeSliderController::SetClampRange(float MinValue, float MaxValue)
+{
+	TRange<float> LocalViewRange = TimeSliderArgs.ViewRange.Get();
+	float LocalClampMin = TimeSliderArgs.ClampRange.Get().GetLowerBoundValue();
+	float LocalClampMax = TimeSliderArgs.ClampRange.Get().GetUpperBoundValue();
+	const float CurrentDistance = LocalClampMax - LocalClampMin;
+	const float ZoomDelta = (LocalViewRange.GetUpperBoundValue() - LocalViewRange.GetLowerBoundValue()) / CurrentDistance;
+
+	MaxValue = MinValue + (MaxValue - MinValue < 2 ? CurrentDistance : MaxValue - MinValue);
+
+	TimeSliderArgs.ClampRange = TRange<float>(MinValue, MaxValue);
+
+	const float LocalViewRangeMin = FMath::Clamp(LocalViewRange.GetLowerBoundValue(), MinValue, MaxValue);
+	const float LocalViewRangeMax = FMath::Clamp(LocalViewRange.GetUpperBoundValue(), MinValue, MaxValue);
+	SetTimeRange(ZoomDelta >= 1 ? MinValue : LocalViewRangeMin, ZoomDelta >= 1 ? MaxValue : LocalViewRangeMax);
+}
+
+FReply FTimeSliderController::OnMouseWheel( SWidget& WidgetOwner, const FGeometry& MyGeometry, const FPointerEvent& MouseEvent )
+{
+	FReply ReturnValue = FReply::Unhandled();
+
+	// if (MouseEvent.IsLeftShiftDown())
+	// {
+	// 	const float ZoomDelta = 0.025f * MouseEvent.GetWheelDelta();
+	// 	TimeSliderArgs.CursorSize.Set(FMath::Clamp(TimeSliderArgs.CursorSize.Get() + ZoomDelta, 0.0f, 1.0f));
+
+	// 	ReturnValue = FReply::Handled();
+	// }
+	// else
+	if ( TimeSliderArgs.AllowZoom )
+	{
+		const float ZoomDelta = -0.1f * MouseEvent.GetWheelDelta();
+
+		{
+			float MouseFractionX = MyGeometry.AbsoluteToLocal(MouseEvent.GetScreenSpacePosition()).X / MyGeometry.GetLocalSize().X;
+
+			TRange<float> LocalViewRange = TimeSliderArgs.ViewRange.Get();
+			float LocalViewRangeMax = LocalViewRange.GetUpperBoundValue();
+			float LocalViewRangeMin = LocalViewRange.GetLowerBoundValue();
+			const float OutputViewSize = LocalViewRangeMax - LocalViewRangeMin;
+			const float OutputChange = OutputViewSize * ZoomDelta;
+
+			float NewViewOutputMin = LocalViewRangeMin - (OutputChange * MouseFractionX);
+			float NewViewOutputMax = LocalViewRangeMax + (OutputChange * (1.0f - MouseFractionX));
+
+			if (NewViewOutputMin < NewViewOutputMax)
+			{
+				float LocalClampMin = TimeSliderArgs.ClampRange.Get().GetLowerBoundValue();
+				float LocalClampMax = TimeSliderArgs.ClampRange.Get().GetUpperBoundValue();
+
+				// // Clamp the range if clamp values are set
+				// if ( NewViewOutputMin < LocalClampMin )
+				// {
+				// 	NewViewOutputMin = LocalClampMin;
+				// }
+				
+				// if ( NewViewOutputMax > LocalClampMax )
+				// {
+				// 	NewViewOutputMax = LocalClampMax;
+				// }
+
+				TimeSliderArgs.OnViewRangeChanged.ExecuteIfBound(TRange<float>(NewViewOutputMin, NewViewOutputMax));
+				if (Scrollbar.IsValid())
+				{
+					float InOffsetFraction = (NewViewOutputMin - LocalClampMin) / (LocalClampMax - LocalClampMin);
+					float InThumbSizeFraction = (NewViewOutputMax - NewViewOutputMin) / (LocalClampMax - LocalClampMin);
+					Scrollbar->SetState(InOffsetFraction, InThumbSizeFraction);
+				}
+				if( !TimeSliderArgs.ViewRange.IsBound() )
+				{	
+					// The  output is not bound to a delegate so we'll manage the value ourselves
+					TimeSliderArgs.ViewRange.Set( TRange<float>( NewViewOutputMin, NewViewOutputMax ) );
+				}
+			}
+		}
+
+		ReturnValue = FReply::Handled();
+	}
+
+	return ReturnValue;
+}
+
+#undef LOCTEXT_NAMESPACE
diff --git a/Engine/Plugins/Animation/GameplayInsights/Source/RewindDebugger/Private/RewindDebuggerTimeSliderController.h b/Engine/Plugins/Animation/GameplayInsights/Source/RewindDebugger/Private/RewindDebuggerTimeSliderController.h
new file mode 100644
--- /dev/null
+++ b/Engine/Plugins/Animation/GameplayInsights/Source/RewindDebugger/Private/RewindDebuggerTimeSliderController.h
@@ -0,0 +1,132 @@
+// Copyright Epic Games, Inc. All Rights Reserved.
+
+#pragma once
+
+#include "CoreMinimal.h"
+#include "Misc/Attribute.h"
+#include "Input/CursorReply.h"
+#include "Input/Reply.h"
+#include "Widgets/SWidget.h"
+#include "Widgets/Layout/SScrollBar.h"
+
+class FSlateWindowElementList;
+
+struct FTimeSliderArgs
+{
+	DECLARE_DELEGATE_OneParam( FOnRangeChanged, TRange<float> )
+	DECLARE_DELEGATE_TwoParams( FOnScrubPositionChanged, float, bool )
+
+	FTimeSliderArgs()
+		: ScrubPosition(0)
+		, ViewRange(TRange<float>(0.0f, 5.0f))
+		, ClampRange(TRange<float>(0.0f, 5.0f))
+		, AllowZoom(true)
+		, CursorSize(0.05f)
+	{}
+
+	/** The scrub position */
+	TAttribute<float> ScrubPosition;
+
+	/** Called when the scrub position changes */
+	FOnScrubPositionChanged OnScrubPositionChanged;
+
+	/** Called right before the scrubber begins to move */
+	FSimpleDelegate OnBeginScrubberMovement;
+
+	/** Called right after the scrubber handle is released by the user */
+	FSimpleDelegate OnEndScrubberMovement;
+
+	/** View time range */
+	TAttribute< TRange<float> > ViewRange;
+
+	/** Clamp time range */
+	TAttribute< TRange<float> > ClampRange;
+
+	/** Called when the view range changes */
+	FOnRangeChanged OnViewRangeChanged;
+
+	/** Attribute defining the time snap interval */
+	TAttribute<float> TimeSnapInterval;
+
+	/** If we are allowed to zoom */
+	bool AllowZoom;
+
+	/** Cursor range for data like histogram graphs, etc. */
+	TAttribute< float > CursorSize;
+};
+
+/**
+ * A time slider controller
+ * Draws and manages time data
+ */
+class FTimeSliderController
+{
+public:
+	FTimeSliderController(const FTimeSliderArgs& InArgs);
+
+	/**
+	* Determines the optimal spacing between tick marks in the slider for a given pixel density
+	* Increments until a minimum amount of slate units specified by MinTick is reached
+	*
+	* @param InPixelsPerInput	The density of pixels between each input
+	* @param MinTick			The minimum slate units per tick allowed
+	* @param MinTickSpacing	The minimum tick spacing in time units allowed
+	* @return the optimal spacing in time units
+	*/
+	float DetermineOptimalSpacing(float InPixelsPerInput, uint32 MinTick, float MinTickSpacing) const;
+	void SetTimesliderArgs(const FTimeSliderArgs& InArgs);
+
+	int32 OnPaintTimeSlider( bool bMirrorLabels, const FGeometry& AllottedGeometry, const FSlateRect& MyCullingRect, FSlateWindowElementList& OutDrawElements, int32 LayerId, const FWidgetStyle& InWidgetStyle, bool bParentEnabled ) const;
+	FReply OnMouseButtonDown( SWidget& WidgetOwner, const FGeometry& MyGeometry, const FPointerEvent& MouseEvent );
+	FReply OnMouseButtonUp( SWidget& WidgetOwner, const FGeometry& MyGeometry, const FPointerEvent& MouseEvent );
+	FReply OnMouseMove( SWidget& WidgetOwner, const FGeometry& MyGeometry, const FPointerEvent& MouseEvent );
+	FReply OnMouseWheel( SWidget& WidgetOwner, const FGeometry& MyGeometry, const FPointerEvent& MouseEvent );
+	FCursorReply OnCursorQuery( TSharedRef<const SWidget> WidgetOwner, const FGeometry& MyGeometry, const FPointerEvent& CursorEvent) const { return FCursorReply::Unhandled(); }
+
+	FTimeSliderArgs& GetTimeSliderArgs() { return TimeSliderArgs; }
+	/**
+	 * Call this method when the user's interaction has changed the scrub position
+	 *
+	 * @param NewValue				Value resulting from the user's interaction
+	 * @param bIsScrubbing			True if done via scrubbing, false if just releasing scrubbing
+	 */
+	void CommitScrubPosition( float NewValue, bool bIsScrubbing );
+
+	void SetExternalScrollbar(TSharedRef<SScrollBar> Scrollbar);
+	void SetTimeRange(float MinValue, float MaxValue);
+	void SetClampRange(float MinValue, float MaxValue);
+	bool IsPanning() { return bPanning; }
+
+	float GetTimeAtCursorPosition(const FGeometry& MyGeometry, const FPointerEvent& MouseEvent) const;
+private:
+	// forward declared as class members to prevent name collision with similar types defined in other units
+	struct FScrubRangeToScreen;
+	struct FDrawTickArgs;
+
+	/**
+	 * Draws time tick marks
+	 *
+	 * @param OutDrawElements	List to add draw elements to
+	 * @param RangeToScreen		Time range to screen space converter
+	 * @param InArgs			Parameters for drawing the tick lines
+	 */
+	void DrawTicks( FSlateWindowElementList& OutDrawElements, const FScrubRangeToScreen& RangeToScreen, FDrawTickArgs& InArgs ) const;
+	void HorizontalScrollBar_OnUserScrolled(float ScrollOffset);
+private:
+	FTimeSliderArgs TimeSliderArgs;
+	/** Brush for drawing an upwards facing scrub handle */
+	const FSlateBrush* ScrubHandleUp;
+	/** Brush for drawing a downwards facing scrub handle */
+	const FSlateBrush* ScrubHandleDown;
+	/** Brush for drawing cursor background to visualize corsor size */
+	const FSlateBrush* CursorBackground;
+	/** Total mouse delta during dragging **/
+	float DistanceDragged;
+	/** If we are dragging the scrubber */
+	bool bDraggingScrubber;
+	/** If we are currently panning the panel */
+	bool bPanning;
+	/***/
+	TSharedPtr<SScrollBar> Scrollbar;
+	FVector2D SoftwareCursorPosition;
+};
diff --git a/Engine/Plugins/Animation/GameplayInsights/Source/RewindDebugger/Private/SRewindDebugger.cpp b/Engine/Plugins/Animation/GameplayInsights/Source/RewindDebugger/Private/SRewindDebugger.cpp
new file mode 100644
--- /dev/null
+++ b/Engine/Plugins/Animation/GameplayInsights/Source/RewindDebugger/Private/SRewindDebugger.cpp
@@ -0,0 +1,361 @@
+// Copyright Epic Games, Inc. All Rights Reserved.
+
+#include "SRewindDebugger.h"
+#include "ActorPickerMode.h"
+#include "Async/Future.h"
+#include "Components/SkeletalMeshComponent.h"
+#include "Editor.h"
+#include "Editor/EditorEngine.h"
+#include "Framework/Application/SlateApplication.h"
+#include "Framework/MultiBox/MultiBoxBuilder.h"
+#include "Insights/IUnrealInsightsModule.h"
+#include "IGameplayInsightsModule.h"
+#include "ISettingsModule.h"
+#include "Modules/ModuleManager.h"
+#include "ObjectTrace.h"
+#include "RewindDebuggerStyle.h"
+#include "RewindDebuggerCommands.h"
+#include "SRewindDebuggerTimeSlider.h"
+#include "SceneOutlinerPublicTypes.h"
+#include "SceneOutlinerModule.h"
+#include "Selection.h"
+#include "Styling/SlateIconFinder.h"
+#include "Widgets/Colors/SColorBlock.h"
+#include "Widgets/Images/SImage.h"
+#include "Widgets/Input/SComboButton.h"
+#include "Widgets/Input/SButton.h"
+#include "Widgets/Input/SSearchBox.h"
+#include "Widgets/Input/STextComboBox.h"
+#include "Widgets/Input/STextComboBox.h"
+#include "Widgets/SOverlay.h"
+#include "Widgets/Views/SListView.h"
+
+
+#define LOCTEXT_NAMESPACE "SAnimationInsights"
+
+SRewindDebugger::SRewindDebugger() 
+	: SCompoundWidget()
+	, DebugComponents(nullptr)
+{ 
+}
+
+SRewindDebugger::~SRewindDebugger() 
+{
+
+}
+
+void SRewindDebugger::TrackCursor(bool bReverse)
+{
+	float ScrubTime = ScrubTimeAttribute.Get();
+	TRange<float> CurrentViewRange = TimeSliderController->GetTimeSliderArgs().ViewRange.Get();
+	float ViewSize = CurrentViewRange.GetUpperBoundValue() - CurrentViewRange.GetLowerBoundValue();
+
+	static const float LeadingEdgeSize = 0.05f;
+	static const float TrailingEdgeThreshold = 0.01f;
+
+	if(bReverse)
+	{
+		// playing in reverse (cursor moving left)
+		if (ScrubTime < CurrentViewRange.GetLowerBoundValue() + ViewSize * LeadingEdgeSize)
+		{
+			CurrentViewRange.SetLowerBound(ScrubTime - ViewSize * LeadingEdgeSize);
+			CurrentViewRange.SetUpperBound(CurrentViewRange.GetLowerBoundValue() + ViewSize);
+		}
+		if (ScrubTime > CurrentViewRange.GetUpperBoundValue() + ViewSize * TrailingEdgeThreshold)
+		{
+			CurrentViewRange.SetUpperBound(ScrubTime);
+			CurrentViewRange.SetLowerBound(CurrentViewRange.GetUpperBoundValue() - ViewSize);
+		}
+	}
+	else
+	{
+		// playing normally or recording (cursor moving right)
+		if (ScrubTime > CurrentViewRange.GetUpperBoundValue() - ViewSize * LeadingEdgeSize)
+		{
+			CurrentViewRange.SetUpperBound(ScrubTime + ViewSize * LeadingEdgeSize);
+			CurrentViewRange.SetLowerBound(CurrentViewRange.GetUpperBoundValue() - ViewSize);
+		}
+		if (ScrubTime < CurrentViewRange.GetLowerBoundValue() - ViewSize * TrailingEdgeThreshold)
+		{
+			CurrentViewRange.SetLowerBound(ScrubTime);
+			CurrentViewRange.SetUpperBound(CurrentViewRange.GetLowerBoundValue() + ViewSize);
+		}
+	}
+
+
+	TimeSliderController->GetTimeSliderArgs().ViewRange = CurrentViewRange;
+}
+
+void SRewindDebugger::SetDebugTargetActor(AActor* Actor)
+{
+	DebugTargetActor.Set(Actor->GetName());
+}
+
+TSharedRef<SWidget> SRewindDebugger::MakeSelectActorMenu()
+{
+	// this menu is partially duplicated from LevelSequenceEditorActorBinding which has a similar workflow for adding actors to sequencer
+
+	FMenuBuilder MenuBuilder(true, nullptr);
+
+	// Set up a menu entry to choosing the selected actor(s) to the sequencer (maybe move this to a submenu and put each selected actor there)
+	TArray<AActor*> SelectedActors;
+	GEditor->GetSelectedActors()->GetSelectedObjects(SelectedActors);
+
+	if (SelectedActors.Num() >= 1)
+	{
+		MenuBuilder.BeginSection("From Selection Section", LOCTEXT("FromSelection", "From Selection"));
+		if (SelectedActors.Num() == 1)
+		{
+			AActor* SelectedActor = SelectedActors[0];
+
+			FText SelectedLabel = FText::FromString(SelectedActor->GetActorLabel());
+			FSlateIcon ActorIcon = FSlateIconFinder::FindIconForClass(SelectedActors[0]->GetClass());
+
+			MenuBuilder.AddMenuEntry(SelectedLabel, FText(), ActorIcon, FExecuteAction::CreateLambda([this, SelectedActor]{
+				FSlateApplication::Get().DismissAllMenus();
+				SetDebugTargetActor(SelectedActor);
+			}));
+		}
+		else if (SelectedActors.Num() >= 1)
+		{
+			MenuBuilder.AddSubMenu(
+				LOCTEXT("FromSelection", "From Selection"),
+				LOCTEXT("FromSelection_Tooltip", "Select an Actor from the list of selected Actors"),
+				FNewMenuDelegate::CreateLambda([this, SelectedActors](FMenuBuilder& SubMenuBuilder)
+				{
+					for(AActor* SelectedActor : SelectedActors)
+					{
+						FText SelectedLabel = FText::FromString(SelectedActor->GetActorLabel());
+						FSlateIcon ActorIcon = FSlateIconFinder::FindIconForClass(SelectedActors[0]->GetClass());
+
+						SubMenuBuilder.AddMenuEntry(SelectedLabel, FText(), ActorIcon, FExecuteAction::CreateLambda([this, SelectedActor]{
+							FSlateApplication::Get().DismissAllMenus();
+							SetDebugTargetActor(SelectedActor);
+						}));
+					}
+
+				})
+			);
+
+		}
+		MenuBuilder.EndSection();
+	}
+
+	// todo: add special menu item for player controlled character
+
+	MenuBuilder.BeginSection("ChooseActorSection", LOCTEXT("ChooseActor", "Choose Actor:"));
+
+	// Set up a menu entry to select any arbitrary actor
+	FSceneOutlinerInitializationOptions InitOptions;
+	{
+		// We hide the header row to keep the UI compact.
+		InitOptions.bShowHeaderRow = false;
+		InitOptions.bShowSearchBox = true;
+		InitOptions.bShowCreateNewFolder = false;
+		InitOptions.bFocusSearchBoxWhenOpened = true;
+
+		// Only want the actor label column
+		InitOptions.ColumnMap.Add(FSceneOutlinerBuiltInColumnTypes::Label(), FSceneOutlinerColumnInfo(ESceneOutlinerColumnVisibility::Visible, 0));
+
+		// todo: optionally filter for only actors that have debug data
+		//InitOptions.Filters->AddFilterPredicate<FActorTreeItem>(FActorTreeItem::FFilterPredicate::CreateLambda(IsActorValidForPossession, ExistingPossessedObjects));
+	}
+
+	// actor selector to allow the user to choose an actor
+	FSceneOutlinerModule& SceneOutlinerModule = FModuleManager::LoadModuleChecked<FSceneOutlinerModule>("SceneOutliner");
+	TSharedRef< SWidget > MiniSceneOutliner =
+		SNew(SBox)
+		.MaxDesiredHeight(400.0f)
+		.WidthOverride(300.0f)
+		[
+			SceneOutlinerModule.CreateActorPicker(
+				InitOptions,
+				FOnActorPicked::CreateLambda([this](AActor* Actor){
+					// Create a new binding for this actor
+					FSlateApplication::Get().DismissAllMenus();
+					SetDebugTargetActor(Actor);
+				})
+			)
+		];
+
+	MenuBuilder.AddWidget(MiniSceneOutliner, FText::GetEmpty(), true);
+	MenuBuilder.EndSection();
+
+
+
+	return MenuBuilder.MakeWidget();
+}
+
+
+void SRewindDebugger::Construct(const FArguments& InArgs, TSharedRef<FUICommandList> CommandList, const TSharedRef<SDockTab>& ConstructUnderMajorTab, const TSharedPtr<SWindow>& ConstructUnderWindow)
+{
+	OnScrubPositionChanged = InArgs._OnScrubPositionChanged;
+	ScrubTimeAttribute = InArgs._ScrubTime;
+	DebugComponents = InArgs._DebugComponents;
+	TraceTime.Initialize(InArgs._TraceTime);
+	RecordingDuration.Initialize(InArgs._RecordingDuration);
+	DebugTargetActor.Initialize(InArgs._DebugTargetActor);
+	
+	FSlimHorizontalToolBarBuilder ToolBarBuilder(CommandList, FMultiBoxCustomization::None, nullptr, true);
+
+	const FRewindDebuggerCommands& Commands = FRewindDebuggerCommands::Get();
+
+	ToolBarBuilder.BeginSection("Debugger");
+	{
+		ToolBarBuilder.AddToolBarButton(Commands.FirstFrame);
+		ToolBarBuilder.AddToolBarButton(Commands.PreviousFrame);
+		ToolBarBuilder.AddToolBarButton(Commands.ReversePlay);
+		ToolBarBuilder.AddToolBarButton(Commands.Pause);
+		ToolBarBuilder.AddToolBarButton(Commands.Play);
+		ToolBarBuilder.AddToolBarButton(Commands.NextFrame);
+		ToolBarBuilder.AddToolBarButton(Commands.LastFrame);
+		ToolBarBuilder.AddToolBarButton(Commands.StartRecording);
+		ToolBarBuilder.AddToolBarButton(Commands.StopRecording);
+	}
+	ToolBarBuilder.EndSection();
+
+	FTimeSliderArgs InitArgs;
+	InitArgs.CursorSize = 0.0f;
+	InitArgs.OnScrubPositionChanged = FTimeSliderArgs::FOnScrubPositionChanged::CreateLambda(
+		[this](float NewScrubTime, bool bIsScrubbing)
+		{
+			if (bIsScrubbing)
+			{
+				OnScrubPositionChanged.ExecuteIfBound( NewScrubTime, bIsScrubbing );
+			}
+		}
+	);
+
+	InitArgs.ScrubPosition = ScrubTimeAttribute;
+
+	InitArgs.ClampRange = InitArgs.ClampRange.CreateLambda(
+		[this]()
+		{ 
+			return TRange<float>(0.0f,RecordingDuration.Get());
+		} );
+
+	TimeSliderController = MakeShared<FTimeSliderController>(InitArgs);
+    TSharedRef<STimeSlider> TimeSlider = SNew(STimeSlider, TimeSliderController.ToSharedRef());
+
+	ComponentListView =	SNew(SListView<TSharedPtr<FDebugObjectInfo>>)
+									.ItemHeight(16.0f)
+									.ListItemsSource(DebugComponents)
+									.OnGenerateRow(this, &SRewindDebugger::ComponentListViewGenerateRow)
+									.SelectionMode(ESelectionMode::Single)
+									.OnSelectionChanged(this, &SRewindDebugger::OnComponentSelectionChanged);
+
+
+	// everything related to creating SAnimGraphSchematicView should be moved to a separate file
+	IUnrealInsightsModule &UnrealInsightsModule = FModuleManager::LoadModuleChecked<IUnrealInsightsModule>("TraceInsights");
+	IGameplayInsightsModule& GameplayInsightsModule = FModuleManager::LoadModuleChecked<IGameplayInsightsModule>("GameplayInsights");
+	TSharedPtr<const TraceServices::IAnalysisSession> Session = UnrealInsightsModule.GetAnalysisSession();
+
+	AnimGraphView = GameplayInsightsModule.CreateAnimGraphSchematicView(0, TraceTime.Get(), *Session);
+	TraceTime.OnPropertyChanged = TraceTime.OnPropertyChanged.CreateSP(AnimGraphView.Get(), &IAnimGraphSchematicView::SetTimeMarker);
+
+	FSlateIcon ActorIcon = FSlateIconFinder::FindIconForClass(AActor::StaticClass());
+
+	ChildSlot
+	[
+		SNew(SSplitter)
+		+SSplitter::Slot().MinSize(270).Value(0)
+		[
+			SNew(SVerticalBox)
+			+SVerticalBox::Slot().AutoHeight()
+			[
+				ToolBarBuilder.MakeWidget()
+			]
+			+SVerticalBox::Slot().AutoHeight()
+			[
+				SNew(SHorizontalBox)
+				+SHorizontalBox::Slot() .FillWidth(1.0f)
+				[
+					SNew(SComboButton)
+					.OnGetMenuContent(this, &SRewindDebugger::MakeSelectActorMenu)
+					.ButtonContent()
+					[
+						SNew(STextBlock)
+						.Text_Lambda([this](){
+							const FString& TargetName = DebugTargetActor.Get();
+							return (TargetName.IsEmpty()) ? LOCTEXT("Select Actor", "Debug Target Actor") : FText::FromString(TargetName);
+						 } )
+					]
+				]
+				+SHorizontalBox::Slot().AutoWidth().HAlign(HAlign_Right)
+				[
+					SNew(SButton)
+						.ButtonStyle(FAppStyle::Get(), "SimpleButton")
+						.OnClicked(this, &SRewindDebugger::OnSelectActorClicked)
+					[
+						SNew(SImage)
+						.Image(FRewindDebuggerStyle::Get().GetBrush("RewindDebugger.SelectActor"))
+					]
+				]
+			]
+			+SVerticalBox::Slot().VAlign(VAlign_Top) .FillHeight(1.0f)
+			[
+				ComponentListView.ToSharedRef()
+			]
+		]
+		+SSplitter::Slot() 
+		[
+			SNew(SVerticalBox)
+			+SVerticalBox::Slot() .VAlign(VAlign_Top) .AutoHeight()
+			[
+				TimeSlider
+			]
+			+SVerticalBox::Slot() .VAlign(VAlign_Top) .FillHeight(1.0f)
+			[
+				// todo: this should be a proxy slot that we swap out based on which component is selected in the treeview
+				AnimGraphView.ToSharedRef()
+			]
+		]
+	];
+}
+
+TSharedRef<ITableRow> SRewindDebugger::ComponentListViewGenerateRow(TSharedPtr<FDebugObjectInfo> InItem, const TSharedRef<STableViewBase>& OwnerTable)
+{
+	return 
+		SNew(STableRow<TSharedPtr<FDebugObjectInfo>>, OwnerTable)
+		[
+			SNew(STextBlock)
+			.Text(FText::FromString(InItem->ObjectName))
+		];
+}
+
+void SRewindDebugger::RefreshDebugComponents()
+{
+	ComponentListView->RebuildList();
+}
+
+void SRewindDebugger::OnComponentSelectionChanged(TSharedPtr<FDebugObjectInfo> SelectedItem, ESelectInfo::Type SelectInfo)
+{
+	if (SelectedItem.IsValid())
+	{
+		AnimGraphView->SetAnimInstanceId(SelectedItem->ObjectId);
+	}
+	else
+	{
+		AnimGraphView->SetAnimInstanceId(0);
+	}
+}
+
+FReply SRewindDebugger::OnSelectActorClicked()
+{
+		FActorPickerModeModule& ActorPickerMode = FModuleManager::Get().GetModuleChecked<FActorPickerModeModule>("ActorPickerMode");
+		
+		// todo: force eject (from within BeginActorPickingMode?)
+
+		ActorPickerMode.BeginActorPickingMode(
+			FOnGetAllowedClasses(), 
+			FOnShouldFilterActor(),
+			FOnActorSelected::CreateRaw(this, &SRewindDebugger::SetDebugTargetActor));
+
+
+
+	return FReply::Handled();
+}
+
+
+
+#undef LOCTEXT_NAMESPACE
diff --git a/Engine/Plugins/Animation/GameplayInsights/Source/RewindDebugger/Private/SRewindDebugger.h b/Engine/Plugins/Animation/GameplayInsights/Source/RewindDebugger/Private/SRewindDebugger.h
new file mode 100644
--- /dev/null
+++ b/Engine/Plugins/Animation/GameplayInsights/Source/RewindDebugger/Private/SRewindDebugger.h
@@ -0,0 +1,79 @@
+// Copyright Epic Games, Inc. All Rights Reserved.
+#pragma once
+
+#include "BindableProperty.h"
+#include "CoreMinimal.h"
+#include "Framework/Commands/UICommandList.h"
+#include "Framework/Docking/TabManager.h"
+#include "IAnimGraphSchematicView.h"
+#include "RewindDebuggerModule.h"
+#include "RewindDebuggerTimeSliderController.h"
+#include "Widgets/DeclarativeSyntaxSupport.h"
+#include "Widgets/Input/SSearchBox.h"
+#include "Widgets/SCompoundWidget.h"
+#include "Widgets/Views/SListView.h"
+
+
+class SRewindDebugger : public SCompoundWidget
+{
+	typedef TBindablePropertyInitializer<FString, BindingType_Out> DebugTargetInitializer;
+
+public:
+	DECLARE_DELEGATE_TwoParams( FOnScrubPositionChanged, float, bool )
+	DECLARE_DELEGATE_OneParam( FOnDebugTargetChanged, TSharedPtr<FString> )
+
+	SLATE_BEGIN_ARGS(SRewindDebugger) { }
+		SLATE_ARGUMENT( TArray< TSharedPtr< FDebugObjectInfo > >*, DebugComponents );
+		SLATE_ARGUMENT(DebugTargetInitializer, DebugTargetActor);
+		SLATE_ARGUMENT(TBindablePropertyInitializer<double>, TraceTime);
+		SLATE_ARGUMENT(TBindablePropertyInitializer<float>, RecordingDuration);
+		SLATE_ATTRIBUTE(float, ScrubTime);
+		SLATE_EVENT(FOnScrubPositionChanged, OnScrubPositionChanged);
+	SLATE_END_ARGS()
+	
+public:
+
+	/**
+	* Default constructor.
+	*/
+	SRewindDebugger();
+	virtual ~SRewindDebugger();
+
+	/**
+	* Constructs the application.
+	*
+	* @param InArgs The Slate argument list.
+	* @param ConstructUnderMajorTab The major tab which will contain the session front-end.
+	* @param ConstructUnderWindow The window in which this widget is being constructed.
+	* @param InStyleSet The style set to use.
+	*/
+	void Construct(const FArguments& InArgs, TSharedRef<FUICommandList> CommandList, const TSharedRef<SDockTab>& ConstructUnderMajorTab, const TSharedPtr<SWindow>& ConstructUnderWindow);
+
+	void TrackCursor(bool bReverse);
+	void RefreshDebugComponents();
+private:
+	TSharedPtr<FTimeSliderController> TimeSliderController;
+
+	TAttribute<float> ScrubTimeAttribute;
+	TAttribute<bool> TrackScrubbingAttribute;
+	FOnScrubPositionChanged OnScrubPositionChanged;
+
+	// debug actor selector
+	TSharedRef<SWidget> MakeSelectActorMenu();
+	void SetDebugTargetActor(AActor* Actor);
+	FReply OnSelectActorClicked();
+
+	TBindableProperty<FString, BindingType_Out> DebugTargetActor;
+	TBindableProperty<double> TraceTime;
+	TBindableProperty<float> RecordingDuration;
+	TBindableProperty<float> DebugTargetAnimInstanceId;
+
+	// debug components list
+	TArray<TSharedPtr<FDebugObjectInfo>>* DebugComponents;
+    void OnComponentSelectionChanged(TSharedPtr<FDebugObjectInfo> SelectedItem, ESelectInfo::Type SelectInfo);
+	TSharedPtr<SListView<TSharedPtr<FDebugObjectInfo>>> ComponentListView;
+	TSharedRef<ITableRow> ComponentListViewGenerateRow(TSharedPtr<FDebugObjectInfo> InItem, const TSharedRef<STableViewBase>& OwnerTable);
+
+	// anim graph view 
+	TSharedPtr<IAnimGraphSchematicView> AnimGraphView;
+};
diff --git a/Engine/Plugins/Animation/GameplayInsights/Source/RewindDebugger/Private/SRewindDebuggerTimeSlider.cpp b/Engine/Plugins/Animation/GameplayInsights/Source/RewindDebugger/Private/SRewindDebuggerTimeSlider.cpp
new file mode 100644
--- /dev/null
+++ b/Engine/Plugins/Animation/GameplayInsights/Source/RewindDebugger/Private/SRewindDebuggerTimeSlider.cpp
@@ -0,0 +1,50 @@
+// Copyright Epic Games, Inc. All Rights Reserved.
+
+#include "SRewindDebuggerTimeSlider.h"
+#include "RewindDebuggerTimeSliderController.h"
+
+#define LOCTEXT_NAMESPACE "SRewindDebuggerTimeSlider"
+
+
+void STimeSlider::Construct( const STimeSlider::FArguments& InArgs, TSharedRef<FTimeSliderController> InTimeSliderController )
+{
+	TimeSliderController = InTimeSliderController;
+	bMirrorLabels = InArgs._MirrorLabels;
+
+	// set clipping on by default, since the OnPaint function is drawing outside the bounds
+	Clipping = EWidgetClipping::ClipToBounds;
+}
+
+int32 STimeSlider::OnPaint( const FPaintArgs& Args, const FGeometry& AllottedGeometry, const FSlateRect& MyCullingRect, FSlateWindowElementList& OutDrawElements, int32 LayerId, const FWidgetStyle& InWidgetStyle, bool bParentEnabled ) const
+{
+	int32 NewLayer = TimeSliderController->OnPaintTimeSlider( bMirrorLabels, AllottedGeometry, MyCullingRect, OutDrawElements, LayerId, InWidgetStyle, bParentEnabled );
+
+	return FMath::Max( NewLayer, SCompoundWidget::OnPaint(Args, AllottedGeometry, MyCullingRect, OutDrawElements, NewLayer, InWidgetStyle, ShouldBeEnabled( bParentEnabled ) ) );
+}
+
+FReply STimeSlider::OnMouseButtonDown( const FGeometry& MyGeometry, const FPointerEvent& MouseEvent )
+{
+	return TimeSliderController->OnMouseButtonDown( *this, MyGeometry, MouseEvent );
+}
+
+FReply STimeSlider::OnMouseButtonUp( const FGeometry& MyGeometry, const FPointerEvent& MouseEvent )
+{
+	return TimeSliderController->OnMouseButtonUp( *this,  MyGeometry, MouseEvent );
+}
+
+FReply STimeSlider::OnMouseMove( const FGeometry& MyGeometry, const FPointerEvent& MouseEvent )
+{
+	return TimeSliderController->OnMouseMove( *this, MyGeometry, MouseEvent );
+}
+
+FVector2D STimeSlider::ComputeDesiredSize( float ) const
+{
+	return FVector2D(100, 22);
+}
+
+FReply STimeSlider::OnMouseWheel( const FGeometry& MyGeometry, const FPointerEvent& MouseEvent )
+{
+	return TimeSliderController->OnMouseWheel( *this, MyGeometry, MouseEvent );
+}
+
+#undef LOCTEXT_NAMESPACE
diff --git a/Engine/Plugins/Animation/GameplayInsights/Source/RewindDebugger/Private/SRewindDebuggerTimeSlider.h b/Engine/Plugins/Animation/GameplayInsights/Source/RewindDebugger/Private/SRewindDebuggerTimeSlider.h
new file mode 100644
--- /dev/null
+++ b/Engine/Plugins/Animation/GameplayInsights/Source/RewindDebugger/Private/SRewindDebuggerTimeSlider.h
@@ -0,0 +1,45 @@
+// Copyright Epic Games, Inc. All Rights Reserved.
+
+#pragma once
+
+#include "CoreMinimal.h"
+#include "Input/Reply.h"
+#include "Widgets/SCompoundWidget.h"
+#include "Widgets/DeclarativeSyntaxSupport.h"
+
+class FPaintArgs;
+class FSlateWindowElementList;
+class FTimeSliderController;
+
+class STimeSlider : public SCompoundWidget
+{
+public:
+
+	SLATE_BEGIN_ARGS(STimeSlider)
+		: _MirrorLabels( false )
+	{}
+		/* If we should mirror the labels on the timeline */
+		SLATE_ARGUMENT( bool, MirrorLabels )
+	SLATE_END_ARGS()
+
+
+	/**
+	 * Construct the widget
+	 * 
+	 * @param InArgs   A declaration from which to construct the widget
+	 */
+	void Construct( const FArguments& InArgs, TSharedRef<FTimeSliderController> InTimeSliderController );
+
+protected:
+	// SWidget interface
+	virtual FVector2D ComputeDesiredSize(float) const override;
+	virtual int32 OnPaint( const FPaintArgs& Args, const FGeometry& AllottedGeometry, const FSlateRect& MyCullingRect, FSlateWindowElementList& OutDrawElements, int32 LayerId, const FWidgetStyle& InWidgetStyle, bool bParentEnabled ) const override;
+	virtual FReply OnMouseButtonDown( const FGeometry& MyGeometry, const FPointerEvent& MouseEvent ) override;
+	virtual FReply OnMouseButtonUp( const FGeometry& MyGeometry, const FPointerEvent& MouseEvent ) override;
+	virtual FReply OnMouseMove( const FGeometry& MyGeometry, const FPointerEvent& MouseEvent ) override;
+	virtual FReply OnMouseWheel( const FGeometry& MyGeometry, const FPointerEvent& MouseEvent ) override;
+private:
+	TSharedPtr<FTimeSliderController> TimeSliderController;
+	bool bMirrorLabels;
+};
+
diff --git a/Engine/Plugins/Animation/GameplayInsights/Source/RewindDebugger/RewindDebugger.Build.cs b/Engine/Plugins/Animation/GameplayInsights/Source/RewindDebugger/RewindDebugger.Build.cs
new file mode 100644
--- /dev/null
+++ b/Engine/Plugins/Animation/GameplayInsights/Source/RewindDebugger/RewindDebugger.Build.cs
@@ -0,0 +1,41 @@
+// Copyright Epic Games, Inc. All Rights Reserved.
+
+namespace UnrealBuildTool.Rules
+{
+	public class RewindDebugger : ModuleRules
+	{
+		public RewindDebugger(ReadOnlyTargetRules Target) : base(Target)
+		{
+			PrivateDependencyModuleNames.AddRange(
+			new string[]
+			{
+				"ActorPickerMode",
+				"AnimationBlueprintEditor",
+				"ApplicationCore",
+				"AssetRegistry",
+				"Core",
+				"CoreUObject",
+				"EditorFramework",
+				"EditorStyle",
+				"EditorWidgets",
+				"Engine",
+				"GameplayInsights",
+				"GameplayInsightsEditor",
+				"InputCore",
+				"Kismet",
+				"Persona",
+				"SceneOutliner",
+				"SequencerWidgets",
+				"SlateCore",
+				"Slate",
+				"ToolMenus",
+				"TraceLog",
+				"TraceAnalysis",
+				"TraceServices",
+				"TraceInsights",
+				"UnrealEd"
+			});
+		}
+	}
+}
+
diff --git a/Engine/Source/Developer/TraceServices/Public/Model/IntervalTimeline.h b/Engine/Source/Developer/TraceServices/Public/Model/IntervalTimeline.h
--- a/Engine/Source/Developer/TraceServices/Public/Model/IntervalTimeline.h
+++ b/Engine/Source/Developer/TraceServices/Public/Model/IntervalTimeline.h
@@ -76,6 +76,71 @@ public:
 			}
 		}
 	}
+
+	virtual void FindNearestEvents(double InTime, const EventType*& OutPrecedingEvent, double& OutPrecedingEventTime, const EventType*& OutFollowingEvent, double& OutFollowingEventTime) const override
+	{
+		OutPrecedingEvent =  nullptr;
+		OutFollowingEvent =  nullptr;
+
+		if (Events.Num() == 0)
+		{
+			return;
+		}
+
+		auto EventIterator = Events.GetIteratorFromPage(0);
+		const FEventInternal* Event = EventIterator.GetCurrentItem();
+		const FEventPage* Page = EventIterator.GetCurrentPage();
+
+		while(Event != nullptr && InTime > Page->EndTime)
+		{
+			EventIterator.NextPage();
+			Page = EventIterator.GetCurrentPage();
+			Event = EventIterator.GetCurrentItem();
+		}
+
+		if (Event == nullptr)
+		{
+			// InTime is past the end, so just output the last event and no following event
+			const FEventInternal& LastEvent = Events.Last();
+			OutPrecedingEvent = &LastEvent.Event;
+			OutPrecedingEventTime = LastEvent.StartTime;
+			return;
+		}
+		if (Event->StartTime > InTime)
+		{
+			// first event on page is past our time, so it will be the following event, preceding event is the last event from the previous page
+			OutFollowingEvent = &Event->Event;
+			OutFollowingEventTime = Event->StartTime;
+
+			EventIterator.PrevPage();
+			if (Event == EventIterator.GetCurrentItem())
+			{
+				OutPrecedingEvent = &Event->Event;
+				OutPrecedingEventTime = Event->StartTime;
+			}
+
+			return;
+		}
+
+		// normal case, just find the event before and after InTime on this page
+		while(Event != nullptr)
+		{
+			const FEventInternal* NextEvent = EventIterator.NextItem();
+			if (NextEvent)
+			{
+				if (NextEvent->StartTime > InTime)
+				{
+					OutPrecedingEvent = &Event->Event;
+					OutPrecedingEventTime = Event->StartTime;
+					OutFollowingEvent = &NextEvent->Event;
+					OutFollowingEventTime = Event->StartTime;
+					return;
+				}
+
+			}
+			Event = NextEvent;
+		}
+	}
 	
 	uint64 AppendBeginEvent(double StartTime, const EventType& Event)
 	{
diff --git a/Engine/Source/Developer/TraceServices/Public/Model/MonotonicTimeline.h b/Engine/Source/Developer/TraceServices/Public/Model/MonotonicTimeline.h
--- a/Engine/Source/Developer/TraceServices/Public/Model/MonotonicTimeline.h
+++ b/Engine/Source/Developer/TraceServices/Public/Model/MonotonicTimeline.h
@@ -160,6 +160,13 @@ public:
 		return DetailLevels[0].Events[InIndex]; 
 	}
 
+	virtual void FindNearestEvents(double InTime, const EventType*& OutPrecedingEvent, double& OutPrecedingEventTime, const EventType*& OutFollowingEvent, double& OutFollowingEventTime) const override
+	{
+		// not implemented
+		OutPrecedingEvent =  nullptr;
+		OutFollowingEvent =  nullptr;
+	}
+
 	virtual double GetStartTime() const override
 	{
 		return DetailLevels[0].ScopeEntries.Num() > 0 ? FMath::Abs(DetailLevels[0].ScopeEntries[0].Time) : 0.0;
diff --git a/Engine/Source/Developer/TraceServices/Public/Model/PointTimeline.h b/Engine/Source/Developer/TraceServices/Public/Model/PointTimeline.h
--- a/Engine/Source/Developer/TraceServices/Public/Model/PointTimeline.h
+++ b/Engine/Source/Developer/TraceServices/Public/Model/PointTimeline.h
@@ -144,6 +144,71 @@ public:
 		return false;
 	}
 
+	virtual void FindNearestEvents(double InTime, const EventType*& OutPrecedingEvent, double& OutPrecedingEventTime, const EventType*& OutFollowingEvent, double& OutFollowingEventTime) const override
+	{
+		OutPrecedingEvent =  nullptr;
+		OutFollowingEvent =  nullptr;
+
+		if (Events.Num() == 0)
+		{
+			return;
+		}
+
+		auto EventIterator = Events.GetIteratorFromPage(0);
+		const FEventInternal* Event = EventIterator.GetCurrentItem();
+		const FEventPage* Page = EventIterator.GetCurrentPage();
+
+		while(Event != nullptr && InTime > Page->EndTime)
+		{
+			EventIterator.NextPage();
+			Page = EventIterator.GetCurrentPage();
+			Event = EventIterator.GetCurrentItem();
+		}
+
+		if (Event == nullptr)
+		{
+			// InTime is past the end, so just output the last event and no following event
+			const FEventInternal& LastEvent = Events.Last();
+			OutPrecedingEvent = &LastEvent.Event;
+			OutPrecedingEventTime = LastEvent.Time;
+			return;
+		}
+		if (Event->Time > InTime)
+		{
+			// first event on page is past our time, so it will be the following event, preceding event is the last event from the previous page
+			OutFollowingEvent = &Event->Event;
+			OutFollowingEventTime = Event->Time;
+
+			EventIterator.PrevPage();
+			if (Event == EventIterator.GetCurrentItem())
+			{
+				OutPrecedingEvent = &Event->Event;
+				OutPrecedingEventTime = Event->Time;
+			}
+
+			return;
+		}
+
+		// normal case, just find the event before and after InTime on this page
+		while(Event != nullptr)
+		{
+			const FEventInternal* NextEvent = EventIterator.NextItem();
+			if (NextEvent)
+			{
+				if (NextEvent->Time > InTime)
+				{
+					OutPrecedingEvent = &Event->Event;
+					OutPrecedingEventTime = Event->Time;
+					OutFollowingEvent = &NextEvent->Event;
+					OutFollowingEventTime = Event->Time;
+					return;
+				}
+
+			}
+			Event = NextEvent;
+		}
+	}
+
 private:
 	struct FEventInternal
 	{
diff --git a/Engine/Source/Developer/TraceServices/Public/TraceServices/Containers/Timelines.h b/Engine/Source/Developer/TraceServices/Public/TraceServices/Containers/Timelines.h
--- a/Engine/Source/Developer/TraceServices/Public/TraceServices/Containers/Timelines.h
+++ b/Engine/Source/Developer/TraceServices/Public/TraceServices/Containers/Timelines.h
@@ -59,6 +59,7 @@ public:
 	virtual void EnumerateEvents(double IntervalStart, double IntervalEnd, EventCallback Callback) const = 0;
 	virtual void EnumerateEvents(double IntervalStart, double IntervalEnd, EventRangeCallback Callback) const = 0;
 	virtual int32 GetDepthAt(double Time) const { return 0; };
+	virtual void FindNearestEvents(double InTime, const EventType*& OutPrecedingEvent, double& OutPrecedingEventTime, const EventType*& OutFollowingEvent, double& OutFollowingEventTime) const = 0;
 	
 	/**
 	 * Finds event information for the event closest to InTime from the interval [InTime - DeltaTime, InTime + DeltaTime]
diff --git a/Engine/Source/Runtime/Engine/Classes/Components/SkeletalMeshComponent.h b/Engine/Source/Runtime/Engine/Classes/Components/SkeletalMeshComponent.h
--- a/Engine/Source/Runtime/Engine/Classes/Components/SkeletalMeshComponent.h
+++ b/Engine/Source/Runtime/Engine/Classes/Components/SkeletalMeshComponent.h
@@ -2384,6 +2384,11 @@ public:
 	/** Consume and return pending root motion from our internal anim instances (main, sub and post) */
 	FRootMotionMovementParams ConsumeRootMotion();
 
+#if WITH_EDITOR
+	/** Called after modifying Component Space Transforms externally */
+	void ApplyEditedComponentSpaceTransforms();
+#endif
+
 protected:
 
 	/** Consume and return pending root motion from our internal anim instances (main, sub and post) */
diff --git a/Engine/Source/Runtime/Engine/Private/Components/SkeletalMeshComponent.cpp b/Engine/Source/Runtime/Engine/Private/Components/SkeletalMeshComponent.cpp
--- a/Engine/Source/Runtime/Engine/Private/Components/SkeletalMeshComponent.cpp
+++ b/Engine/Source/Runtime/Engine/Private/Components/SkeletalMeshComponent.cpp
@@ -772,6 +772,22 @@ void USkeletalMeshComponent::InitAnim(bool bForceReinit)
 	}
 }
 
+#if WITH_EDITOR
+void USkeletalMeshComponent::ApplyEditedComponentSpaceTransforms()
+{
+	// Flip buffers once to copy the directly-written component space transforms
+	bNeedToFlipSpaceBaseBuffers = true;
+	bHasValidBoneTransform = false;
+	FlipEditableSpaceBases();
+	bHasValidBoneTransform = true;
+
+	InvalidateCachedBounds();
+	UpdateBounds();
+	MarkRenderTransformDirty();
+	MarkRenderDynamicDataDirty();
+}
+#endif
+
 bool USkeletalMeshComponent::InitializeAnimScriptInstance(bool bForceReinit, bool bInDeferRootNodeInitialization)
 {
 	bool bInitializedMainInstance = false;
diff --git a/Engine/Source/Runtime/Engine/Private/ObjectTrace.cpp b/Engine/Source/Runtime/Engine/Private/ObjectTrace.cpp
--- a/Engine/Source/Runtime/Engine/Private/ObjectTrace.cpp
+++ b/Engine/Source/Runtime/Engine/Private/ObjectTrace.cpp
@@ -5,6 +5,7 @@
 #if OBJECT_TRACE_ENABLED
 
 #include "CoreMinimal.h"
+#include "GameFramework/WorldSettings.h"
 #include "Trace/Trace.inl"
 #include "UObject/Object.h"
 #include "UObject/Class.h"
@@ -48,6 +49,14 @@ UE_TRACE_EVENT_BEGIN(Object, World)
 	UE_TRACE_EVENT_FIELD(bool, IsSimulating)
 UE_TRACE_EVENT_END()
 
+UE_TRACE_EVENT_BEGIN(Object, RecordingInfo)
+	UE_TRACE_EVENT_FIELD(uint64, WorldId)
+	UE_TRACE_EVENT_FIELD(uint64, Cycle)
+	UE_TRACE_EVENT_FIELD(uint32, RecordingIndex)
+	UE_TRACE_EVENT_FIELD(uint32, FrameIndex)
+	UE_TRACE_EVENT_FIELD(double, ElapsedTime)
+UE_TRACE_EVENT_END()
+
 // Object annotations used for tracing
 struct FTracedObjectAnnotation
 {
@@ -75,18 +84,31 @@ FUObjectAnnotationSparse<FTracedObjectAnnotation, true> GObjectTraceAnnotations;
 // Handle used to hook to world tick
 static FDelegateHandle WorldTickStartHandle;
 
-void FObjectTrace::Init()
+void TickObjectTraceWorldSubsystem(UWorld* InWorld, ELevelTick InTickType, float InDeltaSeconds)
 {
-	WorldTickStartHandle = FWorldDelegates::OnWorldTickStart.AddLambda([](UWorld* InWorld, ELevelTick InTickType, float InDeltaSeconds)
+	if(InTickType == LEVELTICK_All)
 	{
-		if(InTickType == LEVELTICK_All)
+		if (!InWorld->IsPaused())
 		{
 			if(UObjectTraceWorldSubsystem* Subsystem = UWorld::GetSubsystem<UObjectTraceWorldSubsystem>(InWorld))
 			{
 				Subsystem->FrameIndex++;
+				Subsystem->ElapsedTime += InDeltaSeconds;
+
+				UE_TRACE_LOG(Object, RecordingInfo, ObjectChannel)
+		            << RecordingInfo.WorldId(FObjectTrace::GetObjectId(InWorld))
+					<< RecordingInfo.Cycle(FPlatformTime::Cycles64())
+					<< RecordingInfo.ElapsedTime(Subsystem->ElapsedTime)
+					<< RecordingInfo.FrameIndex(Subsystem->FrameIndex)
+					<< RecordingInfo.RecordingIndex(Subsystem->RecordingIndex); 
 			}
 		}
-	});
+	}
+}
+
+void FObjectTrace::Init()
+{
+	WorldTickStartHandle = FWorldDelegates::OnWorldTickStart.AddStatic(&TickObjectTraceWorldSubsystem);
 }
 
 void FObjectTrace::Destroy()
@@ -129,6 +151,72 @@ uint64 FObjectTrace::GetObjectId(const UObject* InObject)
 	return Id | (OuterId << 32);
 }
 
+void FObjectTrace::ResetWorldElapsedTime(const UWorld* World)
+{
+	if(UObjectTraceWorldSubsystem* WorldSubsystem = UWorld::GetSubsystem<UObjectTraceWorldSubsystem>(World))
+	{
+		WorldSubsystem->ElapsedTime = 0;
+	}
+}
+
+double FObjectTrace::GetWorldElapsedTime(const UWorld* World)
+{
+	if(UObjectTraceWorldSubsystem* WorldSubsystem = UWorld::GetSubsystem<UObjectTraceWorldSubsystem>(World))
+	{
+		return WorldSubsystem->ElapsedTime;
+	}
+	return 0;
+}
+
+double FObjectTrace::GetObjectWorldElapsedTime(const UObject* InObject)
+{
+	if(InObject != nullptr)
+	{
+		if(UWorld* World = InObject->GetWorld())
+		{
+			if(UObjectTraceWorldSubsystem* WorldSubsystem = UWorld::GetSubsystem<UObjectTraceWorldSubsystem>(World))
+			{
+				return WorldSubsystem->ElapsedTime;
+			}
+		}
+	}
+
+	return 0;
+}
+
+void FObjectTrace::SetWorldRecordingIndex(const UWorld* World, uint16 Index)
+{
+	if(UObjectTraceWorldSubsystem* WorldSubsystem = UWorld::GetSubsystem<UObjectTraceWorldSubsystem>(World))
+	{
+		WorldSubsystem->RecordingIndex = Index;
+	}
+}
+
+uint16 FObjectTrace::GetWorldRecordingIndex(const UWorld* World)
+{
+	if(UObjectTraceWorldSubsystem* WorldSubsystem = UWorld::GetSubsystem<UObjectTraceWorldSubsystem>(World))
+	{
+		return WorldSubsystem->RecordingIndex;
+	}
+	return 0;
+}
+
+uint16 FObjectTrace::GetObjectWorldRecordingIndex(const UObject* InObject)
+{
+	if(InObject != nullptr)
+	{
+		if(UWorld* World = InObject->GetWorld())
+		{
+			if(UObjectTraceWorldSubsystem* WorldSubsystem = UWorld::GetSubsystem<UObjectTraceWorldSubsystem>(World))
+			{
+				return WorldSubsystem->RecordingIndex;
+			}
+		}
+	}
+
+	return 0;
+}
+
 uint16 FObjectTrace::GetObjectWorldTickCounter(const UObject* InObject)
 {
 	if(InObject != nullptr)
@@ -153,6 +241,8 @@ void FObjectTrace::OutputClass(const UClass* InClass)
 		return;
 	}
 
+	OutputClass(InClass->GetSuperClass());
+
 	FTracedObjectAnnotation Annotation = GObjectTraceAnnotations.GetAnnotation(InClass);
 	if(Annotation.bTraced)
 	{
@@ -177,6 +267,7 @@ void FObjectTrace::OutputClass(const UClass* InClass)
 		<< Class.Id(GetObjectId(InClass))
 		<< Class.SuperId(GetObjectId(InClass->GetSuperClass()))
 		<< Class.Attachment(StringCopyFunc);
+
 }
 
 void FObjectTrace::OutputObject(const UObject* InObject)
@@ -204,6 +295,8 @@ void FObjectTrace::OutputObject(const UObject* InObject)
 		return;
 	}
 
+	OutputObject(InObject->GetOuter());
+
 	Annotation.bTraced = true;
 	GObjectTraceAnnotations.AddAnnotation(InObject, MoveTemp(Annotation));
 
diff --git a/Engine/Source/Runtime/Engine/Public/ObjectTrace.h b/Engine/Source/Runtime/Engine/Public/ObjectTrace.h
--- a/Engine/Source/Runtime/Engine/Public/ObjectTrace.h
+++ b/Engine/Source/Runtime/Engine/Public/ObjectTrace.h
@@ -22,18 +22,23 @@ class UObjectTraceWorldSubsystem : public UWorldSubsystem
 	GENERATED_BODY()
 
 	UObjectTraceWorldSubsystem()
-		: FrameIndex(0)
+		: FrameIndex(0), RecordingIndex(0), ElapsedTime(0.0) 
 	{}
 
 public:
 	// The frame index incremented each tick
 	uint16 FrameIndex;
+	// Trace Recording identifier  (set by RewindDebugger or 0)
+	uint16 RecordingIndex;
+	// Elapsed time since recording started  (or since start of game if RewindDebugger didn't start the trace)
+	double ElapsedTime;
 };
 
 #if OBJECT_TRACE_ENABLED
 
 class UClass;
 class UObject;
+class UWorld;
 
 struct FObjectTrace
 {
@@ -58,8 +63,27 @@ struct FObjectTrace
 	/** Helper function to get an object's world's tick counter */
 	ENGINE_API static uint16 GetObjectWorldTickCounter(const UObject* InObject);
 
+	/** reset the world elapsed time to 0 */
+	ENGINE_API static void ResetWorldElapsedTime(const UWorld* InWorld);
+
+	/** Helper function to get a world's elapsed time */
+	ENGINE_API static double GetWorldElapsedTime(const UWorld* InWorld);
+
+	/** Helper function to get an object's world's elapsed time */
+	ENGINE_API static double GetObjectWorldElapsedTime(const UObject* InObject);
+
+	/** Helper function to set a world's recording index */
+	ENGINE_API static void SetWorldRecordingIndex(const UWorld *World, uint16 Index);
+
+	/** Helper function to get a world's recording index */
+	ENGINE_API static uint16 GetWorldRecordingIndex(const UWorld *InWorld);
+
+	/** Helper function to get an object's world's recording index */
+	ENGINE_API static uint16 GetObjectWorldRecordingIndex(const UObject* InObject);
+
 	/** Helper function to output a world */
 	ENGINE_API static void OutputWorld(const UWorld* InWorld);
+
 };
 
 #define TRACE_CLASS(Class) \
