Fix for non-deterministic crash in RewindDebugger on startup.

diff --git a/Engine/Plugins/Animation/GameplayInsights/Source/RewindDebugger/Private/RewindDebugger.cpp b/Engine/Plugins/Animation/GameplayInsights/Source/RewindDebugger/Private/RewindDebugger.cpp
--- a/Engine/Plugins/Animation/GameplayInsights/Source/RewindDebugger/Private/RewindDebugger.cpp
+++ b/Engine/Plugins/Animation/GameplayInsights/Source/RewindDebugger/Private/RewindDebugger.cpp
@@ -401,18 +401,32 @@ void FRewindDebugger::UpdateTraceTime()
 
 void FRewindDebugger::Tick(float DeltaTime)
 {
+	if (UnrealInsightsModule == nullptr)
+	{
+		UnrealInsightsModule = &FModuleManager::LoadModuleChecked<IUnrealInsightsModule>("TraceInsights");
+	    if (UnrealInsightsModule == nullptr)
+		{
+			return;
+		}
+	}
+
+	TSharedPtr<const TraceServices::IAnalysisSession> Session = UnrealInsightsModule->GetAnalysisSession();
+	if (!Session.IsValid())
+	{
+		return;
+	}
+
 	if (bRecording)
 	{
 		// if you select a debug target before you start recording, update component list when it becomes valid
 		RefreshDebugComponents();
 	}
 
-	TSharedPtr<const TraceServices::IAnalysisSession> Session = UnrealInsightsModule->GetAnalysisSession();
-	TraceServices::FAnalysisSessionReadScope SessionReadScope(*Session.Get());
-	UWorld* World = GetWorldToVisualize();
-
 	if (const IAnimationProvider* AnimationProvider = Session->ReadProvider<IAnimationProvider>("AnimationProvider"))
 	{
+		TraceServices::FAnalysisSessionReadScope SessionReadScope(*Session.Get());
+		UWorld* World = GetWorldToVisualize();
+
 		if (bPIESimulating)
 		{
 			if (bRecording)
