More null checking for analysis session

diff --git a/Engine/Plugins/Animation/GameplayInsights/Source/RewindDebugger/Private/RewindDebugger.cpp b/Engine/Plugins/Animation/GameplayInsights/Source/RewindDebugger/Private/RewindDebugger.cpp
--- a/Engine/Plugins/Animation/GameplayInsights/Source/RewindDebugger/Private/RewindDebugger.cpp
+++ b/Engine/Plugins/Animation/GameplayInsights/Source/RewindDebugger/Private/RewindDebugger.cpp
@@ -109,79 +109,81 @@ void FRewindDebugger::OnPIEStopped(bool bSimulating)
 
 void FRewindDebugger::RefreshDebugComponents()
 {
-	TSharedPtr<const TraceServices::IAnalysisSession> Session = UnrealInsightsModule->GetAnalysisSession();
-	TraceServices::FAnalysisSessionReadScope SessionReadScope(*Session.Get());
-	UWorld* World = GetWorldToVisualize();
-
-	const IGameplayProvider* GameplayProvider = Session->ReadProvider<IGameplayProvider>("GameplayProvider");
-	const IAnimationProvider* AnimationProvider = Session->ReadProvider<IAnimationProvider>("AnimationProvider");
-
-    TArray<TSharedPtr<FDebugObjectInfo>> NewComponentList;	
-	
-	if (DebugTargetActor.Get() == "")
+	if (const TraceServices::IAnalysisSession* Session = GetAnalysisSession())
 	{
-		return;
-	}
-
-	uint64 TargetActorId = 0;
+		TraceServices::FAnalysisSessionReadScope SessionReadScope(*Session);
+		UWorld* World = GetWorldToVisualize();
 
-	// find the Id of the selected actor (move this to a function or cache it)
-	GameplayProvider->EnumerateObjects([this,&TargetActorId](const FObjectInfo& InObjectInfo)
-	{
-		if (DebugTargetActor.Get() == InObjectInfo.Name)
+		if (const IGameplayProvider* GameplayProvider = Session->ReadProvider<IGameplayProvider>("GameplayProvider"))
 		{
-			TargetActorId = InObjectInfo.Id;
-		}
-	});
+			TArray<TSharedPtr<FDebugObjectInfo>> NewComponentList;	
+			
+			if (DebugTargetActor.Get() == "")
+			{
+				return;
+			}
 
-	// add actor (even if it isn't found in the gameplay provider)
-	if (DebugComponents.Num() > 0 && DebugComponents[0]->ObjectName == DebugTargetActor.Get() && DebugComponents[0]->ObjectId == TargetActorId)
-	{
-		// re-use the version from the old array if it exists, to maintain selection in the list view
-		NewComponentList.Add(DebugComponents[0]);
-	}
-	else
-	{
-		NewComponentList.Add(MakeShared<FDebugObjectInfo>(TargetActorId, DebugTargetActor.Get()));
-	}
+			uint64 TargetActorId = 0;
 
-	if (TargetActorId != 0)
-	{
-		GameplayProvider->EnumerateObjects([this, TargetActorId, GameplayProvider, &NewComponentList](const FObjectInfo& InObjectInfo)
-		{
-			// todo: filter components based on creation/destruction frame, and the current scrubbing time (so dynamically created and destroyed components won't add up)
+			// find the Id of the selected actor (move this to a function or cache it)
+			GameplayProvider->EnumerateObjects([this,&TargetActorId](const FObjectInfo& InObjectInfo)
+			{
+				if (DebugTargetActor.Get() == InObjectInfo.Name)
+				{
+					TargetActorId = InObjectInfo.Id;
+				}
+			});
 
-			const FObjectInfo* ObjectInfo = &InObjectInfo;
-			uint64 ObjectId = ObjectInfo->Id;
-			
-			// find any objects that have TargetActorId as OuterId (or of any outer object recursively)
-			while (ObjectInfo && ObjectInfo->OuterId != 0)
+			// add actor (even if it isn't found in the gameplay provider)
+			if (DebugComponents.Num() > 0 && DebugComponents[0]->ObjectName == DebugTargetActor.Get() && DebugComponents[0]->ObjectId == TargetActorId)
+			{
+				// re-use the version from the old array if it exists, to maintain selection in the list view
+				NewComponentList.Add(DebugComponents[0]);
+			}
+			else
+			{
+				NewComponentList.Add(MakeShared<FDebugObjectInfo>(TargetActorId, DebugTargetActor.Get()));
+			}
+
+			if (TargetActorId != 0)
 			{
-				if(ObjectInfo->OuterId == TargetActorId)
+				GameplayProvider->EnumerateObjects([this, TargetActorId, GameplayProvider, &NewComponentList](const FObjectInfo& InObjectInfo)
 				{
-					// re-use the version from the old array if it exists, to maintain selection in the list view
-					int32 FoundIndex = DebugComponents.FindLastByPredicate([ObjectId](const TSharedPtr<FDebugObjectInfo>& Info) { return Info->ObjectId == ObjectId; });
+					// todo: filter components based on creation/destruction frame, and the current scrubbing time (so dynamically created and destroyed components won't add up)
 
-					if (FoundIndex >= 0 && DebugComponents[FoundIndex]->ObjectName == InObjectInfo.Name)
+					const FObjectInfo* ObjectInfo = &InObjectInfo;
+					uint64 ObjectId = ObjectInfo->Id;
+					
+					// find any objects that have TargetActorId as OuterId (or of any outer object recursively)
+					while (ObjectInfo && ObjectInfo->OuterId != 0)
 					{
-						NewComponentList.Add(DebugComponents[FoundIndex]);
-					}
-					else
-					{
-						NewComponentList.Add(MakeShared<FDebugObjectInfo>(ObjectId,InObjectInfo.Name));
-					}
-					return;
-				}
+						if(ObjectInfo->OuterId == TargetActorId)
+						{
+							// re-use the version from the old array if it exists, to maintain selection in the list view
+							int32 FoundIndex = DebugComponents.FindLastByPredicate([ObjectId](const TSharedPtr<FDebugObjectInfo>& Info) { return Info->ObjectId == ObjectId; });
+
+							if (FoundIndex >= 0 && DebugComponents[FoundIndex]->ObjectName == InObjectInfo.Name)
+							{
+								NewComponentList.Add(DebugComponents[FoundIndex]);
+							}
+							else
+							{
+								NewComponentList.Add(MakeShared<FDebugObjectInfo>(ObjectId,InObjectInfo.Name));
+							}
+							return;
+						}
 
-				ObjectInfo = GameplayProvider->FindObjectInfo(ObjectInfo->OuterId);
+						ObjectInfo = GameplayProvider->FindObjectInfo(ObjectInfo->OuterId);
+					}
+				});
 			}
-		});
-	}
 
-	if (DebugComponents != NewComponentList) // since we reused old SharedPtr for any entries that existed, array equality will tell us if anything has changed
-	{
-		DebugComponents = NewComponentList;
-		ComponentListChangedDelegate.ExecuteIfBound();
+			if (DebugComponents != NewComponentList) // since we reused old SharedPtr for any entries that existed, array equality will tell us if anything has changed
+			{
+				DebugComponents = NewComponentList;
+				ComponentListChangedDelegate.ExecuteIfBound();
+			}
+		}
 	}
 }
 
@@ -369,138 +371,141 @@ void  FRewindDebugger::SetCurrentScrubTime(float Time)
 void FRewindDebugger::UpdateTraceTime()
 {
 	// find the current Insights trace time for the current scrub time.
-	TSharedPtr<const TraceServices::IAnalysisSession> Session = UnrealInsightsModule->GetAnalysisSession();
-	TraceServices::FAnalysisSessionReadScope SessionReadScope(*Session.Get());
-	UWorld* World = GetWorldToVisualize();
-
-	if (const IGameplayProvider* GameplayProvider = Session->ReadProvider<IGameplayProvider>("GameplayProvider"))
+	if (const TraceServices::IAnalysisSession* Session = GetAnalysisSession())
 	{
-		if (const IGameplayProvider::RecordingInfoTimeline* Recording = GameplayProvider->GetRecordingInfo(RecordingIndex))
-		{
-			uint64 EventCount = Recording->GetEventCount();
+		TraceServices::FAnalysisSessionReadScope SessionReadScope(*Session);
+		UWorld* World = GetWorldToVisualize();
 
-			if (EventCount > 0)
+		if (const IGameplayProvider* GameplayProvider = Session->ReadProvider<IGameplayProvider>("GameplayProvider"))
+		{
+			if (const IGameplayProvider::RecordingInfoTimeline* Recording = GameplayProvider->GetRecordingInfo(RecordingIndex))
 			{
-				uint64 EventIndex;
+				uint64 EventCount = Recording->GetEventCount();
 
-				for(EventIndex = 0; EventIndex < EventCount-1; EventIndex++)
+				if (EventCount > 0)
 				{
-					const FRecordingInfoMessage& Event = Recording->GetEvent(EventIndex);
-					// todo: optimize this linear search (cache current index, and pass it in as a hint)
-					//		 and do interpolation between nearest two frames
-					if (Event.ElapsedTime >= CurrentScrubTime)
+					uint64 EventIndex;
+
+					for(EventIndex = 0; EventIndex < EventCount-1; EventIndex++)
 					{
-						break;
+						const FRecordingInfoMessage& Event = Recording->GetEvent(EventIndex);
+						// todo: optimize this linear search (cache current index, and pass it in as a hint)
+						//		 and do interpolation between nearest two frames
+						if (Event.ElapsedTime >= CurrentScrubTime)
+						{
+							break;
+						}
 					}
+					TraceTime.Set(Recording->GetEvent(EventIndex).ProfileTime);
 				}
-				TraceTime.Set(Recording->GetEvent(EventIndex).ProfileTime);
 			}
 		}
 	}
 }
 
-void FRewindDebugger::Tick(float DeltaTime)
+const TraceServices::IAnalysisSession* FRewindDebugger::GetAnalysisSession()
 {
 	if (UnrealInsightsModule == nullptr)
 	{
 		UnrealInsightsModule = &FModuleManager::LoadModuleChecked<IUnrealInsightsModule>("TraceInsights");
 	    if (UnrealInsightsModule == nullptr)
 		{
-			return;
+			return nullptr; 
 		}
 	}
 
-	TSharedPtr<const TraceServices::IAnalysisSession> Session = UnrealInsightsModule->GetAnalysisSession();
-	if (!Session.IsValid())
-	{
-		return;
-	}
-
-	if (bRecording)
-	{
-		// if you select a debug target before you start recording, update component list when it becomes valid
-		RefreshDebugComponents();
-	}
+	return UnrealInsightsModule->GetAnalysisSession().Get();
+}
 
-	if (const IAnimationProvider* AnimationProvider = Session->ReadProvider<IAnimationProvider>("AnimationProvider"))
+void FRewindDebugger::Tick(float DeltaTime)
+{
+	if (const TraceServices::IAnalysisSession* Session = GetAnalysisSession())
 	{
-		TraceServices::FAnalysisSessionReadScope SessionReadScope(*Session.Get());
-		UWorld* World = GetWorldToVisualize();
-
-		if (bPIESimulating)
+		if (bRecording)
 		{
-			if (bRecording)
-			{
-				RecordingDuration.Set(FObjectTrace::GetWorldElapsedTime(World));
-				SetCurrentScrubTime(RecordingDuration.Get());
-				TrackCursorDelegate.ExecuteIfBound(false);
-			}
+			// if you select a debug target before you start recording, update component list when it becomes valid
+			RefreshDebugComponents();
 		}
-		else
+
+		if (const IAnimationProvider* AnimationProvider = Session->ReadProvider<IAnimationProvider>("AnimationProvider"))
 		{
-			if (ControlState == EControlState::Play || ControlState == EControlState::PlayReverse)
-			{
-				float Rate = PlaybackRate * (ControlState == EControlState::Play ? 1 : -1);
-				SetCurrentScrubTime(FMath::Clamp(CurrentScrubTime + Rate * DeltaTime, 0.0f, RecordingDuration.Get()));
-				TrackCursorDelegate.ExecuteIfBound(Rate<0);
+			TraceServices::FAnalysisSessionReadScope SessionReadScope(*Session);
+			UWorld* World = GetWorldToVisualize();
 
-				if (CurrentScrubTime == 0 || CurrentScrubTime == RecordingDuration.Get())
+			if (bPIESimulating)
+			{
+				if (bRecording)
 				{
-					// pause at end.
-					ControlState = EControlState::Pause;
+					RecordingDuration.Set(FObjectTrace::GetWorldElapsedTime(World));
+					SetCurrentScrubTime(RecordingDuration.Get());
+					TrackCursorDelegate.ExecuteIfBound(false);
 				}
 			}
-
-			double CurrentTraceTime = TraceTime.Get();
-
-			// update pose on all SkeletalMeshComponents
-			ULevel* CurLevel = World->GetCurrentLevel();
-			for( AActor* LevelActor : CurLevel->Actors )
+			else
 			{
-				if (LevelActor)
+				if (ControlState == EControlState::Play || ControlState == EControlState::PlayReverse)
 				{
-					TInlineComponentArray<USkeletalMeshComponent*> SkeletalMeshComponents;
-					LevelActor->GetComponents(SkeletalMeshComponents);
+					float Rate = PlaybackRate * (ControlState == EControlState::Play ? 1 : -1);
+					SetCurrentScrubTime(FMath::Clamp(CurrentScrubTime + Rate * DeltaTime, 0.0f, RecordingDuration.Get()));
+					TrackCursorDelegate.ExecuteIfBound(Rate<0);
 
-					for (USkeletalMeshComponent* MeshComponent : SkeletalMeshComponents)
+					if (CurrentScrubTime == 0 || CurrentScrubTime == RecordingDuration.Get())
 					{
-						int64 ObjectId = FObjectTrace::GetObjectId(MeshComponent);
+						// pause at end.
+						ControlState = EControlState::Pause;
+					}
+				}
+
+				double CurrentTraceTime = TraceTime.Get();
 
-						AnimationProvider->ReadSkeletalMeshPoseTimeline(ObjectId, [this, ObjectId, CurrentTraceTime, MeshComponent, AnimationProvider](const IAnimationProvider::SkeletalMeshPoseTimeline& TimelineData, bool bHasCurves)
+				// update pose on all SkeletalMeshComponents
+				ULevel* CurLevel = World->GetCurrentLevel();
+				for( AActor* LevelActor : CurLevel->Actors )
+				{
+					if (LevelActor)
+					{
+						TInlineComponentArray<USkeletalMeshComponent*> SkeletalMeshComponents;
+						LevelActor->GetComponents(SkeletalMeshComponents);
+
+						for (USkeletalMeshComponent* MeshComponent : SkeletalMeshComponents)
 						{
-							double PrecedingPoseTime;
-							double FollowingPoseTime;
-							const FSkeletalMeshPoseMessage* PrecedingPose;
-							const FSkeletalMeshPoseMessage* FollowingPose;
-
-							TimelineData.FindNearestEvents(CurrentTraceTime, PrecedingPose, PrecedingPoseTime, FollowingPose, FollowingPoseTime);
-
-							if (FollowingPose || PrecedingPose)
-							{	
-								// Ideally we should iterpolate between the two poses here if both are valid (we would need to transform the component space transforms to local space though)
-								const FSkeletalMeshPoseMessage& PoseMessage = FollowingPose ? *FollowingPose : *PrecedingPose;
-
-								FTransform ComponentWorldTransform;
-								const FSkeletalMeshInfo* SkeletalMeshInfo = AnimationProvider->FindSkeletalMeshInfo(PoseMessage.MeshId);
-								AnimationProvider->GetSkeletalMeshComponentSpacePose(PoseMessage, *SkeletalMeshInfo, ComponentWorldTransform, MeshComponent->GetEditableComponentSpaceTransforms());
-
-								if (MeshComponentsToReset.Find(ObjectId) == nullptr)
-								{
-									FMeshComponentResetData ResetData;
-									ResetData.Component = MeshComponent;
-									ResetData.RelativeTransform = MeshComponent->GetRelativeTransform();
-									MeshComponentsToReset.Add(ObjectId, ResetData);
+							int64 ObjectId = FObjectTrace::GetObjectId(MeshComponent);
+
+							AnimationProvider->ReadSkeletalMeshPoseTimeline(ObjectId, [this, ObjectId, CurrentTraceTime, MeshComponent, AnimationProvider](const IAnimationProvider::SkeletalMeshPoseTimeline& TimelineData, bool bHasCurves)
+							{
+								double PrecedingPoseTime;
+								double FollowingPoseTime;
+								const FSkeletalMeshPoseMessage* PrecedingPose;
+								const FSkeletalMeshPoseMessage* FollowingPose;
+
+								TimelineData.FindNearestEvents(CurrentTraceTime, PrecedingPose, PrecedingPoseTime, FollowingPose, FollowingPoseTime);
+
+								if (FollowingPose || PrecedingPose)
+								{	
+									// Ideally we should iterpolate between the two poses here if both are valid (we would need to transform the component space transforms to local space though)
+									const FSkeletalMeshPoseMessage& PoseMessage = FollowingPose ? *FollowingPose : *PrecedingPose;
+
+									FTransform ComponentWorldTransform;
+									const FSkeletalMeshInfo* SkeletalMeshInfo = AnimationProvider->FindSkeletalMeshInfo(PoseMessage.MeshId);
+									AnimationProvider->GetSkeletalMeshComponentSpacePose(PoseMessage, *SkeletalMeshInfo, ComponentWorldTransform, MeshComponent->GetEditableComponentSpaceTransforms());
+
+									if (MeshComponentsToReset.Find(ObjectId) == nullptr)
+									{
+										FMeshComponentResetData ResetData;
+										ResetData.Component = MeshComponent;
+										ResetData.RelativeTransform = MeshComponent->GetRelativeTransform();
+										MeshComponentsToReset.Add(ObjectId, ResetData);
+									}
+
+									MeshComponent->SetWorldTransform(ComponentWorldTransform);
+									MeshComponent->SetForcedLOD(PoseMessage.LodIndex + 1);
+									MeshComponent->ApplyEditedComponentSpaceTransforms();
 								}
-
-								MeshComponent->SetWorldTransform(ComponentWorldTransform);
-								MeshComponent->SetForcedLOD(PoseMessage.LodIndex + 1);
-								MeshComponent->ApplyEditedComponentSpaceTransforms();
-							}
-						});
+							});
+						}
 					}
 				}
 			}
-
 		}
 	}
 }
diff --git a/Engine/Plugins/Animation/GameplayInsights/Source/RewindDebugger/Private/RewindDebugger.h b/Engine/Plugins/Animation/GameplayInsights/Source/RewindDebugger/Private/RewindDebugger.h
--- a/Engine/Plugins/Animation/GameplayInsights/Source/RewindDebugger/Private/RewindDebugger.h
+++ b/Engine/Plugins/Animation/GameplayInsights/Source/RewindDebugger/Private/RewindDebugger.h
@@ -7,6 +7,11 @@
 #include "RewindDebuggerModule.h"
 
 
+namespace TraceServices
+{
+	class IAnalysisSession;
+}
+
 // Singleton class that handles the logic for the Rewind Debugger
 // handles:
 //  Playback/Scrubbing state
@@ -86,6 +91,7 @@ private:
 	void OnPIEResumed(bool bSimulating);
 	void OnPIEStopped(bool bSimulating);
 
+	const TraceServices::IAnalysisSession* GetAnalysisSession();
 	UWorld* GetWorldToVisualize() const;
 
 	void SetCurrentScrubTime(float Time);
