Fix for agent head meshes not replaying properly in Frosty

diff --git a/Engine/Plugins/Animation/GameplayInsights/Source/RewindDebugger/Private/RewindDebugger.cpp b/Engine/Plugins/Animation/GameplayInsights/Source/RewindDebugger/Private/RewindDebugger.cpp
--- a/Engine/Plugins/Animation/GameplayInsights/Source/RewindDebugger/Private/RewindDebugger.cpp
+++ b/Engine/Plugins/Animation/GameplayInsights/Source/RewindDebugger/Private/RewindDebugger.cpp
@@ -499,9 +499,12 @@ void FRewindDebugger::Tick(float DeltaTime)
 
 									MeshComponent->SetWorldTransform(ComponentWorldTransform);
 									MeshComponent->SetForcedLOD(PoseMessage.LodIndex + 1);
-									MeshComponent->ApplyEditedComponentSpaceTransforms();
 								}
 							});
+
+							// calling this here, even on meshes which have no recorded data makes meshes with a MasterPoseComponent attach properly during replay
+							// longer term solution: we should be recording those meshes
+							MeshComponent->ApplyEditedComponentSpaceTransforms();
 						}
 					}
 				}
