Replaces references to the old SCS editor with the new subobject editor in gameplay insights

#rnx

diff --git a/Engine/Plugins/Animation/GameplayInsights/Source/GameplayInsights/GameplayInsights.Build.cs b/Engine/Plugins/Animation/GameplayInsights/Source/GameplayInsights/GameplayInsights.Build.cs
--- a/Engine/Plugins/Animation/GameplayInsights/Source/GameplayInsights/GameplayInsights.Build.cs
+++ b/Engine/Plugins/Animation/GameplayInsights/Source/GameplayInsights/GameplayInsights.Build.cs
@@ -45,6 +45,7 @@ namespace UnrealBuildTool.Rules
 					"EditorWidgets",
 					"ToolMenus",
 					"Kismet",
+					"SubobjectEditor",
 				});
 			}
 		}
diff --git a/Engine/Plugins/Animation/GameplayInsights/Source/GameplayInsights/Private/GameplayInsightsModule.cpp b/Engine/Plugins/Animation/GameplayInsights/Source/GameplayInsights/Private/GameplayInsightsModule.cpp
--- a/Engine/Plugins/Animation/GameplayInsights/Source/GameplayInsights/Private/GameplayInsightsModule.cpp
+++ b/Engine/Plugins/Animation/GameplayInsights/Source/GameplayInsights/Private/GameplayInsightsModule.cpp
@@ -27,7 +27,8 @@
 #include "Engine/Selection.h"
 #include "SSCSEditorMenuContext.h"
 #include "GameplayInsightsStyle.h"
-#include "SSCSEditor.h"
+#include "SSubobjectInstanceEditor.h"
+
 #endif
 
 #if WITH_ENGINE
@@ -321,21 +322,21 @@ void FGameplayInsightsModule::RegisterMenus()
 			EUserInterfaceActionType::ToggleButton);
 	}
 	{
-		UToolMenu* Menu = UToolMenus::Get()->ExtendMenu("Kismet.SCSEditorContextMenu");
+		UToolMenu* Menu = UToolMenus::Get()->ExtendMenu("Kismet.SubobjectEditorContextMenu");
 
 		FToolMenuSection& Section = Menu->AddSection("GameplayInsights", LOCTEXT("GameplayInsights", "Gameplay Insights"));
 
-		auto GetCheckState = [](const TSharedPtr<SSCSEditor>& InSCSEditor)
+		auto GetCheckState = [](const TSharedPtr<SSubobjectEditor>& InSubobjectEditor)
 		{
-			if (InSCSEditor->GetNumSelectedNodes() > 0 && FObjectPropertyTrace::IsEnabled())
+			if (InSubobjectEditor->GetNumSelectedNodes() > 0 && FObjectPropertyTrace::IsEnabled())
 			{
-				TArray<FSCSEditorTreeNodePtrType> SelectedNodes = InSCSEditor->GetSelectedNodes();
+				TArray<FSubobjectEditorTreeNodePtrType> SelectedNodes = InSubobjectEditor->GetSelectedNodes();
 				int32 TotalObjectCount = SelectedNodes.Num();
 				int32 RegisteredObjectCount = 0;
 
-				for (FSCSEditorTreeNodePtrType SCSNode : SelectedNodes)
+				for (FSubobjectEditorTreeNodePtrType SubobjectNode : SelectedNodes)
 				{
-					const UObject* SelectedComponent = SCSNode->GetObject<UObject>();
+					const UObject* SelectedComponent = SubobjectNode->GetObject();
 					if (FObjectPropertyTrace::IsObjectRegistered(SelectedComponent))
 					{
 						RegisteredObjectCount++;
@@ -365,24 +366,22 @@ void FGameplayInsightsModule::RegisterMenus()
 			if (FObjectPropertyTrace::IsEnabled())
 			{
 				USSCSEditorMenuContext* ContextObject = InContext.FindContext<USSCSEditorMenuContext>();
-				if (ContextObject && ContextObject->SCSEditor.IsValid() && ContextObject->SCSEditor.Pin()->GetEditorMode() == EComponentEditorMode::ActorInstance)
-				{
-					TSharedPtr<SSCSEditor> SCSEditor = ContextObject->SCSEditor.Pin();
-					if(SCSEditor.IsValid())
-					{
-						ECheckBoxState CheckState = GetCheckState(SCSEditor);
+				TSharedPtr<SSubobjectEditor> SubobjectEditor = ContextObject ? ContextObject->SubobjectEditor.Pin() : nullptr;
+
+				if (SubobjectEditor.IsValid() && StaticCastSharedPtr<SSubobjectInstanceEditor>(SubobjectEditor))
+				{					
+					ECheckBoxState CheckState = GetCheckState(SubobjectEditor);
 
-						for(FSCSEditorTreeNodePtrType SCSNode : SCSEditor->GetSelectedNodes())
+					for(FSubobjectEditorTreeNodePtrType Node : SubobjectEditor->GetSelectedNodes())
+					{
+						const UObject* SelectedComponent = Node->GetObject();
+						if(CheckState == ECheckBoxState::Unchecked)
+						{
+							FObjectPropertyTrace::RegisterObject(SelectedComponent);
+						}
+						else
 						{
-							const UObject* SelectedComponent = SCSNode->GetObject<UObject>();
-							if(CheckState == ECheckBoxState::Unchecked)
-							{
-								FObjectPropertyTrace::RegisterObject(SelectedComponent);
-							}
-							else
-							{
-								FObjectPropertyTrace::UnregisterObject(SelectedComponent);
-							}
+							FObjectPropertyTrace::UnregisterObject(SelectedComponent);
 						}
 					}
 				}
@@ -393,12 +392,12 @@ void FGameplayInsightsModule::RegisterMenus()
 			if (FObjectPropertyTrace::IsEnabled())
 			{
 				USSCSEditorMenuContext* ContextObject = InContext.FindContext<USSCSEditorMenuContext>();
-				if (ContextObject && ContextObject->SCSEditor.IsValid() && ContextObject->SCSEditor.Pin()->GetEditorMode() == EComponentEditorMode::ActorInstance)
+				if (ContextObject && ContextObject->SubobjectEditor.IsValid() && StaticCastSharedPtr<SSubobjectInstanceEditor>(ContextObject->SubobjectEditor.Pin()))
 				{
-					TSharedPtr<SSCSEditor> SCSEditor = ContextObject->SCSEditor.Pin();
-					if (SCSEditor.IsValid())
+					TSharedPtr<SSubobjectEditor> SubobjectEditor = ContextObject->SubobjectEditor.Pin();
+					if (SubobjectEditor.IsValid())
 					{
-						return SCSEditor->GetNumSelectedNodes() > 0;
+						return SubobjectEditor->GetNumSelectedNodes() > 0;
 					}
 				}
 			}
@@ -408,12 +407,13 @@ void FGameplayInsightsModule::RegisterMenus()
 		Action.GetActionCheckState = FToolMenuGetActionCheckState::CreateLambda([&GetCheckState](const FToolMenuContext& InContext)
 		{
 			USSCSEditorMenuContext* ContextObject = InContext.FindContext<USSCSEditorMenuContext>();
-			if (ContextObject && ContextObject->SCSEditor.IsValid() && ContextObject->SCSEditor.Pin()->GetEditorMode() == EComponentEditorMode::ActorInstance)
+			if (ContextObject && StaticCastSharedPtr<SSubobjectInstanceEditor>(ContextObject->SubobjectEditor.Pin()))
 			{
-				TSharedPtr<SSCSEditor> SCSEditor = ContextObject->SCSEditor.Pin();
-				if (SCSEditor.IsValid())
+				TSharedPtr<SSubobjectEditor> SubobjectEditor = ContextObject->SubobjectEditor.Pin();
+
+				if (SubobjectEditor.IsValid())
 				{
-					return GetCheckState(SCSEditor);
+					return GetCheckState(SubobjectEditor);
 				}
 			}
 
@@ -424,7 +424,7 @@ void FGameplayInsightsModule::RegisterMenus()
 			if (FObjectPropertyTrace::IsEnabled())
 			{
 				USSCSEditorMenuContext* ContextObject = InContext.FindContext<USSCSEditorMenuContext>();
-				if (ContextObject && ContextObject->SCSEditor.IsValid() && ContextObject->SCSEditor.Pin()->GetEditorMode() == EComponentEditorMode::ActorInstance)
+				if (ContextObject && ContextObject->SubobjectEditor.IsValid() && StaticCastSharedPtr<SSubobjectInstanceEditor>(ContextObject->SubobjectEditor.Pin()))
 				{
 					return true;
 				}
