Improved component view for Rewind debugger

diff --git a/Engine/Plugins/Animation/GameplayInsights/Source/RewindDebugger/Private/RewindDebugger.cpp b/Engine/Plugins/Animation/GameplayInsights/Source/RewindDebugger/Private/RewindDebugger.cpp
--- a/Engine/Plugins/Animation/GameplayInsights/Source/RewindDebugger/Private/RewindDebugger.cpp
+++ b/Engine/Plugins/Animation/GameplayInsights/Source/RewindDebugger/Private/RewindDebugger.cpp
@@ -107,6 +107,42 @@ void FRewindDebugger::OnPIEStopped(bool bSimulating)
 	SetCurrentScrubTime(0);
 }
 
+bool FRewindDebugger::UpdateComponentList(uint64 ParentId, TArray<TSharedPtr<FDebugObjectInfo>>& ComponentList)
+{
+	TSharedPtr<const TraceServices::IAnalysisSession> Session = UnrealInsightsModule->GetAnalysisSession();
+	TraceServices::FAnalysisSessionReadScope SessionReadScope(*Session.Get());
+	UWorld* World = GetWorldToVisualize();
+
+	const IGameplayProvider* GameplayProvider = Session->ReadProvider<IGameplayProvider>("GameplayProvider");
+
+	bool bChanged = false;
+
+	GameplayProvider->EnumerateObjects([this, &bChanged, ParentId, GameplayProvider, &ComponentList](const FObjectInfo& InObjectInfo)
+	{
+		// todo: filter components based on creation/destruction frame, and the current scrubbing time (so dynamically created and destroyed components won't add up)
+
+		const FObjectInfo* ObjectInfo = &InObjectInfo;
+		if(ObjectInfo->OuterId == ParentId)
+		{
+			int32 FoundIndex = ComponentList.FindLastByPredicate([ObjectInfo](const TSharedPtr<FDebugObjectInfo>& Info) { return Info->ObjectName == ObjectInfo->Name; });
+
+			if (FoundIndex >= 0 && ComponentList[FoundIndex]->ObjectName == InObjectInfo.Name)
+			{
+				ComponentList[FoundIndex]->ObjectId = ObjectInfo->Id; // there is an issue with skeletal mesh components changing ids
+				bChanged = bChanged || UpdateComponentList(ObjectInfo->Id, ComponentList[FoundIndex]->Children);
+			}
+			else
+			{
+				bChanged = true;
+				ComponentList.Add(MakeShared<FDebugObjectInfo>(ObjectInfo->Id,InObjectInfo.Name));
+				UpdateComponentList(ObjectInfo->Id, ComponentList.Last()->Children);
+			}
+		}
+	});
+
+	return bChanged;
+}
+
 void FRewindDebugger::RefreshDebugComponents()
 {
 	if (const TraceServices::IAnalysisSession* Session = GetAnalysisSession())
@@ -116,8 +152,6 @@ void FRewindDebugger::RefreshDebugComponents()
 
 		if (const IGameplayProvider* GameplayProvider = Session->ReadProvider<IGameplayProvider>("GameplayProvider"))
 		{
-			TArray<TSharedPtr<FDebugObjectInfo>> NewComponentList;	
-			
 			if (DebugTargetActor.Get() == "")
 			{
 				return;
@@ -134,53 +168,30 @@ void FRewindDebugger::RefreshDebugComponents()
 				}
 			});
 
+			bool bChanged = false;
+
 			// add actor (even if it isn't found in the gameplay provider)
-			if (DebugComponents.Num() > 0 && DebugComponents[0]->ObjectName == DebugTargetActor.Get() && DebugComponents[0]->ObjectId == TargetActorId)
+			if (DebugComponents.Num() == 0)
 			{
-				// re-use the version from the old array if it exists, to maintain selection in the list view
-				NewComponentList.Add(DebugComponents[0]);
+				bChanged = true;
+				DebugComponents.Add(MakeShared<FDebugObjectInfo>(TargetActorId, DebugTargetActor.Get()));
 			}
 			else
 			{
-				NewComponentList.Add(MakeShared<FDebugObjectInfo>(TargetActorId, DebugTargetActor.Get()));
+				if (DebugComponents[0]->ObjectName != DebugTargetActor.Get() || DebugComponents[0]->ObjectId != TargetActorId)
+				{
+					bChanged = true;
+					DebugComponents[0] = MakeShared<FDebugObjectInfo>(TargetActorId, DebugTargetActor.Get());
+				}
 			}
 
-			if (TargetActorId != 0)
+			if (TargetActorId != 0 && DebugComponents.Num() > 0)
 			{
-				GameplayProvider->EnumerateObjects([this, TargetActorId, GameplayProvider, &NewComponentList](const FObjectInfo& InObjectInfo)
-				{
-					// todo: filter components based on creation/destruction frame, and the current scrubbing time (so dynamically created and destroyed components won't add up)
-
-					const FObjectInfo* ObjectInfo = &InObjectInfo;
-					uint64 ObjectId = ObjectInfo->Id;
-					
-					// find any objects that have TargetActorId as OuterId (or of any outer object recursively)
-					while (ObjectInfo && ObjectInfo->OuterId != 0)
-					{
-						if(ObjectInfo->OuterId == TargetActorId)
-						{
-							// re-use the version from the old array if it exists, to maintain selection in the list view
-							int32 FoundIndex = DebugComponents.FindLastByPredicate([ObjectId](const TSharedPtr<FDebugObjectInfo>& Info) { return Info->ObjectId == ObjectId; });
-
-							if (FoundIndex >= 0 && DebugComponents[FoundIndex]->ObjectName == InObjectInfo.Name)
-							{
-								NewComponentList.Add(DebugComponents[FoundIndex]);
-							}
-							else
-							{
-								NewComponentList.Add(MakeShared<FDebugObjectInfo>(ObjectId,InObjectInfo.Name));
-							}
-							return;
-						}
-
-						ObjectInfo = GameplayProvider->FindObjectInfo(ObjectInfo->OuterId);
-					}
-				});
+				bChanged = bChanged || UpdateComponentList(TargetActorId, DebugComponents[0]->Children);
 			}
 
-			if (DebugComponents != NewComponentList) // since we reused old SharedPtr for any entries that existed, array equality will tell us if anything has changed
+			if (bChanged)
 			{
-				DebugComponents = NewComponentList;
 				ComponentListChangedDelegate.ExecuteIfBound();
 			}
 		}
@@ -497,6 +508,7 @@ void FRewindDebugger::Tick(float DeltaTime)
 										MeshComponentsToReset.Add(ObjectId, ResetData);
 									}
 
+									// todo: we probably need to take into account tick order requirements for attached objects here
 									MeshComponent->SetWorldTransform(ComponentWorldTransform);
 									MeshComponent->SetForcedLOD(PoseMessage.LodIndex + 1);
 								}
diff --git a/Engine/Plugins/Animation/GameplayInsights/Source/RewindDebugger/Private/RewindDebugger.h b/Engine/Plugins/Animation/GameplayInsights/Source/RewindDebugger/Private/RewindDebugger.h
--- a/Engine/Plugins/Animation/GameplayInsights/Source/RewindDebugger/Private/RewindDebugger.h
+++ b/Engine/Plugins/Animation/GameplayInsights/Source/RewindDebugger/Private/RewindDebugger.h
@@ -94,6 +94,7 @@ private:
 	const TraceServices::IAnalysisSession* GetAnalysisSession();
 	UWorld* GetWorldToVisualize() const;
 
+	bool UpdateComponentList(uint64 ParentId, TArray<TSharedPtr<FDebugObjectInfo>>& NewComponentList);
 	void SetCurrentScrubTime(float Time);
 	void UpdateTraceTime();
 
diff --git a/Engine/Plugins/Animation/GameplayInsights/Source/RewindDebugger/Private/RewindDebuggerModule.h b/Engine/Plugins/Animation/GameplayInsights/Source/RewindDebugger/Private/RewindDebuggerModule.h
--- a/Engine/Plugins/Animation/GameplayInsights/Source/RewindDebugger/Private/RewindDebuggerModule.h
+++ b/Engine/Plugins/Animation/GameplayInsights/Source/RewindDebugger/Private/RewindDebuggerModule.h
@@ -10,12 +10,15 @@ class SRewindDebugger;
 
 struct FDebugObjectInfo
 {
-	FDebugObjectInfo(uint64 Id, const FString& Name): ObjectId(Id), ObjectName(Name)
+	FDebugObjectInfo(uint64 Id, const FString& Name): ObjectId(Id), ObjectName(Name), bExpanded(true)
 	{
 	}
 
 	uint64 ObjectId;
 	FString ObjectName;
+	bool bExpanded;
+
+	TArray<TSharedPtr<FDebugObjectInfo>> Children;
 };
 
 class FRewindDebuggerModule : public IModuleInterface
diff --git a/Engine/Plugins/Animation/GameplayInsights/Source/RewindDebugger/Private/SRewindDebugger.cpp b/Engine/Plugins/Animation/GameplayInsights/Source/RewindDebugger/Private/SRewindDebugger.cpp
--- a/Engine/Plugins/Animation/GameplayInsights/Source/RewindDebugger/Private/SRewindDebugger.cpp
+++ b/Engine/Plugins/Animation/GameplayInsights/Source/RewindDebugger/Private/SRewindDebugger.cpp
@@ -20,15 +20,10 @@
 #include "SceneOutlinerModule.h"
 #include "Selection.h"
 #include "Styling/SlateIconFinder.h"
-#include "Widgets/Colors/SColorBlock.h"
 #include "Widgets/Images/SImage.h"
 #include "Widgets/Input/SComboButton.h"
 #include "Widgets/Input/SButton.h"
 #include "Widgets/Input/SSearchBox.h"
-#include "Widgets/Input/STextComboBox.h"
-#include "Widgets/Input/STextComboBox.h"
-#include "Widgets/SOverlay.h"
-#include "Widgets/Views/SListView.h"
 
 
 #define LOCTEXT_NAMESPACE "SAnimationInsights"
@@ -186,7 +181,6 @@ TSharedRef<SWidget> SRewindDebugger::MakeSelectActorMenu()
 	return MenuBuilder.MakeWidget();
 }
 
-
 void SRewindDebugger::Construct(const FArguments& InArgs, TSharedRef<FUICommandList> CommandList, const TSharedRef<SDockTab>& ConstructUnderMajorTab, const TSharedPtr<SWindow>& ConstructUnderWindow)
 {
 	OnScrubPositionChanged = InArgs._OnScrubPositionChanged;
@@ -237,13 +231,9 @@ void SRewindDebugger::Construct(const FArguments& InArgs, TSharedRef<FUICommandL
 	TimeSliderController = MakeShared<FTimeSliderController>(InitArgs);
     TSharedRef<STimeSlider> TimeSlider = SNew(STimeSlider, TimeSliderController.ToSharedRef());
 
-	ComponentListView =	SNew(SListView<TSharedPtr<FDebugObjectInfo>>)
-									.ItemHeight(16.0f)
-									.ListItemsSource(DebugComponents)
-									.OnGenerateRow(this, &SRewindDebugger::ComponentListViewGenerateRow)
-									.SelectionMode(ESelectionMode::Single)
-									.OnSelectionChanged(this, &SRewindDebugger::OnComponentSelectionChanged);
-
+	ComponentTreeView =	SNew(SRewindDebuggerComponentTree)
+				.DebugComponents(InArgs._DebugComponents)
+				.OnSelectionChanged(this, &SRewindDebugger::OnComponentSelectionChanged);
 
 	// everything related to creating SAnimGraphSchematicView should be moved to a separate file
 	IUnrealInsightsModule &UnrealInsightsModule = FModuleManager::LoadModuleChecked<IUnrealInsightsModule>("TraceInsights");
@@ -253,8 +243,6 @@ void SRewindDebugger::Construct(const FArguments& InArgs, TSharedRef<FUICommandL
 	AnimGraphView = GameplayInsightsModule.CreateAnimGraphSchematicView(0, TraceTime.Get(), *Session);
 	TraceTime.OnPropertyChanged = TraceTime.OnPropertyChanged.CreateSP(AnimGraphView.Get(), &IAnimGraphSchematicView::SetTimeMarker);
 
-	FSlateIcon ActorIcon = FSlateIconFinder::FindIconForClass(AActor::StaticClass());
-
 	ChildSlot
 	[
 		SNew(SSplitter)
@@ -274,11 +262,47 @@ void SRewindDebugger::Construct(const FArguments& InArgs, TSharedRef<FUICommandL
 					.OnGetMenuContent(this, &SRewindDebugger::MakeSelectActorMenu)
 					.ButtonContent()
 					[
-						SNew(STextBlock)
-						.Text_Lambda([this](){
-							const FString& TargetName = DebugTargetActor.Get();
-							return (TargetName.IsEmpty()) ? LOCTEXT("Select Actor", "Debug Target Actor") : FText::FromString(TargetName);
-						 } )
+						SNew(SHorizontalBox)
+						+SHorizontalBox::Slot().AutoWidth().Padding(3)
+						[
+							SNew(SImage)
+							.Image_Lambda([this]
+								{
+									FSlateIcon ActorIcon = FSlateIconFinder::FindIconForClass(AActor::StaticClass());
+									if (DebugComponents != nullptr && DebugComponents->Num()>0)
+									{
+										if (UObject* Object = FObjectTrace::GetObjectFromId((*DebugComponents)[0]->ObjectId))
+										{
+											ActorIcon = FSlateIconFinder::FindIconForClass(Object->GetClass());
+										}
+									}
+
+									return ActorIcon.GetIcon();
+								}
+							)
+						]
+						+SHorizontalBox::Slot().Padding(3)
+						[
+							SNew(STextBlock)
+							.Text_Lambda([this](){
+								if (DebugComponents == nullptr || DebugComponents->Num()==0)
+								{
+									return LOCTEXT("Select Actor", "Debug Target Actor");
+								}
+
+								FString ReadableName = (*DebugComponents)[0]->ObjectName;
+
+								if (UObject* Object = FObjectTrace::GetObjectFromId((*DebugComponents)[0]->ObjectId))
+								{
+									if (AActor* Actor = Cast<AActor>(Object))
+									{
+										ReadableName = Actor->GetActorLabel();
+									}
+								}
+
+								return FText::FromString(ReadableName);
+							} )
+						]
 					]
 				]
 				+SHorizontalBox::Slot().AutoWidth().HAlign(HAlign_Right)
@@ -294,7 +318,7 @@ void SRewindDebugger::Construct(const FArguments& InArgs, TSharedRef<FUICommandL
 			]
 			+SVerticalBox::Slot().VAlign(VAlign_Top) .FillHeight(1.0f)
 			[
-				ComponentListView.ToSharedRef()
+				ComponentTreeView.ToSharedRef()
 			]
 		]
 		+SSplitter::Slot() 
@@ -313,19 +337,9 @@ void SRewindDebugger::Construct(const FArguments& InArgs, TSharedRef<FUICommandL
 	];
 }
 
-TSharedRef<ITableRow> SRewindDebugger::ComponentListViewGenerateRow(TSharedPtr<FDebugObjectInfo> InItem, const TSharedRef<STableViewBase>& OwnerTable)
-{
-	return 
-		SNew(STableRow<TSharedPtr<FDebugObjectInfo>>, OwnerTable)
-		[
-			SNew(STextBlock)
-			.Text(FText::FromString(InItem->ObjectName))
-		];
-}
-
 void SRewindDebugger::RefreshDebugComponents()
 {
-	ComponentListView->RebuildList();
+	ComponentTreeView->Refresh();
 }
 
 void SRewindDebugger::OnComponentSelectionChanged(TSharedPtr<FDebugObjectInfo> SelectedItem, ESelectInfo::Type SelectInfo)
diff --git a/Engine/Plugins/Animation/GameplayInsights/Source/RewindDebugger/Private/SRewindDebugger.h b/Engine/Plugins/Animation/GameplayInsights/Source/RewindDebugger/Private/SRewindDebugger.h
--- a/Engine/Plugins/Animation/GameplayInsights/Source/RewindDebugger/Private/SRewindDebugger.h
+++ b/Engine/Plugins/Animation/GameplayInsights/Source/RewindDebugger/Private/SRewindDebugger.h
@@ -8,10 +8,10 @@
 #include "IAnimGraphSchematicView.h"
 #include "RewindDebuggerModule.h"
 #include "RewindDebuggerTimeSliderController.h"
+#include "SRewindDebuggerComponentTree.h"
 #include "Widgets/DeclarativeSyntaxSupport.h"
 #include "Widgets/Input/SSearchBox.h"
 #include "Widgets/SCompoundWidget.h"
-#include "Widgets/Views/SListView.h"
 
 
 class SRewindDebugger : public SCompoundWidget
@@ -68,11 +68,10 @@ private:
 	TBindableProperty<float> RecordingDuration;
 	TBindableProperty<float> DebugTargetAnimInstanceId;
 
-	// debug components list
+	// component tree view
 	TArray<TSharedPtr<FDebugObjectInfo>>* DebugComponents;
     void OnComponentSelectionChanged(TSharedPtr<FDebugObjectInfo> SelectedItem, ESelectInfo::Type SelectInfo);
-	TSharedPtr<SListView<TSharedPtr<FDebugObjectInfo>>> ComponentListView;
-	TSharedRef<ITableRow> ComponentListViewGenerateRow(TSharedPtr<FDebugObjectInfo> InItem, const TSharedRef<STableViewBase>& OwnerTable);
+	TSharedPtr<SRewindDebuggerComponentTree> ComponentTreeView;
 
 	// anim graph view 
 	TSharedPtr<IAnimGraphSchematicView> AnimGraphView;
diff --git a/Engine/Plugins/Animation/GameplayInsights/Source/RewindDebugger/Private/SRewindDebuggerComponentTree.cpp b/Engine/Plugins/Animation/GameplayInsights/Source/RewindDebugger/Private/SRewindDebuggerComponentTree.cpp
new file mode 100644
--- /dev/null
+++ b/Engine/Plugins/Animation/GameplayInsights/Source/RewindDebugger/Private/SRewindDebuggerComponentTree.cpp
@@ -0,0 +1,114 @@
+// Copyright Epic Games, Inc. All Rights Reserved.
+
+#include "SRewindDebuggerComponentTree.h"
+#include "Async/Future.h"
+#include "Editor.h"
+#include "Editor/EditorEngine.h"
+#include "Framework/Application/SlateApplication.h"
+#include "ObjectTrace.h"
+#include "Styling/SlateIconFinder.h"
+#include "Widgets/Images/SImage.h"
+#include "Widgets/SOverlay.h"
+#include "Widgets/Views/SListView.h"
+
+
+#define LOCTEXT_NAMESPACE "SAnimationInsights"
+
+SRewindDebuggerComponentTree::SRewindDebuggerComponentTree() 
+	: SCompoundWidget()
+	, DebugComponents(nullptr)
+{ 
+}
+
+SRewindDebuggerComponentTree::~SRewindDebuggerComponentTree() 
+{
+
+}
+
+TSharedRef<ITableRow> ComponentTreeViewGenerateRow(TSharedPtr<FDebugObjectInfo> InItem, const TSharedRef<STableViewBase>& OwnerTable)
+{
+	FString ReadableName = InItem->ObjectName;
+
+	FSlateIcon ObjectIcon = FSlateIconFinder::FindIconForClass(UObject::StaticClass());
+
+	if (UObject* Object = FObjectTrace::GetObjectFromId(InItem->ObjectId))
+	{
+		if (AActor* Actor = Cast<AActor>(Object))
+		{
+			ReadableName = Actor->GetActorLabel();
+		}
+		else if (UActorComponent* Component = Cast<UActorComponent>(Object))
+		{
+			ReadableName = Component->GetName();
+		}
+
+		ObjectIcon = FSlateIconFinder::FindIconForClass(Object->GetClass());
+	}
+
+	return 
+		SNew(STableRow<TSharedPtr<FDebugObjectInfo>>, OwnerTable)
+		[
+			SNew(SHorizontalBox)
+			+SHorizontalBox::Slot().AutoWidth().Padding(3)
+			[
+				SNew(SImage)
+				.Image(ObjectIcon.GetIcon())
+			]
+			+SHorizontalBox::Slot().Padding(3)
+			[
+				SNew(STextBlock)
+				.Text(FText::FromString(ReadableName))
+			]
+		];
+}
+
+void ComponentTreeViewGetChildren(TSharedPtr<FDebugObjectInfo> InItem, TArray<TSharedPtr<FDebugObjectInfo>>& OutChildren)
+{
+	OutChildren.Append(InItem->Children);
+}
+
+void ComponentTreeViewExpansionChanged(TSharedPtr<FDebugObjectInfo> InItem, bool bShouldBeExpanded)
+{
+	InItem->bExpanded = bShouldBeExpanded;
+}
+
+void SRewindDebuggerComponentTree::Construct(const FArguments& InArgs)
+{
+	DebugComponents = InArgs._DebugComponents;
+
+	ComponentTreeView = SNew(STreeView<TSharedPtr<FDebugObjectInfo>>)
+									.ItemHeight(16.0f)
+									.TreeItemsSource(DebugComponents)
+									.OnGenerateRow_Static(&ComponentTreeViewGenerateRow)
+									.OnGetChildren_Static(&ComponentTreeViewGetChildren)
+									.OnExpansionChanged_Static(&ComponentTreeViewExpansionChanged)
+									.SelectionMode(ESelectionMode::Single)
+									.OnSelectionChanged(InArgs._OnSelectionChanged);
+
+	ChildSlot
+	[
+		ComponentTreeView.ToSharedRef()
+	];
+}
+
+void RestoreExpansion(TArray<TSharedPtr<FDebugObjectInfo>>& Components, TSharedPtr<STreeView<TSharedPtr<FDebugObjectInfo>>>& TreeView)
+{
+	for(TSharedPtr<FDebugObjectInfo> &Component : Components)
+	{
+		TreeView->SetItemExpansion(Component, Component->bExpanded);
+		RestoreExpansion(Component->Children, TreeView);
+	}
+}
+
+void SRewindDebuggerComponentTree::Refresh()
+{
+	ComponentTreeView->RebuildList();
+
+	if (DebugComponents)
+	{
+		// make sure any newly added TreeView nodes are created expanded
+		RestoreExpansion(*DebugComponents, ComponentTreeView);
+	}
+}
+
+#undef LOCTEXT_NAMESPACE
diff --git a/Engine/Plugins/Animation/GameplayInsights/Source/RewindDebugger/Private/SRewindDebuggerComponentTree.h b/Engine/Plugins/Animation/GameplayInsights/Source/RewindDebugger/Private/SRewindDebuggerComponentTree.h
new file mode 100644
--- /dev/null
+++ b/Engine/Plugins/Animation/GameplayInsights/Source/RewindDebugger/Private/SRewindDebuggerComponentTree.h
@@ -0,0 +1,44 @@
+// Copyright Epic Games, Inc. All Rights Reserved.
+#pragma once
+
+#include "CoreMinimal.h"
+#include "RewindDebuggerModule.h"
+#include "Widgets/DeclarativeSyntaxSupport.h"
+#include "Widgets/SCompoundWidget.h"
+#include "Widgets/Views/STreeView.h"
+
+
+class SRewindDebuggerComponentTree : public SCompoundWidget
+{
+	using FOnSelectionChanged = typename STreeView<TSharedPtr<FDebugObjectInfo>>::FOnSelectionChanged;
+
+public:
+	SLATE_BEGIN_ARGS(SRewindDebuggerComponentTree) { }
+		SLATE_ARGUMENT( TArray< TSharedPtr< FDebugObjectInfo > >*, DebugComponents );
+		SLATE_EVENT( FOnSelectionChanged, OnSelectionChanged )
+	SLATE_END_ARGS()
+	
+public:
+
+	/**
+	* Default constructor.
+	*/
+	SRewindDebuggerComponentTree();
+	virtual ~SRewindDebuggerComponentTree();
+
+	/**
+	* Constructs the application.
+	*
+	* @param InArgs The Slate argument list.
+	*/
+	void Construct(const FArguments& InArgs);
+
+	void Refresh();
+
+private:
+	TArray<TSharedPtr<FDebugObjectInfo>>* DebugComponents;
+    void OnComponentSelectionChanged(TSharedPtr<FDebugObjectInfo> SelectedItem, ESelectInfo::Type SelectInfo);
+	TSharedPtr<STreeView<TSharedPtr<FDebugObjectInfo>>> ComponentTreeView;
+
+	FOnSelectionChanged OnSelectionChanged;
+};
