Don't try to update mesh poses from Rewind Debugger while paused unless there is active recording data

diff --git a/Engine/Plugins/Animation/GameplayInsights/Source/RewindDebugger/Private/RewindDebugger.cpp b/Engine/Plugins/Animation/GameplayInsights/Source/RewindDebugger/Private/RewindDebugger.cpp
--- a/Engine/Plugins/Animation/GameplayInsights/Source/RewindDebugger/Private/RewindDebugger.cpp
+++ b/Engine/Plugins/Animation/GameplayInsights/Source/RewindDebugger/Private/RewindDebugger.cpp
@@ -454,69 +454,72 @@ void FRewindDebugger::Tick(float DeltaTime)
 			}
 			else
 			{
-				if (ControlState == EControlState::Play || ControlState == EControlState::PlayReverse)
+				if (RecordingDuration.Get() > 0)
 				{
-					float Rate = PlaybackRate * (ControlState == EControlState::Play ? 1 : -1);
-					SetCurrentScrubTime(FMath::Clamp(CurrentScrubTime + Rate * DeltaTime, 0.0f, RecordingDuration.Get()));
-					TrackCursorDelegate.ExecuteIfBound(Rate<0);
-
-					if (CurrentScrubTime == 0 || CurrentScrubTime == RecordingDuration.Get())
+					if (ControlState == EControlState::Play || ControlState == EControlState::PlayReverse)
 					{
-						// pause at end.
-						ControlState = EControlState::Pause;
+						float Rate = PlaybackRate * (ControlState == EControlState::Play ? 1 : -1);
+						SetCurrentScrubTime(FMath::Clamp(CurrentScrubTime + Rate * DeltaTime, 0.0f, RecordingDuration.Get()));
+						TrackCursorDelegate.ExecuteIfBound(Rate<0);
+
+						if (CurrentScrubTime == 0 || CurrentScrubTime == RecordingDuration.Get())
+						{
+							// pause at end.
+							ControlState = EControlState::Pause;
+						}
 					}
-				}
 
-				double CurrentTraceTime = TraceTime.Get();
+					double CurrentTraceTime = TraceTime.Get();
 
-				// update pose on all SkeletalMeshComponents
-				ULevel* CurLevel = World->GetCurrentLevel();
-				for( AActor* LevelActor : CurLevel->Actors )
-				{
-					if (LevelActor)
+					// update pose on all SkeletalMeshComponents
+					ULevel* CurLevel = World->GetCurrentLevel();
+					for( AActor* LevelActor : CurLevel->Actors )
 					{
-						TInlineComponentArray<USkeletalMeshComponent*> SkeletalMeshComponents;
-						LevelActor->GetComponents(SkeletalMeshComponents);
-
-						for (USkeletalMeshComponent* MeshComponent : SkeletalMeshComponents)
+						if (LevelActor)
 						{
-							int64 ObjectId = FObjectTrace::GetObjectId(MeshComponent);
+							TInlineComponentArray<USkeletalMeshComponent*> SkeletalMeshComponents;
+							LevelActor->GetComponents(SkeletalMeshComponents);
 
-							AnimationProvider->ReadSkeletalMeshPoseTimeline(ObjectId, [this, ObjectId, CurrentTraceTime, MeshComponent, AnimationProvider](const IAnimationProvider::SkeletalMeshPoseTimeline& TimelineData, bool bHasCurves)
+							for (USkeletalMeshComponent* MeshComponent : SkeletalMeshComponents)
 							{
-								double PrecedingPoseTime;
-								double FollowingPoseTime;
-								const FSkeletalMeshPoseMessage* PrecedingPose;
-								const FSkeletalMeshPoseMessage* FollowingPose;
-
-								TimelineData.FindNearestEvents(CurrentTraceTime, PrecedingPose, PrecedingPoseTime, FollowingPose, FollowingPoseTime);
-
-								if (FollowingPose || PrecedingPose)
-								{	
-									// Ideally we should iterpolate between the two poses here if both are valid (we would need to transform the component space transforms to local space though)
-									const FSkeletalMeshPoseMessage& PoseMessage = FollowingPose ? *FollowingPose : *PrecedingPose;
-
-									FTransform ComponentWorldTransform;
-									const FSkeletalMeshInfo* SkeletalMeshInfo = AnimationProvider->FindSkeletalMeshInfo(PoseMessage.MeshId);
-									AnimationProvider->GetSkeletalMeshComponentSpacePose(PoseMessage, *SkeletalMeshInfo, ComponentWorldTransform, MeshComponent->GetEditableComponentSpaceTransforms());
-
-									if (MeshComponentsToReset.Find(ObjectId) == nullptr)
-									{
-										FMeshComponentResetData ResetData;
-										ResetData.Component = MeshComponent;
-										ResetData.RelativeTransform = MeshComponent->GetRelativeTransform();
-										MeshComponentsToReset.Add(ObjectId, ResetData);
+								int64 ObjectId = FObjectTrace::GetObjectId(MeshComponent);
+
+								AnimationProvider->ReadSkeletalMeshPoseTimeline(ObjectId, [this, ObjectId, CurrentTraceTime, MeshComponent, AnimationProvider](const IAnimationProvider::SkeletalMeshPoseTimeline& TimelineData, bool bHasCurves)
+								{
+									double PrecedingPoseTime;
+									double FollowingPoseTime;
+									const FSkeletalMeshPoseMessage* PrecedingPose;
+									const FSkeletalMeshPoseMessage* FollowingPose;
+
+									TimelineData.FindNearestEvents(CurrentTraceTime, PrecedingPose, PrecedingPoseTime, FollowingPose, FollowingPoseTime);
+
+									if (FollowingPose || PrecedingPose)
+									{	
+										// Ideally we should iterpolate between the two poses here if both are valid (we would need to transform the component space transforms to local space though)
+										const FSkeletalMeshPoseMessage& PoseMessage = FollowingPose ? *FollowingPose : *PrecedingPose;
+
+										FTransform ComponentWorldTransform;
+										const FSkeletalMeshInfo* SkeletalMeshInfo = AnimationProvider->FindSkeletalMeshInfo(PoseMessage.MeshId);
+										AnimationProvider->GetSkeletalMeshComponentSpacePose(PoseMessage, *SkeletalMeshInfo, ComponentWorldTransform, MeshComponent->GetEditableComponentSpaceTransforms());
+
+										if (MeshComponentsToReset.Find(ObjectId) == nullptr)
+										{
+											FMeshComponentResetData ResetData;
+											ResetData.Component = MeshComponent;
+											ResetData.RelativeTransform = MeshComponent->GetRelativeTransform();
+											MeshComponentsToReset.Add(ObjectId, ResetData);
+										}
+
+										// todo: we probably need to take into account tick order requirements for attached objects here
+										MeshComponent->SetWorldTransform(ComponentWorldTransform);
+										MeshComponent->SetForcedLOD(PoseMessage.LodIndex + 1);
 									}
+								});
 
-									// todo: we probably need to take into account tick order requirements for attached objects here
-									MeshComponent->SetWorldTransform(ComponentWorldTransform);
-									MeshComponent->SetForcedLOD(PoseMessage.LodIndex + 1);
-								}
-							});
-
-							// calling this here, even on meshes which have no recorded data makes meshes with a MasterPoseComponent attach properly during replay
-							// longer term solution: we should be recording those meshes
-							MeshComponent->ApplyEditedComponentSpaceTransforms();
+								// calling this here, even on meshes which have no recorded data makes meshes with a MasterPoseComponent attach properly during replay
+								// longer term solution: we should be recording those meshes
+								MeshComponent->ApplyEditedComponentSpaceTransforms();
+							}
 						}
 					}
 				}
