Rewind debugger Debug Views creator factory system.
Allows anyone to register a widget that will be shown when a particular component type is selected in the component treeview

diff --git a/Engine/Plugins/Animation/GameplayInsights/Source/GameplayInsights/Private/GameplayInsightsDebugViewCreator.cpp b/Engine/Plugins/Animation/GameplayInsights/Source/GameplayInsights/Private/GameplayInsightsDebugViewCreator.cpp
new file mode 100644
--- /dev/null
+++ b/Engine/Plugins/Animation/GameplayInsights/Source/GameplayInsights/Private/GameplayInsightsDebugViewCreator.cpp
@@ -0,0 +1,45 @@
+// Copyright Epic Games, Inc. All Rights Reserved.
+
+#pragma once
+
+#include "GameplayInsightsDebugViewCreator.h"
+#include "IGameplayProvider.h"
+
+namespace TraceServices
+{
+    class IAnalysisSession;
+}
+
+void FGameplayInsightsDebugViewCreator::RegisterDebugViewCreator(FName TypeName, FCreateDebugView Creator)
+{
+    ViewCreators.Add(TypeName, Creator);
+}
+
+void FGameplayInsightsDebugViewCreator::CreateDebugViews(uint64 ObjectId, double CurrentTime, const TraceServices::IAnalysisSession& Session, TArray<TSharedPtr<IGameplayInsightsDebugView>>& OutDebugViews) 
+{
+    if (const IGameplayProvider* GameplayProvider = Session.ReadProvider<IGameplayProvider>("GameplayProvider"))
+    {
+		TraceServices::FAnalysisSessionReadScope SessionReadScope(Session);
+
+        const FObjectInfo& ObjectInfo = GameplayProvider->GetObjectInfo(ObjectId);
+        uint64 ClassId = ObjectInfo.ClassId;
+
+        while (ClassId != 0)
+        {
+            const FClassInfo& ClassInfo = GameplayProvider->GetClassInfo(ClassId);
+
+            TArray<FCreateDebugView> List;
+            ViewCreators.MultiFind(ClassInfo.Name, List);
+
+            for (FCreateDebugView& Creator : List)
+            {
+                if (Creator.IsBound())
+                {
+                    OutDebugViews.Add(Creator.Execute(ObjectId, CurrentTime, Session));
+                }
+            }
+
+            ClassId = ClassInfo.SuperId;
+        }
+    }
+}
\ No newline at end of file
diff --git a/Engine/Plugins/Animation/GameplayInsights/Source/GameplayInsights/Private/GameplayInsightsDebugViewCreator.h b/Engine/Plugins/Animation/GameplayInsights/Source/GameplayInsights/Private/GameplayInsightsDebugViewCreator.h
new file mode 100644
--- /dev/null
+++ b/Engine/Plugins/Animation/GameplayInsights/Source/GameplayInsights/Private/GameplayInsightsDebugViewCreator.h
@@ -0,0 +1,19 @@
+// Copyright Epic Games, Inc. All Rights Reserved.
+
+#pragma once
+
+#include "CoreMinimal.h"
+#include "IGameplayInsightsDebugViewCreator.h"
+#include "Containers/Map.h"
+
+class FGameplayInsightsDebugViewCreator : public IGameplayInsightsDebugViewCreator
+{
+	public:
+		virtual ~FGameplayInsightsDebugViewCreator() {};
+
+		virtual void RegisterDebugViewCreator(FName TypeName, FCreateDebugView Creator) override;
+		virtual void CreateDebugViews(uint64 ObjectId, double CurrentTime,  const TraceServices::IAnalysisSession& InAnalysisSession, TArray<TSharedPtr<IGameplayInsightsDebugView>>& OutDebugViews) override;
+
+	private:
+		TMultiMap<FName, FCreateDebugView> ViewCreators;
+};
\ No newline at end of file
diff --git a/Engine/Plugins/Animation/GameplayInsights/Source/GameplayInsights/Private/GameplayInsightsModule.cpp b/Engine/Plugins/Animation/GameplayInsights/Source/GameplayInsights/Private/GameplayInsightsModule.cpp
--- a/Engine/Plugins/Animation/GameplayInsights/Source/GameplayInsights/Private/GameplayInsightsModule.cpp
+++ b/Engine/Plugins/Animation/GameplayInsights/Source/GameplayInsights/Private/GameplayInsightsModule.cpp
@@ -39,11 +39,6 @@
 
 const FName GameplayInsightsTabs::DocumentTab("DocumentTab");
 
-TSharedRef<IAnimGraphSchematicView> FGameplayInsightsModule::CreateAnimGraphSchematicView(uint64 InAnimInstanceId, double InTimeMarker, const TraceServices::IAnalysisSession& InAnalysisSession)
-{
-	return SNew(SAnimGraphSchematicView, InAnimInstanceId, InTimeMarker, InAnalysisSession);
-}
-
 void FGameplayInsightsModule::StartupModule()
 {
 	IModularFeatures::Get().RegisterModularFeature(TraceServices::ModuleFeatureName, &GameplayTraceModule);
@@ -66,7 +61,16 @@ void FGameplayInsightsModule::StartupModule()
 		}
 	});
 
+
+	DebugViewCreator.RegisterDebugViewCreator("AnimInstance", IGameplayInsightsDebugViewCreator::FCreateDebugView::CreateLambda(
+			[](uint64 InAnimInstanceId, double InTimeMarker, const TraceServices::IAnalysisSession& InAnalysisSession)
+			{
+				return SNew(SAnimGraphSchematicView, InAnimInstanceId, InTimeMarker, InAnalysisSession);
+			}
+	));
+
 #if WITH_EDITOR
+
 	if (!IsRunningCommandlet())
 	{
 		IAnimationBlueprintEditorModule& AnimationBlueprintEditorModule = FModuleManager::LoadModuleChecked<IAnimationBlueprintEditorModule>("AnimationBlueprintEditor");
diff --git a/Engine/Plugins/Animation/GameplayInsights/Source/GameplayInsights/Private/GameplayInsightsModule.h b/Engine/Plugins/Animation/GameplayInsights/Source/GameplayInsights/Private/GameplayInsightsModule.h
--- a/Engine/Plugins/Animation/GameplayInsights/Source/GameplayInsights/Private/GameplayInsightsModule.h
+++ b/Engine/Plugins/Animation/GameplayInsights/Source/GameplayInsights/Private/GameplayInsightsModule.h
@@ -8,6 +8,7 @@
 #include "GameplayTimingViewExtender.h"
 #include "Framework/Docking/TabManager.h"
 #include "IGameplayInsightsModule.h"
+#include "GameplayInsightsDebugViewCreator.h"
 
 struct FInsightsMajorTabExtender;
 
@@ -27,11 +28,12 @@ public:
 	virtual void StartupModule() override;
 	virtual void ShutdownModule() override;
 
+	// IGameplayInsightsModule interface
+	virtual IGameplayInsightsDebugViewCreator* GetDebugViewCreator() override { return &DebugViewCreator; }
+
 	// Spawn a document tab
 	TSharedRef<SDockTab> SpawnTimingProfilerDocumentTab(const FTabManager::FSearchPreference& InSearchPreference);
 
-	virtual TSharedRef<IAnimGraphSchematicView> CreateAnimGraphSchematicView(uint64 InAnimInstanceId, double InTimeMarker, const TraceServices::IAnalysisSession& InAnalysisSession) override;
-
 protected:
 	void RegisterTimingProfilerLayoutExtensions(FInsightsMajorTabExtender& InOutExtender);
 #if WITH_EDITOR
@@ -39,6 +41,8 @@ protected:
 #endif
 
 private:
+	FGameplayInsightsDebugViewCreator DebugViewCreator;
+
 	FGameplayTraceModule GameplayTraceModule;
 
 	FGameplayTimingViewExtender GameplayTimingViewExtender;
diff --git a/Engine/Plugins/Animation/GameplayInsights/Source/GameplayInsights/Private/SAnimGraphSchematicView.cpp b/Engine/Plugins/Animation/GameplayInsights/Source/GameplayInsights/Private/SAnimGraphSchematicView.cpp
--- a/Engine/Plugins/Animation/GameplayInsights/Source/GameplayInsights/Private/SAnimGraphSchematicView.cpp
+++ b/Engine/Plugins/Animation/GameplayInsights/Source/GameplayInsights/Private/SAnimGraphSchematicView.cpp
@@ -822,15 +822,6 @@ void SAnimGraphSchematicView::SetTimeMarker(double InTimeMarker)
 	}
 }
 
-void SAnimGraphSchematicView::SetAnimInstanceId(uint64 InAnimInstanceId)
-{
-	if (AnimInstanceId != InAnimInstanceId)
-	{
-		AnimInstanceId = InAnimInstanceId;
-		RefreshNodes();
-	}
-}
-
 TSharedRef<SWidget> SAnimGraphSchematicView::HandleGetViewMenuContent()
 {
 	FMenuBuilder MenuBuilder(true, nullptr);
diff --git a/Engine/Plugins/Animation/GameplayInsights/Source/GameplayInsights/Private/SAnimGraphSchematicView.h b/Engine/Plugins/Animation/GameplayInsights/Source/GameplayInsights/Private/SAnimGraphSchematicView.h
--- a/Engine/Plugins/Animation/GameplayInsights/Source/GameplayInsights/Private/SAnimGraphSchematicView.h
+++ b/Engine/Plugins/Animation/GameplayInsights/Source/GameplayInsights/Private/SAnimGraphSchematicView.h
@@ -3,7 +3,7 @@
 #pragma once
 
 #include "CoreMinimal.h"
-#include "IAnimGraphSchematicView.h"
+#include "IGameplayInsightsDebugView.h"
 #include "Widgets/SCompoundWidget.h"
 #include "Widgets/Views/STreeView.h"
 
@@ -20,14 +20,13 @@ class SComboButton;
 class SSplitter;
 enum class EAnimGraphSchematicFilterState;
 
-class SAnimGraphSchematicView : public IAnimGraphSchematicView
+class SAnimGraphSchematicView : public IGameplayInsightsDebugView
 {
 	SLATE_BEGIN_ARGS(SAnimGraphSchematicView) {}
 	SLATE_END_ARGS()
 public:
 	void Construct(const FArguments& InArgs, uint64 InAnimInstanceId, double InTimeMarker, const TraceServices::IAnalysisSession& InAnalysisSession);
 
-	virtual void SetAnimInstanceId(uint64 InAnimInstanceId) override;
 	uint64 GetAnimInstanceId() const { return AnimInstanceId; }
 
 	virtual void SetTimeMarker(double InTimeMarker) override;
diff --git a/Engine/Plugins/Animation/GameplayInsights/Source/GameplayInsights/Public/IAnimGraphSchematicView.h b/Engine/Plugins/Animation/GameplayInsights/Source/GameplayInsights/Public/IAnimGraphSchematicView.h
deleted file mode 100644
diff --git a/Engine/Plugins/Animation/GameplayInsights/Source/GameplayInsights/Public/IGameplayInsightsDebugView.h b/Engine/Plugins/Animation/GameplayInsights/Source/GameplayInsights/Public/IGameplayInsightsDebugView.h
new file mode 100644
--- /dev/null
+++ b/Engine/Plugins/Animation/GameplayInsights/Source/GameplayInsights/Public/IGameplayInsightsDebugView.h
@@ -0,0 +1,15 @@
+// Copyright Epic Games, Inc. All Rights Reserved.
+
+#pragma once
+
+#include "CoreMinimal.h"
+#include "Widgets/SCompoundWidget.h"
+
+// Interface class for views of recorded debug data
+// a delegate which creates implementations of this interface can be registered with a UObject type name to IGameplayInsightsDebugViewCreator
+// and they will be shown in the Rewind Debugger when that type of object is selected
+class IGameplayInsightsDebugView : public SCompoundWidget
+{
+	public:
+		virtual void SetTimeMarker(double InTimeMarker) = 0;
+};
diff --git a/Engine/Plugins/Animation/GameplayInsights/Source/GameplayInsights/Public/IGameplayInsightsDebugViewCreator.h b/Engine/Plugins/Animation/GameplayInsights/Source/GameplayInsights/Public/IGameplayInsightsDebugViewCreator.h
new file mode 100644
--- /dev/null
+++ b/Engine/Plugins/Animation/GameplayInsights/Source/GameplayInsights/Public/IGameplayInsightsDebugViewCreator.h
@@ -0,0 +1,22 @@
+// Copyright Epic Games, Inc. All Rights Reserved.
+
+#pragma once
+
+#include "CoreMinimal.h"
+#include "IGameplayInsightsDebugView.h"
+
+namespace TraceServices
+{
+	class IAnalysisSession;
+}
+
+// this class handles creating debug view widgets for the Rewind Debugger.
+// systems can register a debug view creator function with a UObject type name, and when an object of that type is selected, that widget will be created and shown by the debugger.
+class IGameplayInsightsDebugViewCreator
+{
+	public:
+		DECLARE_DELEGATE_RetVal_ThreeParams(TSharedRef<IGameplayInsightsDebugView>, FCreateDebugView, uint64 /*ObjectId*/, double /*Time*/, const TraceServices::IAnalysisSession& InAnalysisSession);
+
+		virtual void RegisterDebugViewCreator(FName TypeName, FCreateDebugView Creator) = 0;
+		virtual void CreateDebugViews(uint64 ObjectId, double CurrentTime, const TraceServices::IAnalysisSession& InAnalysisSession, TArray<TSharedPtr<IGameplayInsightsDebugView>>& OutDebugViews) = 0;
+};
\ No newline at end of file
diff --git a/Engine/Plugins/Animation/GameplayInsights/Source/GameplayInsights/Public/IGameplayInsightsModule.h b/Engine/Plugins/Animation/GameplayInsights/Source/GameplayInsights/Public/IGameplayInsightsModule.h
--- a/Engine/Plugins/Animation/GameplayInsights/Source/GameplayInsights/Public/IGameplayInsightsModule.h
+++ b/Engine/Plugins/Animation/GameplayInsights/Source/GameplayInsights/Public/IGameplayInsightsModule.h
@@ -3,23 +3,26 @@
 #pragma once
 
 #include "CoreMinimal.h"
-#include "IAnimGraphSchematicView.h"
+#include "IGameplayInsightsDebugViewCreator.h"
+
 
 namespace GameplayInsightsTabs
 {
 	extern GAMEPLAYINSIGHTS_API const FName DocumentTab;
-
 };
 
+
 namespace TraceServices
 {
 	class IAnalysisSession;
 }
 
+class IGameplayInsightsDebugViewCreator;
+
 class IGameplayInsightsModule
 	: public IModuleInterface
 {
 public:
-	/** Creates a widget for displaying the anim graph schematic */
-	virtual TSharedRef<IAnimGraphSchematicView> CreateAnimGraphSchematicView(uint64 InAnimInstanceId, double InTimeMarker, const TraceServices::IAnalysisSession& InAnalysisSession) = 0;
+	
+	virtual IGameplayInsightsDebugViewCreator* GetDebugViewCreator() = 0;
 };
diff --git a/Engine/Plugins/Animation/GameplayInsights/Source/RewindDebugger/Private/SRewindDebugger.cpp b/Engine/Plugins/Animation/GameplayInsights/Source/RewindDebugger/Private/SRewindDebugger.cpp
--- a/Engine/Plugins/Animation/GameplayInsights/Source/RewindDebugger/Private/SRewindDebugger.cpp
+++ b/Engine/Plugins/Animation/GameplayInsights/Source/RewindDebugger/Private/SRewindDebugger.cpp
@@ -235,13 +235,10 @@ void SRewindDebugger::Construct(const FArguments& InArgs, TSharedRef<FUICommandL
 				.DebugComponents(InArgs._DebugComponents)
 				.OnSelectionChanged(this, &SRewindDebugger::OnComponentSelectionChanged);
 
-	// everything related to creating SAnimGraphSchematicView should be moved to a separate file
-	IUnrealInsightsModule &UnrealInsightsModule = FModuleManager::LoadModuleChecked<IUnrealInsightsModule>("TraceInsights");
-	IGameplayInsightsModule& GameplayInsightsModule = FModuleManager::LoadModuleChecked<IGameplayInsightsModule>("GameplayInsights");
-	TSharedPtr<const TraceServices::IAnalysisSession> Session = UnrealInsightsModule.GetAnalysisSession();
 
-	AnimGraphView = GameplayInsightsModule.CreateAnimGraphSchematicView(0, TraceTime.Get(), *Session);
-	TraceTime.OnPropertyChanged = TraceTime.OnPropertyChanged.CreateSP(AnimGraphView.Get(), &IAnimGraphSchematicView::SetTimeMarker);
+	DebugViewContainer = SNew(SVerticalBox);
+
+	TraceTime.OnPropertyChanged = TraceTime.OnPropertyChanged.CreateRaw(this, &SRewindDebugger::TraceTimeChanged);
 
 	ChildSlot
 	[
@@ -330,8 +327,7 @@ void SRewindDebugger::Construct(const FArguments& InArgs, TSharedRef<FUICommandL
 			]
 			+SVerticalBox::Slot() .VAlign(VAlign_Top) .FillHeight(1.0f)
 			[
-				// todo: this should be a proxy slot that we swap out based on which component is selected in the treeview
-				AnimGraphView.ToSharedRef()
+				DebugViewContainer.ToSharedRef()
 			]
 		]
 	];
@@ -342,15 +338,38 @@ void SRewindDebugger::RefreshDebugComponents()
 	ComponentTreeView->Refresh();
 }
 
-void SRewindDebugger::OnComponentSelectionChanged(TSharedPtr<FDebugObjectInfo> SelectedItem, ESelectInfo::Type SelectInfo)
+void SRewindDebugger::TraceTimeChanged(double Time)
 {
-	if (SelectedItem.IsValid())
+	for(TSharedPtr<IGameplayInsightsDebugView>& DebugView : DebugViews)
 	{
-		AnimGraphView->SetAnimInstanceId(SelectedItem->ObjectId);
+		DebugView->SetTimeMarker(Time);
 	}
-	else
+}
+
+
+void SRewindDebugger::OnComponentSelectionChanged(TSharedPtr<FDebugObjectInfo> SelectedItem, ESelectInfo::Type SelectInfo)
+{
+	DebugViewContainer->ClearChildren();
+	DebugViews.Empty();
+
+	if (SelectedItem.IsValid())
 	{
-		AnimGraphView->SetAnimInstanceId(0);
+		// everything related to creating SAnimGraphSchematicView should be moved to a separate file
+		IUnrealInsightsModule &UnrealInsightsModule = FModuleManager::LoadModuleChecked<IUnrealInsightsModule>("TraceInsights");
+		IGameplayInsightsModule& GameplayInsightsModule = FModuleManager::LoadModuleChecked<IGameplayInsightsModule>("GameplayInsights");
+		TSharedPtr<const TraceServices::IAnalysisSession> Session = UnrealInsightsModule.GetAnalysisSession();
+
+		GameplayInsightsModule.GetDebugViewCreator()->CreateDebugViews(SelectedItem->ObjectId, TraceTime.Get(), *Session, DebugViews);
+
+		for(TSharedPtr<IGameplayInsightsDebugView>& DebugView : DebugViews)
+		{
+			// for now, just add all the widgets to a vertical box.
+			// this will be replaced by something better once there are more DebugViews implemented
+			DebugViewContainer->AddSlot()
+			[
+				DebugView.ToSharedRef()
+			];
+		}
 	}
 }
 
diff --git a/Engine/Plugins/Animation/GameplayInsights/Source/RewindDebugger/Private/SRewindDebugger.h b/Engine/Plugins/Animation/GameplayInsights/Source/RewindDebugger/Private/SRewindDebugger.h
--- a/Engine/Plugins/Animation/GameplayInsights/Source/RewindDebugger/Private/SRewindDebugger.h
+++ b/Engine/Plugins/Animation/GameplayInsights/Source/RewindDebugger/Private/SRewindDebugger.h
@@ -5,7 +5,7 @@
 #include "CoreMinimal.h"
 #include "Framework/Commands/UICommandList.h"
 #include "Framework/Docking/TabManager.h"
-#include "IAnimGraphSchematicView.h"
+#include "IGameplayInsightsDebugView.h"
 #include "RewindDebuggerModule.h"
 #include "RewindDebuggerTimeSliderController.h"
 #include "SRewindDebuggerComponentTree.h"
@@ -68,11 +68,14 @@ private:
 	TBindableProperty<float> RecordingDuration;
 	TBindableProperty<float> DebugTargetAnimInstanceId;
 
+	void TraceTimeChanged(double Time);
+
 	// component tree view
 	TArray<TSharedPtr<FDebugObjectInfo>>* DebugComponents;
     void OnComponentSelectionChanged(TSharedPtr<FDebugObjectInfo> SelectedItem, ESelectInfo::Type SelectInfo);
 	TSharedPtr<SRewindDebuggerComponentTree> ComponentTreeView;
 
 	// anim graph view 
-	TSharedPtr<IAnimGraphSchematicView> AnimGraphView;
+	TArray<TSharedPtr<IGameplayInsightsDebugView>> DebugViews;
+	TSharedPtr<SVerticalBox> DebugViewContainer;
 };
