Add support for object property display, and enabling/disabling property trace in Rewind Debugger.

diff --git a/Engine/Plugins/Animation/GameplayInsights/Source/GameplayInsights/Private/GameplayInsightsModule.cpp b/Engine/Plugins/Animation/GameplayInsights/Source/GameplayInsights/Private/GameplayInsightsModule.cpp
--- a/Engine/Plugins/Animation/GameplayInsights/Source/GameplayInsights/Private/GameplayInsightsModule.cpp
+++ b/Engine/Plugins/Animation/GameplayInsights/Source/GameplayInsights/Private/GameplayInsightsModule.cpp
@@ -18,6 +18,7 @@
 #include "Stats/Stats.h"
 #include "ObjectPropertyTrace.h"
 #include "SAnimGraphSchematicView.h"
+#include "SObjectPropertiesView.h"
 
 #if WITH_EDITOR
 #include "IAnimationBlueprintEditorModule.h"
@@ -62,6 +63,13 @@ void FGameplayInsightsModule::StartupModule()
 	});
 
 
+	DebugViewCreator.RegisterDebugViewCreator("Object", IGameplayInsightsDebugViewCreator::FCreateDebugView::CreateLambda(
+			[](uint64 InObjectId, double InTimeMarker, const TraceServices::IAnalysisSession& InAnalysisSession)
+			{
+				return SNew(SObjectPropertiesView, InObjectId, InTimeMarker, InAnalysisSession);
+			}
+	));
+
 	DebugViewCreator.RegisterDebugViewCreator("AnimInstance", IGameplayInsightsDebugViewCreator::FCreateDebugView::CreateLambda(
 			[](uint64 InAnimInstanceId, double InTimeMarker, const TraceServices::IAnalysisSession& InAnalysisSession)
 			{
@@ -447,6 +455,30 @@ void FGameplayInsightsModule::RegisterMenus()
 	}
 #endif
 }
+
+void FGameplayInsightsModule::EnableObjectPropertyTrace(UObject* Object, bool bEnable)
+{
+#if OBJECT_PROPERTY_TRACE_ENABLED
+	if (bEnable)
+	{
+		FObjectPropertyTrace::RegisterObject(Object);
+	}
+	else
+	{
+		FObjectPropertyTrace::UnregisterObject(Object);
+	}
+#endif
+}
+
+bool FGameplayInsightsModule::IsObjectPropertyTraceEnabled(UObject* Object)
+{
+#if OBJECT_PROPERTY_TRACE_ENABLED
+	return FObjectPropertyTrace::IsObjectRegistered(Object);
+#else
+	return false;
+#endif
+}
+
 #endif
 
 IMPLEMENT_MODULE(FGameplayInsightsModule, GameplayInsights);
diff --git a/Engine/Plugins/Animation/GameplayInsights/Source/GameplayInsights/Private/GameplayInsightsModule.h b/Engine/Plugins/Animation/GameplayInsights/Source/GameplayInsights/Private/GameplayInsightsModule.h
--- a/Engine/Plugins/Animation/GameplayInsights/Source/GameplayInsights/Private/GameplayInsightsModule.h
+++ b/Engine/Plugins/Animation/GameplayInsights/Source/GameplayInsights/Private/GameplayInsightsModule.h
@@ -31,6 +31,11 @@ public:
 	// IGameplayInsightsModule interface
 	virtual IGameplayInsightsDebugViewCreator* GetDebugViewCreator() override { return &DebugViewCreator; }
 
+#if WITH_EDITOR
+	virtual void EnableObjectPropertyTrace(UObject* Object, bool bEnable = true) override;
+	virtual bool IsObjectPropertyTraceEnabled(UObject* Object) override;
+#endif
+
 	// Spawn a document tab
 	TSharedRef<SDockTab> SpawnTimingProfilerDocumentTab(const FTabManager::FSearchPreference& InSearchPreference);
 
diff --git a/Engine/Plugins/Animation/GameplayInsights/Source/GameplayInsights/Private/SAnimGraphSchematicView.cpp b/Engine/Plugins/Animation/GameplayInsights/Source/GameplayInsights/Private/SAnimGraphSchematicView.cpp
--- a/Engine/Plugins/Animation/GameplayInsights/Source/GameplayInsights/Private/SAnimGraphSchematicView.cpp
+++ b/Engine/Plugins/Animation/GameplayInsights/Source/GameplayInsights/Private/SAnimGraphSchematicView.cpp
@@ -822,6 +822,11 @@ void SAnimGraphSchematicView::SetTimeMarker(double InTimeMarker)
 	}
 }
 
+FText SAnimGraphSchematicView::GetTitle()
+{
+	return LOCTEXT("Anim Graph Update", "Anim Graph Update");
+}
+
 TSharedRef<SWidget> SAnimGraphSchematicView::HandleGetViewMenuContent()
 {
 	FMenuBuilder MenuBuilder(true, nullptr);
diff --git a/Engine/Plugins/Animation/GameplayInsights/Source/GameplayInsights/Private/SAnimGraphSchematicView.h b/Engine/Plugins/Animation/GameplayInsights/Source/GameplayInsights/Private/SAnimGraphSchematicView.h
--- a/Engine/Plugins/Animation/GameplayInsights/Source/GameplayInsights/Private/SAnimGraphSchematicView.h
+++ b/Engine/Plugins/Animation/GameplayInsights/Source/GameplayInsights/Private/SAnimGraphSchematicView.h
@@ -30,6 +30,7 @@ public:
 	uint64 GetAnimInstanceId() const { return AnimInstanceId; }
 
 	virtual void SetTimeMarker(double InTimeMarker) override;
+	virtual FText GetTitle() override;
 private:
 	// Generate a row widget for an item
 	TSharedRef<ITableRow> HandleGenerateRow(TSharedRef<FAnimGraphSchematicNode> Item, const TSharedRef<STableViewBase>& OwnerTable);
diff --git a/Engine/Plugins/Animation/GameplayInsights/Source/GameplayInsights/Private/SObjectPropertiesView.cpp b/Engine/Plugins/Animation/GameplayInsights/Source/GameplayInsights/Private/SObjectPropertiesView.cpp
new file mode 100644
--- /dev/null
+++ b/Engine/Plugins/Animation/GameplayInsights/Source/GameplayInsights/Private/SObjectPropertiesView.cpp
@@ -0,0 +1,98 @@
+// Copyright Epic Games, Inc. All Rights Reserved.
+
+#include "SObjectPropertiesView.h"
+#include "AnimationProvider.h"
+#include "TraceServices/Model/AnalysisSession.h"
+#include "GameplayProvider.h"
+#include "TraceServices/Model/Frames.h"
+#include "VariantTreeNode.h"
+
+#define LOCTEXT_NAMESPACE "SObjectPropertiesView"
+
+void SObjectPropertiesView::Construct(const FArguments& InArgs, uint64 InObjectId, double InTimeMarker, const TraceServices::IAnalysisSession& InAnalysisSession)
+{
+	ObjectId = InObjectId;
+	AnalysisSession = &InAnalysisSession;
+
+	View = SNew(SVariantValueView, InAnalysisSession).OnGetVariantValues(this, &SObjectPropertiesView::GetVariantsAtFrame);
+
+	SetTimeMarker(InTimeMarker);
+
+	ChildSlot
+	[
+		View.ToSharedRef()
+	];
+
+}
+
+void SObjectPropertiesView::SetTimeMarker(double Time)
+{
+	TimeMarker = Time;
+	TraceServices::FAnalysisSessionReadScope SessionReadScope(*AnalysisSession);
+	const TraceServices::IFrameProvider& FramesProvider = TraceServices::ReadFrameProvider(*AnalysisSession);
+	TraceServices::FFrame MarkerFrame;
+	if(FramesProvider.GetFrameFromTime(ETraceFrameType::TraceFrameType_Game, TimeMarker, MarkerFrame))
+	{
+		View->RequestRefresh(MarkerFrame);
+	}
+}
+
+void SObjectPropertiesView::GetVariantsAtFrame(const TraceServices::FFrame& InFrame, TArray<TSharedRef<FVariantTreeNode>>& OutVariants) const
+{
+	if (const FGameplayProvider* GameplayProvider = AnalysisSession->ReadProvider<FGameplayProvider>(FGameplayProvider::ProviderName))
+	{
+		TraceServices::FAnalysisSessionReadScope SessionReadScope(*AnalysisSession);
+
+		const FClassInfo& ClassInfo = GameplayProvider->GetClassInfoFromObject(ObjectId);
+
+		if(ClassInfo.Properties.Num() > 0)
+		{
+			TSharedRef<FVariantTreeNode> Header = OutVariants.Add_GetRef(FVariantTreeNode::MakeHeader(LOCTEXT("Properties", "Properties"), INDEX_NONE));
+
+			// Build the class tree
+			TArray<TSharedPtr<FVariantTreeNode>> PropertyVariants;
+			PropertyVariants.SetNum(ClassInfo.Properties.Num());
+			for(int32 PropertyIndex = 0; PropertyIndex < ClassInfo.Properties.Num(); ++PropertyIndex)
+			{
+				const FClassPropertyInfo& PropertyInfo = ClassInfo.Properties[PropertyIndex];
+
+				// Add a string node with a default value
+				const TCHAR* Key = GameplayProvider->GetPropertyName(PropertyInfo.KeyStringId);
+				PropertyVariants[PropertyIndex] = FVariantTreeNode::MakeString(FText::FromString(Key), TEXT("Unknown"), PropertyIndex);
+
+				// note assumes that order is parent->child in the properties array
+				if(PropertyInfo.ParentId != INDEX_NONE)
+				{
+					PropertyVariants[PropertyInfo.ParentId]->AddChild(PropertyVariants[PropertyIndex].ToSharedRef());
+				}
+				else
+				{
+					Header->AddChild(PropertyVariants[PropertyIndex].ToSharedRef());
+				}
+			}
+
+			// object events
+			GameplayProvider->ReadObjectPropertiesTimeline(ObjectId, [this, &InFrame, &GameplayProvider, &PropertyVariants](const FGameplayProvider::ObjectPropertiesTimeline& InTimeline)
+			{
+				InTimeline.EnumerateEvents(InFrame.StartTime, InFrame.EndTime, [this, &GameplayProvider, &PropertyVariants](double InStartTime, double InEndTime, uint32 InDepth, const FObjectPropertiesMessage& InMessage)
+				{
+					GameplayProvider->EnumerateObjectPropertyValues(ObjectId, InMessage, [&PropertyVariants](const FObjectPropertyValue& InValue)
+					{
+						PropertyVariants[InValue.PropertyId]->GetValue().String.Value = InValue.Value;
+					});
+					return TraceServices::EEventEnumerate::Stop;
+				});
+			});
+		}
+	}
+}
+
+FText SObjectPropertiesView::GetTitle()
+{
+	return LOCTEXT("Object Properties", "Object Properties");
+}
+
+
+
+
+#undef LOCTEXT_NAMESPACE
diff --git a/Engine/Plugins/Animation/GameplayInsights/Source/GameplayInsights/Private/SObjectPropertiesView.h b/Engine/Plugins/Animation/GameplayInsights/Source/GameplayInsights/Private/SObjectPropertiesView.h
new file mode 100644
--- /dev/null
+++ b/Engine/Plugins/Animation/GameplayInsights/Source/GameplayInsights/Private/SObjectPropertiesView.h
@@ -0,0 +1,28 @@
+// Copyright Epic Games, Inc. All Rights Reserved.
+
+#pragma once
+
+#include "CoreMinimal.h"
+#include "IGameplayInsightsDebugView.h"
+#include "SVariantValueView.h"
+
+namespace TraceServices { class IAnalysisSession; }
+
+class SObjectPropertiesView : public IGameplayInsightsDebugView
+{
+	SLATE_BEGIN_ARGS(SObjectPropertiesView) {}
+	SLATE_END_ARGS()
+public:
+	void Construct(const FArguments& InArgs, uint64 InObjectId, double InTimeMarker, const TraceServices::IAnalysisSession& InAnalysisSession);
+
+	virtual void SetTimeMarker(double InTimeMarker) override;
+	virtual FText GetTitle() override;
+
+	void GetVariantsAtFrame(const TraceServices::FFrame& InFrame, TArray<TSharedRef<FVariantTreeNode>>& OutVariants) const;
+private:
+	TSharedPtr<SVariantValueView> View;
+
+	uint64 ObjectId;
+	double TimeMarker;
+	const TraceServices::IAnalysisSession* AnalysisSession;
+};
diff --git a/Engine/Plugins/Animation/GameplayInsights/Source/GameplayInsights/Public/IGameplayInsightsDebugView.h b/Engine/Plugins/Animation/GameplayInsights/Source/GameplayInsights/Public/IGameplayInsightsDebugView.h
--- a/Engine/Plugins/Animation/GameplayInsights/Source/GameplayInsights/Public/IGameplayInsightsDebugView.h
+++ b/Engine/Plugins/Animation/GameplayInsights/Source/GameplayInsights/Public/IGameplayInsightsDebugView.h
@@ -11,5 +11,9 @@
 class IGameplayInsightsDebugView : public SCompoundWidget
 {
 	public:
+		// text for widget header
+		virtual FText GetTitle() = 0;
+
+		// called by the debugger when the scrubbing bar position changes
 		virtual void SetTimeMarker(double InTimeMarker) = 0;
 };
diff --git a/Engine/Plugins/Animation/GameplayInsights/Source/GameplayInsights/Public/IGameplayInsightsModule.h b/Engine/Plugins/Animation/GameplayInsights/Source/GameplayInsights/Public/IGameplayInsightsModule.h
--- a/Engine/Plugins/Animation/GameplayInsights/Source/GameplayInsights/Public/IGameplayInsightsModule.h
+++ b/Engine/Plugins/Animation/GameplayInsights/Source/GameplayInsights/Public/IGameplayInsightsModule.h
@@ -25,4 +25,11 @@ class IGameplayInsightsModule
 public:
 	
 	virtual IGameplayInsightsDebugViewCreator* GetDebugViewCreator() = 0;
+
+#if WITH_EDITOR
+	// start or stop tracing all properties of this Object to Insights
+	virtual void EnableObjectPropertyTrace(UObject* Object, bool bEnable = true) = 0;
+	// check if an Object has it's properties tracing to Insights
+	virtual bool IsObjectPropertyTraceEnabled(UObject* Object) = 0;
+#endif
 };
diff --git a/Engine/Plugins/Animation/GameplayInsights/Source/RewindDebugger/Private/RewindDebugger.cpp b/Engine/Plugins/Animation/GameplayInsights/Source/RewindDebugger/Private/RewindDebugger.cpp
--- a/Engine/Plugins/Animation/GameplayInsights/Source/RewindDebugger/Private/RewindDebugger.cpp
+++ b/Engine/Plugins/Animation/GameplayInsights/Source/RewindDebugger/Private/RewindDebugger.cpp
@@ -25,14 +25,24 @@ FRewindDebugger::FRewindDebugger()  :
 {
 	RecordingDuration.Set(0);
 
-	FEditorDelegates::PostPIEStarted.AddRaw(this, &FRewindDebugger::OnPIEStarted);
+	FEditorDelegates::PreBeginPIE.AddRaw(this, &FRewindDebugger::OnPIEStarted);
 	FEditorDelegates::PausePIE.AddRaw(this, &FRewindDebugger::OnPIEPaused);
 	FEditorDelegates::ResumePIE.AddRaw(this, &FRewindDebugger::OnPIEResumed);
 	FEditorDelegates::EndPIE.AddRaw(this, &FRewindDebugger::OnPIEStopped);
+	FEditorDelegates::SingleStepPIE.AddRaw(this, &FRewindDebugger::OnPIESingleStepped);
 
 	DebugTargetActor.OnPropertyChanged = DebugTargetActor.OnPropertyChanged.CreateLambda([this](FString Target) { RefreshDebugComponents(); });
 
 	UnrealInsightsModule = &FModuleManager::LoadModuleChecked<IUnrealInsightsModule>("TraceInsights");
+
+	TickerHandle = FTicker::GetCoreTicker().AddTicker(TEXT("RewindDebugger"), 0.0f, [this](float DeltaTime)
+	{
+		QUICK_SCOPE_CYCLE_COUNTER(STAT_FRewindDebuggerModule_Tick);
+
+		Tick(DeltaTime);
+
+		return true;
+	});
 }
 
 FRewindDebugger::~FRewindDebugger() 
@@ -41,6 +51,10 @@ FRewindDebugger::~FRewindDebugger()
 	FEditorDelegates::PausePIE.RemoveAll(this);
 	FEditorDelegates::ResumePIE.RemoveAll(this);
 	FEditorDelegates::EndPIE.RemoveAll(this);
+	FEditorDelegates::SingleStepPIE.RemoveAll(this);
+
+	FTicker::GetCoreTicker().RemoveTicker(TickerHandle);
+
 }
 
 void FRewindDebugger::Initialize() 
@@ -68,6 +82,8 @@ void FRewindDebugger::OnPIEStarted(bool bSimulating)
 	bPIEStarted = true;
 	bPIESimulating = true;
 
+	UE::Trace::ToggleChannel(TEXT("Object"), true);
+
 	if (bAutoRecord)
 	{
 		StartRecording();
@@ -78,6 +94,13 @@ void FRewindDebugger::OnPIEPaused(bool bSimulating)
 {
 	bPIESimulating = false;
 	ControlState = EControlState::Pause;
+
+	if (bRecording)
+	{
+		UWorld* World = GetWorldToVisualize();
+		RecordingDuration.Set(FObjectTrace::GetWorldElapsedTime(World));
+		SetCurrentScrubTime(RecordingDuration.Get());
+	}
 }
 
 void FRewindDebugger::OnPIEResumed(bool bSimulating)
@@ -95,12 +118,37 @@ void FRewindDebugger::OnPIEResumed(bool bSimulating)
 
 	MeshComponentsToReset.Empty();
 }
+
+void FRewindDebugger::OnPIESingleStepped(bool bSimulating)
+{
+	// restore all relative transforms of any meshes that may have been moved while scrubbing
+	for (TTuple<uint64, FMeshComponentResetData>& MeshData : MeshComponentsToReset)
+	{
+		if (USkeletalMeshComponent* MeshComponent = MeshData.Value.Component.Get())
+		{
+			MeshComponent->SetRelativeTransform(MeshData.Value.RelativeTransform);
+		}
+	}
+
+	MeshComponentsToReset.Empty();
+
+	if (bRecording)
+	{
+		UWorld* World = GetWorldToVisualize();
+		RecordingDuration.Set(FObjectTrace::GetWorldElapsedTime(World));
+		SetCurrentScrubTime(RecordingDuration.Get());
+	}
+}
+
+
 void FRewindDebugger::OnPIEStopped(bool bSimulating)
 {
 	bPIEStarted = false;
 	bPIESimulating = false;
 	MeshComponentsToReset.Empty();
 
+	UE::Trace::ToggleChannel(TEXT("Object"), false);
+
 	StopRecording();
 	// clear the current recording (until we support playback in the Editor world on spawned actors)
 	RecordingDuration.Set(0);
@@ -206,7 +254,7 @@ void FRewindDebugger::StartRecording()
 	}
 
 	// Enable Object and Animation Trace filters
-	UE::Trace::ToggleChannel(TEXT("Object"), true);
+	UE::Trace::ToggleChannel(TEXT("ObjectProperties"), true);
 	UE::Trace::ToggleChannel(TEXT("Animation"), true);
 	UE::Trace::ToggleChannel(TEXT("Frame"), true);
 
@@ -226,7 +274,7 @@ void FRewindDebugger::StopRecording()
 	if (bRecording)
 	{
 		// Enable Object and Animation Trace filters
-		UE::Trace::ToggleChannel(TEXT("Object"), false);
+		UE::Trace::ToggleChannel(TEXT("ObjectProperties"), false);
 		UE::Trace::ToggleChannel(TEXT("Animation"), false);
 		UE::Trace::ToggleChannel(TEXT("Frame"), false);
 
@@ -456,6 +504,8 @@ void FRewindDebugger::Tick(float DeltaTime)
 			{
 				if (RecordingDuration.Get() > 0)
 				{
+					UpdateTraceTime(); 
+
 					if (ControlState == EControlState::Play || ControlState == EControlState::PlayReverse)
 					{
 						float Rate = PlaybackRate * (ControlState == EControlState::Play ? 1 : -1);
diff --git a/Engine/Plugins/Animation/GameplayInsights/Source/RewindDebugger/Private/RewindDebugger.h b/Engine/Plugins/Animation/GameplayInsights/Source/RewindDebugger/Private/RewindDebugger.h
--- a/Engine/Plugins/Animation/GameplayInsights/Source/RewindDebugger/Private/RewindDebugger.h
+++ b/Engine/Plugins/Animation/GameplayInsights/Source/RewindDebugger/Private/RewindDebugger.h
@@ -90,6 +90,7 @@ private:
 	void OnPIEPaused(bool bSimulating);
 	void OnPIEResumed(bool bSimulating);
 	void OnPIEStopped(bool bSimulating);
+	void OnPIESingleStepped(bool bSimulating);
 
 	const TraceServices::IAnalysisSession* GetAnalysisSession();
 	UWorld* GetWorldToVisualize() const;
@@ -137,4 +138,5 @@ private:
 	TMap<uint64, FMeshComponentResetData> MeshComponentsToReset;
 
 	class IUnrealInsightsModule *UnrealInsightsModule;
+	FDelegateHandle TickerHandle;
 };
diff --git a/Engine/Plugins/Animation/GameplayInsights/Source/RewindDebugger/Private/RewindDebuggerModule.cpp b/Engine/Plugins/Animation/GameplayInsights/Source/RewindDebugger/Private/RewindDebuggerModule.cpp
--- a/Engine/Plugins/Animation/GameplayInsights/Source/RewindDebugger/Private/RewindDebuggerModule.cpp
+++ b/Engine/Plugins/Animation/GameplayInsights/Source/RewindDebugger/Private/RewindDebuggerModule.cpp
@@ -27,6 +27,11 @@ static const FName RewindDebuggerTabName("RewindDebugger");
 
 TSharedRef<SDockTab> FRewindDebuggerModule::SpawnRewindDebuggerTab(const FSpawnTabArgs& SpawnTabArgs)
 {
+	if (FRewindDebugger::Instance() == nullptr)
+	{
+		FRewindDebugger::Initialize();
+	}
+
 	const TSharedRef<SDockTab> MajorTab = SNew(SDockTab)
 		.TabRole(ETabRole::NomadTab);
 
@@ -111,7 +116,6 @@ TSharedRef<SDockTab> FRewindDebuggerModule::SpawnRewindDebuggerTab(const FSpawnT
 
 void FRewindDebuggerModule::StartupModule()
 {
-	FRewindDebugger::Initialize();
 	FRewindDebuggerStyle::Initialize();
 	FRewindDebuggerCommands::Register();
 
@@ -123,21 +127,10 @@ void FRewindDebuggerModule::StartupModule()
 		.SetTooltipText(LOCTEXT("TooltipText", "Opens Rewind Debugger."));
 
 
-	TickerHandle = FTicker::GetCoreTicker().AddTicker(TEXT("RewindDebugger"), 0.0f, [this](float DeltaTime)
-	{
-		QUICK_SCOPE_CYCLE_COUNTER(STAT_FRewindDebuggerModule_Tick);
-
-		FRewindDebugger::Instance()->Tick(DeltaTime);
-
-		return true;
-	});
-
 }
 
 void FRewindDebuggerModule::ShutdownModule()
 {
-	FTicker::GetCoreTicker().RemoveTicker(TickerHandle);
-
 	FRewindDebuggerCommands::Unregister();
 	FRewindDebuggerStyle::Shutdown();
 	FRewindDebugger::Shutdown();
diff --git a/Engine/Plugins/Animation/GameplayInsights/Source/RewindDebugger/Private/RewindDebuggerModule.h b/Engine/Plugins/Animation/GameplayInsights/Source/RewindDebugger/Private/RewindDebuggerModule.h
--- a/Engine/Plugins/Animation/GameplayInsights/Source/RewindDebugger/Private/RewindDebuggerModule.h
+++ b/Engine/Plugins/Animation/GameplayInsights/Source/RewindDebugger/Private/RewindDebuggerModule.h
@@ -32,6 +32,4 @@ public:
 
 private:
 	TSharedPtr<SRewindDebugger> RewindDebuggerWidget;
-
-	FDelegateHandle TickerHandle;
 };
diff --git a/Engine/Plugins/Animation/GameplayInsights/Source/RewindDebugger/Private/SRewindDebugger.cpp b/Engine/Plugins/Animation/GameplayInsights/Source/RewindDebugger/Private/SRewindDebugger.cpp
--- a/Engine/Plugins/Animation/GameplayInsights/Source/RewindDebugger/Private/SRewindDebugger.cpp
+++ b/Engine/Plugins/Animation/GameplayInsights/Source/RewindDebugger/Private/SRewindDebugger.cpp
@@ -24,6 +24,8 @@
 #include "Widgets/Input/SComboButton.h"
 #include "Widgets/Input/SButton.h"
 #include "Widgets/Input/SSearchBox.h"
+#include "Widgets/Layout/SExpandableArea.h"
+#include "Widgets/Layout/SScrollBox.h"
 
 
 #define LOCTEXT_NAMESPACE "SAnimationInsights"
@@ -236,7 +238,7 @@ void SRewindDebugger::Construct(const FArguments& InArgs, TSharedRef<FUICommandL
 				.OnSelectionChanged(this, &SRewindDebugger::OnComponentSelectionChanged);
 
 
-	DebugViewContainer = SNew(SVerticalBox);
+	DebugViewContainer = SNew(SScrollBox);
 
 	TraceTime.OnPropertyChanged = TraceTime.OnPropertyChanged.CreateRaw(this, &SRewindDebugger::TraceTimeChanged);
 
@@ -367,7 +369,14 @@ void SRewindDebugger::OnComponentSelectionChanged(TSharedPtr<FDebugObjectInfo> S
 			// this will be replaced by something better once there are more DebugViews implemented
 			DebugViewContainer->AddSlot()
 			[
-				DebugView.ToSharedRef()
+				SNew(SExpandableArea)
+				.InitiallyCollapsed(false)
+				.Padding(10.0f)
+				.AreaTitle(DebugView->GetTitle())
+				.BodyContent()
+				[
+					DebugView.ToSharedRef()
+				]
 			];
 		}
 	}
diff --git a/Engine/Plugins/Animation/GameplayInsights/Source/RewindDebugger/Private/SRewindDebugger.h b/Engine/Plugins/Animation/GameplayInsights/Source/RewindDebugger/Private/SRewindDebugger.h
--- a/Engine/Plugins/Animation/GameplayInsights/Source/RewindDebugger/Private/SRewindDebugger.h
+++ b/Engine/Plugins/Animation/GameplayInsights/Source/RewindDebugger/Private/SRewindDebugger.h
@@ -14,6 +14,8 @@
 #include "Widgets/SCompoundWidget.h"
 
 
+class SScrollBox;
+
 class SRewindDebugger : public SCompoundWidget
 {
 	typedef TBindablePropertyInitializer<FString, BindingType_Out> DebugTargetInitializer;
@@ -77,5 +79,5 @@ private:
 
 	// anim graph view 
 	TArray<TSharedPtr<IGameplayInsightsDebugView>> DebugViews;
-	TSharedPtr<SVerticalBox> DebugViewContainer;
+	TSharedPtr<SScrollBox> DebugViewContainer;
 };
diff --git a/Engine/Plugins/Animation/GameplayInsights/Source/RewindDebugger/Private/SRewindDebuggerComponentTree.cpp b/Engine/Plugins/Animation/GameplayInsights/Source/RewindDebugger/Private/SRewindDebuggerComponentTree.cpp
--- a/Engine/Plugins/Animation/GameplayInsights/Source/RewindDebugger/Private/SRewindDebuggerComponentTree.cpp
+++ b/Engine/Plugins/Animation/GameplayInsights/Source/RewindDebugger/Private/SRewindDebuggerComponentTree.cpp
@@ -5,13 +5,15 @@
 #include "Editor.h"
 #include "Editor/EditorEngine.h"
 #include "Framework/Application/SlateApplication.h"
+#include "Framework/MultiBox/MultiBoxBuilder.h"
+#include "IGameplayInsightsModule.h"
+#include "Modules/ModuleManager.h"
 #include "ObjectTrace.h"
 #include "Styling/SlateIconFinder.h"
 #include "Widgets/Images/SImage.h"
 #include "Widgets/SOverlay.h"
 #include "Widgets/Views/SListView.h"
 
-
 #define LOCTEXT_NAMESPACE "SAnimationInsights"
 
 SRewindDebuggerComponentTree::SRewindDebuggerComponentTree() 
@@ -72,6 +74,57 @@ void ComponentTreeViewExpansionChanged(TSharedPtr<FDebugObjectInfo> InItem, bool
 	InItem->bExpanded = bShouldBeExpanded;
 }
 
+TSharedPtr<SWidget> SRewindDebuggerComponentTree::ComponentTreeOnContextMenuOpening()
+{
+	FMenuBuilder MenuBuilder(/*bInShouldCloseWindowAfterMenuSelection=*/true, nullptr, nullptr, /*bCloseSelfOnly=*/true);
+	IGameplayInsightsModule& GameplayInsightsModule = FModuleManager::LoadModuleChecked<IGameplayInsightsModule>("GameplayInsights");
+
+	MenuBuilder.BeginSection(NAME_None, LOCTEXT("Property Tracing", "Property Tracing"));
+	{
+		MenuBuilder.AddMenuEntry(
+			LOCTEXT("Trace Object Properties", "Trace Object Properties"),
+			LOCTEXT("Trace Object Properties Tooltip", "Record this object's properties so they will show in the Rewind Debugger when scrubbing."),
+			FSlateIcon(), 
+			FUIAction(FExecuteAction::CreateLambda([this, &GameplayInsightsModule]()
+				{
+					TArray<TSharedPtr<FDebugObjectInfo>> SelectedObjects = ComponentTreeView->GetSelectedItems();
+
+					if (SelectedObjects.Num() > 0)
+					{
+						if (UObject* Object = FObjectTrace::GetObjectFromId(SelectedObjects[0]->ObjectId))
+						{
+							GameplayInsightsModule.EnableObjectPropertyTrace(Object, !GameplayInsightsModule.IsObjectPropertyTraceEnabled(Object));
+						}
+					}
+				}),
+				FCanExecuteAction(),
+				FGetActionCheckState::CreateLambda([this, &GameplayInsightsModule]()
+				{
+					bool bEnabled = false;
+					TArray<TSharedPtr<FDebugObjectInfo>> SelectedObjects = ComponentTreeView->GetSelectedItems();
+
+					if (SelectedObjects.Num() > 0)
+					{
+						if (UObject* Object = FObjectTrace::GetObjectFromId(SelectedObjects[0]->ObjectId))
+						{
+							bEnabled = GameplayInsightsModule.IsObjectPropertyTraceEnabled(Object);
+						}
+					}
+                    return bEnabled ? ECheckBoxState::Checked : ECheckBoxState::Unchecked;
+				}
+				),
+				FIsActionButtonVisible()
+			),
+			NAME_None,
+			EUserInterfaceActionType::Check
+		);
+	}
+	MenuBuilder.EndSection();
+
+	return MenuBuilder.MakeWidget();
+}
+
+
 void SRewindDebuggerComponentTree::Construct(const FArguments& InArgs)
 {
 	DebugComponents = InArgs._DebugComponents;
@@ -83,7 +136,8 @@ void SRewindDebuggerComponentTree::Construct(const FArguments& InArgs)
 									.OnGetChildren_Static(&ComponentTreeViewGetChildren)
 									.OnExpansionChanged_Static(&ComponentTreeViewExpansionChanged)
 									.SelectionMode(ESelectionMode::Single)
-									.OnSelectionChanged(InArgs._OnSelectionChanged);
+									.OnSelectionChanged(InArgs._OnSelectionChanged)
+									.OnContextMenuOpening(this, &SRewindDebuggerComponentTree::ComponentTreeOnContextMenuOpening);
 
 	ChildSlot
 	[
diff --git a/Engine/Plugins/Animation/GameplayInsights/Source/RewindDebugger/Private/SRewindDebuggerComponentTree.h b/Engine/Plugins/Animation/GameplayInsights/Source/RewindDebugger/Private/SRewindDebuggerComponentTree.h
--- a/Engine/Plugins/Animation/GameplayInsights/Source/RewindDebugger/Private/SRewindDebuggerComponentTree.h
+++ b/Engine/Plugins/Animation/GameplayInsights/Source/RewindDebugger/Private/SRewindDebuggerComponentTree.h
@@ -35,6 +35,8 @@ public:
 
 	void Refresh();
 
+	TSharedPtr<SWidget> ComponentTreeOnContextMenuOpening();
+
 private:
 	TArray<TSharedPtr<FDebugObjectInfo>>* DebugComponents;
     void OnComponentSelectionChanged(TSharedPtr<FDebugObjectInfo> SelectedItem, ESelectInfo::Type SelectInfo);
