fix animation insights with LWC
support for loading LWC captures in non-LWC build and vice versa

diff --git a/Engine/Plugins/Animation/GameplayInsights/Source/GameplayInsights/Private/AnimationAnalyzer.cpp b/Engine/Plugins/Animation/GameplayInsights/Source/GameplayInsights/Private/AnimationAnalyzer.cpp
--- a/Engine/Plugins/Animation/GameplayInsights/Source/GameplayInsights/Private/AnimationAnalyzer.cpp
+++ b/Engine/Plugins/Animation/GameplayInsights/Source/GameplayInsights/Private/AnimationAnalyzer.cpp
@@ -110,39 +110,13 @@ bool FAnimationAnalyzer::OnEvent(uint16 RouteId, EStyle Style, const FOnEventCon
 		uint16 FrameCounter = EventData.GetValue<uint16>("FrameCounter");
 		uint16 LodIndex = EventData.GetValue<uint16>("LodIndex");
 
-		// ComponentToWorld. FTransform's data needs to be aligned.
 		TArrayView<const float> ComponentToWorldFloatArray = EventData.GetArrayView<float>("ComponentToWorld");
-		const int TransformFloatCount = ComponentToWorldFloatArray.Num();
-
-		// ComponentToWorld and Pose were encoded assuming 48 bytes (12 floats) for an FTransform.
-		check(TransformFloatCount == 12);
-
-		const float* TransformFloats = reinterpret_cast<const float*>(ComponentToWorldFloatArray.GetData());
-		const FQuat Rotation(TransformFloats[0], TransformFloats[1], TransformFloats[2], TransformFloats[3]);
-		const FVector Translation(TransformFloats[4], TransformFloats[5], TransformFloats[6]); // TransformFloats[7] is waste
-		const FVector Scale3D(TransformFloats[8], TransformFloats[9], TransformFloats[10]); // TransformFloats[11] is waste
-		const FTransform ComponentToWorld(Rotation, Translation, Scale3D);
-
-		// Pose. FTransform's data needs to be aligned.
 		TArrayView<const float> PoseFloatArray = EventData.GetArrayView<float>("Pose");
-		const int PoseTransformCount = PoseFloatArray.Num() / TransformFloatCount;
-		TArray<FTransform> Pose;
-		Pose.AddDefaulted(PoseTransformCount);
-		const float* PoseTransformFloats = reinterpret_cast<const float*>(PoseFloatArray.GetData());
-		for (int PoseIndex = 0; PoseIndex < PoseTransformCount; ++PoseIndex)
-		{
-			const FQuat PoseRotation(PoseTransformFloats[0], PoseTransformFloats[1], PoseTransformFloats[2], PoseTransformFloats[3]);
-			const FVector PoseTranslation(PoseTransformFloats[4], PoseTransformFloats[5], PoseTransformFloats[6]);
-			const FVector PoseScale3D(PoseTransformFloats[8], PoseTransformFloats[9], PoseTransformFloats[10]);
-			Pose[PoseIndex].SetComponents(PoseRotation, PoseTranslation, PoseScale3D);
-			PoseTransformFloats += TransformFloatCount;
-		}
-
 		TArrayView<const uint32> CurveIds = EventData.GetArrayView<uint32>("CurveIds");
 		TArrayView<const float> CurveValues = EventData.GetArrayView<float>("CurveValues");
 		check(CurveIds.Num() == CurveValues.Num());
 
-		AnimationProvider.AppendSkeletalMeshComponent(ComponentId, MeshId, Context.EventTime.AsSeconds(Cycle), LodIndex, FrameCounter, ComponentToWorld, Pose, CurveIds, CurveValues);
+		AnimationProvider.AppendSkeletalMeshComponent(ComponentId, MeshId, Context.EventTime.AsSeconds(Cycle), LodIndex, FrameCounter, ComponentToWorldFloatArray, PoseFloatArray, CurveIds, CurveValues);
 		break;
 	}
 	case RouteId_Name:
@@ -345,4 +319,4 @@ bool FAnimationAnalyzer::OnEvent(uint16 RouteId, EStyle Style, const FOnEventCon
 	}
 
 	return true;
-}
+}
\ No newline at end of file
diff --git a/Engine/Plugins/Animation/GameplayInsights/Source/GameplayInsights/Private/AnimationProvider.cpp b/Engine/Plugins/Animation/GameplayInsights/Source/GameplayInsights/Private/AnimationProvider.cpp
--- a/Engine/Plugins/Animation/GameplayInsights/Source/GameplayInsights/Private/AnimationProvider.cpp
+++ b/Engine/Plugins/Animation/GameplayInsights/Source/GameplayInsights/Private/AnimationProvider.cpp
@@ -534,7 +534,52 @@ void FAnimationProvider::AppendSkeletalMeshComponent(uint64 InObjectId, uint64 I
 	Session.UpdateDurationSeconds(InTime);
 }
 
-void FAnimationProvider::AppendSkeletalMeshComponent(uint64 InObjectId, uint64 InMeshId, double InTime, uint16 InLodIndex, uint16 InFrameCounter, const FTransform& InComponentToWorld, const TArrayView<const FTransform>& InPose, const TArrayView<const uint32>& InCurveIds, const TArrayView<const float>& InCurveValues)
+// support for deserializing LWC/non LWC Transforms from builds with non-matching LWC setting
+static FTransform ConvertTransform(int TransformSize, const float* TransformFloats)
+{
+	FQuat Rotation;
+	FVector Translation;
+	FVector Scale3D;
+
+	static const int LWCSize = 4 + 8 * 2; // 4 floats + 8 doubles
+	static const int NonLWCSize = 12; // 12 floats
+
+	check (TransformSize == NonLWCSize || TransformSize == LWCSize);
+
+	if (TransformSize == NonLWCSize)
+	{
+		// non LWC (rotation/translation/scale, all floats)
+		Rotation.X = TransformFloats[0];
+		Rotation.Y = TransformFloats[1];
+		Rotation.Z = TransformFloats[2];
+		Rotation.W = TransformFloats[3];
+		Translation.X = TransformFloats[4];
+		Translation.Y = TransformFloats[5];
+		Translation.Z = TransformFloats[6];
+		Scale3D.X = TransformFloats[8];
+		Scale3D.Y = TransformFloats[9];
+		Scale3D.Z = TransformFloats[10];
+	}
+	if (TransformSize == LWCSize)
+	{
+		// LWC (rotation in floats translation/scale in doubles)
+		Rotation.X = TransformFloats[0];
+		Rotation.Y = TransformFloats[1];
+		Rotation.Z = TransformFloats[2];
+		Rotation.W = TransformFloats[3];
+		const double* TransformDoubles = reinterpret_cast<const double*>(TransformFloats + 4);
+		Translation.X = TransformDoubles[0];
+		Translation.Y = TransformDoubles[1];
+		Translation.Z = TransformDoubles[2];
+		Scale3D.X = TransformDoubles[4];
+		Scale3D.Y = TransformDoubles[5];
+		Scale3D.Z = TransformDoubles[6];
+	}
+
+	return FTransform(Rotation, Translation, Scale3D);
+}
+
+void FAnimationProvider::AppendSkeletalMeshComponent(uint64 InObjectId, uint64 InMeshId, double InTime, uint16 InLodIndex, uint16 InFrameCounter, const TArrayView<const float>& InComponentToWorldRaw, const TArrayView<const float>& InPoseRaw, const TArrayView<const uint32>& InCurveIds, const TArrayView<const float>& InCurveValues)
 {
 	Session.WriteAccessCheck();
 
@@ -564,23 +609,49 @@ void FAnimationProvider::AppendSkeletalMeshComponent(uint64 InObjectId, uint64 I
 
 	const int32 NumCurves = InCurveIds.Num();
 
+	const int CaptureTransformSize = InComponentToWorldRaw.Num();
+	const int LocalTransformSize = sizeof (FTransform) / sizeof(float);
+
+	FTransform ComponentToWorld;
+
+	if (CaptureTransformSize == LocalTransformSize)
+	{
+	 	FMemory::Memcpy(&ComponentToWorld, &InComponentToWorldRaw[0], sizeof(FTransform));
+	}
+	else
+	{
+	 	ComponentToWorld = ConvertTransform(CaptureTransformSize, &InComponentToWorldRaw[0]);
+	}
+
+	const int PoseTransformCount = InPoseRaw.Num()/CaptureTransformSize;
+
 	FSkeletalMeshPoseMessage Message;
-	Message.ComponentToWorld = InComponentToWorld;
+	Message.ComponentToWorld = ComponentToWorld;
 	Message.TransformStartIndex = SkeletalMeshPoseTransforms.Num();
 	Message.CurveStartIndex = SkeletalMeshCurves.Num();
 	Message.ComponentId = InObjectId;
 	Message.MeshId = InMeshId;
 	Message.MeshName = GameplayProvider.GetObjectInfo(Message.MeshId).Name;
-	Message.NumTransforms = (uint16)InPose.Num();
+	Message.NumTransforms = (uint16)PoseTransformCount;
 	Message.NumCurves = (uint16)NumCurves;
 	Message.LodIndex = InLodIndex;
 	Message.FrameCounter = InFrameCounter;
 
 	TimelineStorage->Timeline->AppendBeginEvent(InTime, Message);
 
-	for(const FTransform& Transform : InPose)
+	if (CaptureTransformSize == LocalTransformSize)
 	{
-		SkeletalMeshPoseTransforms.PushBack() = Transform;
+		for(int i=0; i<PoseTransformCount; i++)
+		{
+			FMemory::Memcpy(&SkeletalMeshPoseTransforms.PushBack(), &InPoseRaw[CaptureTransformSize * i], sizeof(FTransform));
+		}
+	}
+	else
+	{
+		for(int i=0; i<PoseTransformCount; i++)
+		{
+			SkeletalMeshPoseTransforms.PushBack() = ConvertTransform(CaptureTransformSize, &InPoseRaw[CaptureTransformSize * i]);
+		}
 	}
 
 	for(int32 CurveIndex = 0; CurveIndex < NumCurves; ++CurveIndex)
diff --git a/Engine/Plugins/Animation/GameplayInsights/Source/GameplayInsights/Private/AnimationProvider.h b/Engine/Plugins/Animation/GameplayInsights/Source/GameplayInsights/Private/AnimationProvider.h
--- a/Engine/Plugins/Animation/GameplayInsights/Source/GameplayInsights/Private/AnimationProvider.h
+++ b/Engine/Plugins/Animation/GameplayInsights/Source/GameplayInsights/Private/AnimationProvider.h
@@ -56,7 +56,7 @@ public:
 
 	/** Add a skeletal mesh pose/curves etc. */
 	void AppendSkeletalMeshComponent(uint64 InObjectId, uint64 InMeshId, double InTime, uint16 InLodIndex, uint16 InFrameCounter, const TArrayView<const FTransform>& InPose, const TArrayView<const FSkeletalMeshNamedCurve>& InCurves);
-	void AppendSkeletalMeshComponent(uint64 InObjectId, uint64 InMeshId, double InTime, uint16 InLodIndex, uint16 InFrameCounter, const FTransform& InComponentToWorld, const TArrayView<const FTransform>& InPose, const TArrayView<const uint32>& InCurveIds, const TArrayView<const float>& InCurveValues);
+	void AppendSkeletalMeshComponent(uint64 InObjectId, uint64 InMeshId, double InTime, uint16 InLodIndex, uint16 InFrameCounter, const TArrayView<const float>& InComponentToWorldRaw, const TArrayView<const float>& InPoseRaw, const TArrayView<const uint32>& InCurveIds, const TArrayView<const float>& InCurveValues);
 
 	/** Get a skeletal mesh for a specified path. If the mesh no longer exists, make a fake one */
 	USkeletalMesh* GetSkeletalMesh(const TCHAR* InPath);
