Fix for warnings from static analysis, and non-unity build issue

diff --git a/Engine/Plugins/Animation/GameplayInsights/Source/GameplayInsights/Private/AnimationProvider.cpp b/Engine/Plugins/Animation/GameplayInsights/Source/GameplayInsights/Private/AnimationProvider.cpp
--- a/Engine/Plugins/Animation/GameplayInsights/Source/GameplayInsights/Private/AnimationProvider.cpp
+++ b/Engine/Plugins/Animation/GameplayInsights/Source/GameplayInsights/Private/AnimationProvider.cpp
@@ -567,7 +567,7 @@ static FTransform ConvertTransform(int TransformSize, const float* TransformFloa
 		Rotation.Y = TransformFloats[1];
 		Rotation.Z = TransformFloats[2];
 		Rotation.W = TransformFloats[3];
-		const double* TransformDoubles = reinterpret_cast<const double*>(TransformFloats + 4);
+		const double* TransformDoubles = reinterpret_cast<const double*>((void*)(TransformFloats + 4));
 		Translation.X = TransformDoubles[0];
 		Translation.Y = TransformDoubles[1];
 		Translation.Z = TransformDoubles[2];
diff --git a/Engine/Plugins/Animation/GameplayInsights/Source/GameplayInsights/Private/AnimationProvider.h b/Engine/Plugins/Animation/GameplayInsights/Source/GameplayInsights/Private/AnimationProvider.h
--- a/Engine/Plugins/Animation/GameplayInsights/Source/GameplayInsights/Private/AnimationProvider.h
+++ b/Engine/Plugins/Animation/GameplayInsights/Source/GameplayInsights/Private/AnimationProvider.h
@@ -2,6 +2,7 @@
 
 #pragma once
 
+#include "CoreMinimal.h"
 #include "IAnimationProvider.h"
 #include "Common/PagedArray.h"
 #include "Model/IntervalTimeline.h"
