Rewind debugger tabs implementation

diff --git a/Engine/Plugins/Animation/GameplayInsights/Source/GameplayInsights/Private/GameplayInsightsDebugViewCreator.cpp b/Engine/Plugins/Animation/GameplayInsights/Source/GameplayInsights/Private/GameplayInsightsDebugViewCreator.cpp
--- a/Engine/Plugins/Animation/GameplayInsights/Source/GameplayInsights/Private/GameplayInsightsDebugViewCreator.cpp
+++ b/Engine/Plugins/Animation/GameplayInsights/Source/GameplayInsights/Private/GameplayInsightsDebugViewCreator.cpp
@@ -8,12 +8,40 @@ namespace TraceServices
     class IAnalysisSession;
 }
 
-void FGameplayInsightsDebugViewCreator::RegisterDebugViewCreator(FName TypeName, FCreateDebugView Creator)
+void FGameplayInsightsDebugViewCreator::RegisterDebugViewCreator(FName TypeName, TSharedPtr<ICreateGameplayInsightsDebugView> Creator)
 {
-    ViewCreators.Add(TypeName, Creator);
+    FViewCreatorPair Pair;
+    Pair.Creator = Creator;
+    Pair.TypeName = TypeName;
+
+    ViewCreators.Add(Pair);
 }
 
-void FGameplayInsightsDebugViewCreator::CreateDebugViews(uint64 ObjectId, double CurrentTime, const TraceServices::IAnalysisSession& Session, TArray<TSharedPtr<IGameplayInsightsDebugView>>& OutDebugViews) 
+void FGameplayInsightsDebugViewCreator::EnumerateCreators(TFunctionRef<void(const TSharedPtr<ICreateGameplayInsightsDebugView>&)> Callback) const
+{
+    for(const FViewCreatorPair& CreatorPair : ViewCreators)
+    {
+        if (CreatorPair.Creator.IsValid())
+        {
+            Callback(CreatorPair.Creator);
+        }
+    }
+}
+
+TSharedPtr<ICreateGameplayInsightsDebugView> FGameplayInsightsDebugViewCreator::GetCreator(FName CreatorName) const
+{
+    for(const FViewCreatorPair& CreatorPair : ViewCreators)
+    {
+        if (CreatorPair.Creator.IsValid() && CreatorPair.Creator->GetName() == CreatorName)
+        {
+            return CreatorPair.Creator;
+        }
+    }
+
+    return nullptr;
+}
+
+void FGameplayInsightsDebugViewCreator::CreateDebugViews(uint64 ObjectId, double CurrentTime, const TraceServices::IAnalysisSession& Session, TArray<TSharedPtr<IGameplayInsightsDebugView>>& OutDebugViews) const
 {
     if (const IGameplayProvider* GameplayProvider = Session.ReadProvider<IGameplayProvider>("GameplayProvider"))
     {
@@ -26,14 +54,11 @@ void FGameplayInsightsDebugViewCreator::CreateDebugViews(uint64 ObjectId, double
         {
             const FClassInfo& ClassInfo = GameplayProvider->GetClassInfo(ClassId);
 
-            TArray<FCreateDebugView> List;
-            ViewCreators.MultiFind(ClassInfo.Name, List);
-
-            for (FCreateDebugView& Creator : List)
+            for(const FViewCreatorPair& CreatorPair : ViewCreators)
             {
-                if (Creator.IsBound())
+                if (CreatorPair.Creator.IsValid() && CreatorPair.TypeName == ClassInfo.Name)
                 {
-                    OutDebugViews.Add(Creator.Execute(ObjectId, CurrentTime, Session));
+                    OutDebugViews.Add(CreatorPair.Creator->CreateDebugView(ObjectId, CurrentTime, Session));
                 }
             }
 
diff --git a/Engine/Plugins/Animation/GameplayInsights/Source/GameplayInsights/Private/GameplayInsightsDebugViewCreator.h b/Engine/Plugins/Animation/GameplayInsights/Source/GameplayInsights/Private/GameplayInsightsDebugViewCreator.h
--- a/Engine/Plugins/Animation/GameplayInsights/Source/GameplayInsights/Private/GameplayInsightsDebugViewCreator.h
+++ b/Engine/Plugins/Animation/GameplayInsights/Source/GameplayInsights/Private/GameplayInsightsDebugViewCreator.h
@@ -11,9 +11,15 @@ class FGameplayInsightsDebugViewCreator : public IGameplayInsightsDebugViewCreat
 	public:
 		virtual ~FGameplayInsightsDebugViewCreator() {};
 
-		virtual void RegisterDebugViewCreator(FName TypeName, FCreateDebugView Creator) override;
-		virtual void CreateDebugViews(uint64 ObjectId, double CurrentTime,  const TraceServices::IAnalysisSession& InAnalysisSession, TArray<TSharedPtr<IGameplayInsightsDebugView>>& OutDebugViews) override;
-
+		virtual void RegisterDebugViewCreator(FName TypeName, TSharedPtr<ICreateGameplayInsightsDebugView> Creator) override;
+		virtual void CreateDebugViews(uint64 ObjectId, double CurrentTime,  const TraceServices::IAnalysisSession& InAnalysisSession, TArray<TSharedPtr<IGameplayInsightsDebugView>>& OutDebugViews)const override;
+		virtual void EnumerateCreators(TFunctionRef<void(const TSharedPtr<ICreateGameplayInsightsDebugView>&)> Callback) const;
+		virtual TSharedPtr<ICreateGameplayInsightsDebugView> GetCreator(FName CreatorName) const;
 	private:
-		TMultiMap<FName, FCreateDebugView> ViewCreators;
+		struct FViewCreatorPair
+		{
+			FName TypeName;
+			TSharedPtr<ICreateGameplayInsightsDebugView> Creator;
+		};
+		TArray<FViewCreatorPair> ViewCreators;
 };
\ No newline at end of file
diff --git a/Engine/Plugins/Animation/GameplayInsights/Source/GameplayInsights/Private/GameplayInsightsModule.cpp b/Engine/Plugins/Animation/GameplayInsights/Source/GameplayInsights/Private/GameplayInsightsModule.cpp
--- a/Engine/Plugins/Animation/GameplayInsights/Source/GameplayInsights/Private/GameplayInsightsModule.cpp
+++ b/Engine/Plugins/Animation/GameplayInsights/Source/GameplayInsights/Private/GameplayInsightsModule.cpp
@@ -62,20 +62,8 @@ void FGameplayInsightsModule::StartupModule()
 		}
 	});
 
-
-	DebugViewCreator.RegisterDebugViewCreator("Object", IGameplayInsightsDebugViewCreator::FCreateDebugView::CreateLambda(
-			[](uint64 InObjectId, double InTimeMarker, const TraceServices::IAnalysisSession& InAnalysisSession)
-			{
-				return SNew(SObjectPropertiesView, InObjectId, InTimeMarker, InAnalysisSession);
-			}
-	));
-
-	DebugViewCreator.RegisterDebugViewCreator("AnimInstance", IGameplayInsightsDebugViewCreator::FCreateDebugView::CreateLambda(
-			[](uint64 InAnimInstanceId, double InTimeMarker, const TraceServices::IAnalysisSession& InAnalysisSession)
-			{
-				return SNew(SAnimGraphSchematicView, InAnimInstanceId, InTimeMarker, InAnalysisSession);
-			}
-	));
+	DebugViewCreator.RegisterDebugViewCreator("Object", MakeShared<FObjectPropertiesViewCreator>());
+	DebugViewCreator.RegisterDebugViewCreator("AnimInstance", MakeShared<FAnimGraphSchematicViewCreator>());
 
 #if WITH_EDITOR
 
diff --git a/Engine/Plugins/Animation/GameplayInsights/Source/GameplayInsights/Private/SAnimGraphSchematicView.cpp b/Engine/Plugins/Animation/GameplayInsights/Source/GameplayInsights/Private/SAnimGraphSchematicView.cpp
--- a/Engine/Plugins/Animation/GameplayInsights/Source/GameplayInsights/Private/SAnimGraphSchematicView.cpp
+++ b/Engine/Plugins/Animation/GameplayInsights/Source/GameplayInsights/Private/SAnimGraphSchematicView.cpp
@@ -1,6 +1,7 @@
 // Copyright Epic Games, Inc. All Rights Reserved.
 
 #include "SAnimGraphSchematicView.h"
+#include "Animation/AnimInstance.h"
 #include "AnimationProvider.h"
 #include "Widgets/Input/SSearchBox.h"
 #include "TraceServices/Model/AnalysisSession.h"
@@ -9,6 +10,7 @@
 #include "Widgets/Layout/SBorder.h"
 #include "Widgets/SBoxPanel.h"
 #include "Styling/CoreStyle.h"
+#include "Styling/SlateIconFinder.h"
 #include "GameplayInsightsStyle.h"
 #include "Widgets/Layout/SScrollBorder.h"
 #include "Framework/MultiBox/MultiBoxBuilder.h"
@@ -25,6 +27,8 @@
 
 #define LOCTEXT_NAMESPACE "SAnimGraphSchematicView"
 
+static const FName AnimGraphSchematicName("AnimGraphSchematic");
+
 namespace AnimGraphSchematicPropertyColumns
 {
 	static const FName Name("Name");
@@ -822,9 +826,10 @@ void SAnimGraphSchematicView::SetTimeMarker(double InTimeMarker)
 	}
 }
 
-FText SAnimGraphSchematicView::GetTitle()
+
+FName SAnimGraphSchematicView::GetName() const
 {
-	return LOCTEXT("Anim Graph Update", "Anim Graph Update");
+	return AnimGraphSchematicName;
 }
 
 TSharedRef<SWidget> SAnimGraphSchematicView::HandleGetViewMenuContent()
@@ -998,4 +1003,24 @@ void SAnimGraphSchematicView::RefreshDetails(const TArray<TSharedRef<FAnimGraphS
 	}
 }
 
+FName FAnimGraphSchematicViewCreator::GetName() const
+{
+	return AnimGraphSchematicName;
+}
+
+FText FAnimGraphSchematicViewCreator::GetTitle() const
+{
+	return LOCTEXT("Anim Graph Update", "Anim Graph");
+}
+
+FSlateIcon FAnimGraphSchematicViewCreator::GetIcon() const
+{
+	return FSlateIconFinder::FindIconForClass(UAnimInstance::StaticClass());
+}
+
+TSharedPtr<IGameplayInsightsDebugView> FAnimGraphSchematicViewCreator::CreateDebugView(uint64 ObjectId, double CurrentTime, const TraceServices::IAnalysisSession& AnalysisSession) const
+{
+	return SNew(SAnimGraphSchematicView, ObjectId, CurrentTime, AnalysisSession);
+}
+
 #undef LOCTEXT_NAMESPACE
diff --git a/Engine/Plugins/Animation/GameplayInsights/Source/GameplayInsights/Private/SAnimGraphSchematicView.h b/Engine/Plugins/Animation/GameplayInsights/Source/GameplayInsights/Private/SAnimGraphSchematicView.h
--- a/Engine/Plugins/Animation/GameplayInsights/Source/GameplayInsights/Private/SAnimGraphSchematicView.h
+++ b/Engine/Plugins/Animation/GameplayInsights/Source/GameplayInsights/Private/SAnimGraphSchematicView.h
@@ -4,6 +4,7 @@
 
 #include "CoreMinimal.h"
 #include "IGameplayInsightsDebugView.h"
+#include "IGameplayInsightsDebugViewCreator.h"
 #include "Widgets/SCompoundWidget.h"
 #include "Widgets/Views/STreeView.h"
 
@@ -30,7 +31,9 @@ public:
 	uint64 GetAnimInstanceId() const { return AnimInstanceId; }
 
 	virtual void SetTimeMarker(double InTimeMarker) override;
-	virtual FText GetTitle() override;
+	virtual FName GetName() const override;
+	virtual uint64 GetObjectId() const override { return AnimInstanceId; }
+
 private:
 	// Generate a row widget for an item
 	TSharedRef<ITableRow> HandleGenerateRow(TSharedRef<FAnimGraphSchematicNode> Item, const TSharedRef<STableViewBase>& OwnerTable);
@@ -112,3 +115,12 @@ private:
 
 	uint64 AnimInstanceId;
 };
+
+class FAnimGraphSchematicViewCreator : public ICreateGameplayInsightsDebugView
+{
+	public:
+		virtual FName GetName() const override;
+		virtual FText GetTitle() const override;
+		virtual FSlateIcon GetIcon() const override;
+		virtual TSharedPtr<IGameplayInsightsDebugView> CreateDebugView(uint64 ObjectId, double CurrentTime, const TraceServices::IAnalysisSession& InAnalysisSession) const override;
+};
\ No newline at end of file
diff --git a/Engine/Plugins/Animation/GameplayInsights/Source/GameplayInsights/Private/SObjectPropertiesView.cpp b/Engine/Plugins/Animation/GameplayInsights/Source/GameplayInsights/Private/SObjectPropertiesView.cpp
--- a/Engine/Plugins/Animation/GameplayInsights/Source/GameplayInsights/Private/SObjectPropertiesView.cpp
+++ b/Engine/Plugins/Animation/GameplayInsights/Source/GameplayInsights/Private/SObjectPropertiesView.cpp
@@ -4,6 +4,7 @@
 #include "AnimationProvider.h"
 #include "TraceServices/Model/AnalysisSession.h"
 #include "GameplayProvider.h"
+#include "Styling/SlateIconFinder.h"
 #include "TraceServices/Model/Frames.h"
 #include "VariantTreeNode.h"
 
@@ -87,12 +88,33 @@ void SObjectPropertiesView::GetVariantsAtFrame(const TraceServices::FFrame& InFr
 	}
 }
 
-FText SObjectPropertiesView::GetTitle()
+
+static const FName ObjectPropertiesName("ObjectProperties");
+
+FName SObjectPropertiesView::GetName() const
+{
+	return ObjectPropertiesName;
+}
+
+FName FObjectPropertiesViewCreator::GetName() const
+{
+	return ObjectPropertiesName;
+}
+
+FText FObjectPropertiesViewCreator::GetTitle() const
 {
-	return LOCTEXT("Object Properties", "Object Properties");
+	return LOCTEXT("Object Properties", "Properties");
 }
 
+FSlateIcon FObjectPropertiesViewCreator::GetIcon() const
+{
+	return FSlateIconFinder::FindIconForClass(UObject::StaticClass());
+}
 
+TSharedPtr<IGameplayInsightsDebugView> FObjectPropertiesViewCreator::CreateDebugView(uint64 ObjectId, double CurrentTime, const TraceServices::IAnalysisSession& AnalysisSession) const
+{
+	return SNew(SObjectPropertiesView, ObjectId, CurrentTime, AnalysisSession);
+}
 
 
 #undef LOCTEXT_NAMESPACE
diff --git a/Engine/Plugins/Animation/GameplayInsights/Source/GameplayInsights/Private/SObjectPropertiesView.h b/Engine/Plugins/Animation/GameplayInsights/Source/GameplayInsights/Private/SObjectPropertiesView.h
--- a/Engine/Plugins/Animation/GameplayInsights/Source/GameplayInsights/Private/SObjectPropertiesView.h
+++ b/Engine/Plugins/Animation/GameplayInsights/Source/GameplayInsights/Private/SObjectPropertiesView.h
@@ -4,6 +4,7 @@
 
 #include "CoreMinimal.h"
 #include "IGameplayInsightsDebugView.h"
+#include "IGameplayInsightsDebugViewCreator.h"
 #include "SVariantValueView.h"
 
 namespace TraceServices { class IAnalysisSession; }
@@ -16,7 +17,8 @@ public:
 	void Construct(const FArguments& InArgs, uint64 InObjectId, double InTimeMarker, const TraceServices::IAnalysisSession& InAnalysisSession);
 
 	virtual void SetTimeMarker(double InTimeMarker) override;
-	virtual FText GetTitle() override;
+	virtual FName GetName() const override;
+	virtual uint64 GetObjectId() const override { return ObjectId; }
 
 	void GetVariantsAtFrame(const TraceServices::FFrame& InFrame, TArray<TSharedRef<FVariantTreeNode>>& OutVariants) const;
 private:
@@ -26,3 +28,12 @@ private:
 	double TimeMarker;
 	const TraceServices::IAnalysisSession* AnalysisSession;
 };
+
+class FObjectPropertiesViewCreator : public ICreateGameplayInsightsDebugView
+{
+	public:
+		virtual FName GetName() const override;
+		virtual FText GetTitle() const override;
+		virtual FSlateIcon GetIcon() const override;
+		virtual TSharedPtr<IGameplayInsightsDebugView> CreateDebugView(uint64 ObjectId, double CurrentTime, const TraceServices::IAnalysisSession& InAnalysisSession) const override;
+};
\ No newline at end of file
diff --git a/Engine/Plugins/Animation/GameplayInsights/Source/GameplayInsights/Public/IGameplayInsightsDebugView.h b/Engine/Plugins/Animation/GameplayInsights/Source/GameplayInsights/Public/IGameplayInsightsDebugView.h
--- a/Engine/Plugins/Animation/GameplayInsights/Source/GameplayInsights/Public/IGameplayInsightsDebugView.h
+++ b/Engine/Plugins/Animation/GameplayInsights/Source/GameplayInsights/Public/IGameplayInsightsDebugView.h
@@ -11,8 +11,11 @@
 class IGameplayInsightsDebugView : public SCompoundWidget
 {
 	public:
-		// text for widget header
-		virtual FText GetTitle() = 0;
+		// unique name for widget
+		virtual FName GetName() const = 0;
+
+		// id of target object
+		virtual uint64 GetObjectId() const = 0;
 
 		// called by the debugger when the scrubbing bar position changes
 		virtual void SetTimeMarker(double InTimeMarker) = 0;
diff --git a/Engine/Plugins/Animation/GameplayInsights/Source/GameplayInsights/Public/IGameplayInsightsDebugViewCreator.h b/Engine/Plugins/Animation/GameplayInsights/Source/GameplayInsights/Public/IGameplayInsightsDebugViewCreator.h
--- a/Engine/Plugins/Animation/GameplayInsights/Source/GameplayInsights/Public/IGameplayInsightsDebugViewCreator.h
+++ b/Engine/Plugins/Animation/GameplayInsights/Source/GameplayInsights/Public/IGameplayInsightsDebugViewCreator.h
@@ -4,19 +4,44 @@
 
 #include "CoreMinimal.h"
 #include "IGameplayInsightsDebugView.h"
+#include "Textures/SlateIcon.h"
 
 namespace TraceServices
 {
 	class IAnalysisSession;
 }
 
+// Interface class which creates debug widgets
+class ICreateGameplayInsightsDebugView
+{
+	public:
+		// returns a unique name for identifying this type of widget (same value returned by IGameplayInsightsDebugView::GetName)
+		virtual FName GetName() const = 0;
+
+		// text for tab header
+		virtual FText GetTitle() const = 0;
+
+		// icon for tab header
+		virtual FSlateIcon GetIcon() const = 0;
+
+		// creates and returns a widget, which will be displayed in Rewind Debugger
+		virtual TSharedPtr<IGameplayInsightsDebugView> CreateDebugView(uint64 ObjectId, double CurrentTime, const TraceServices::IAnalysisSession& InAnalysisSession) const = 0;
+};
+
 // this class handles creating debug view widgets for the Rewind Debugger.
-// systems can register a debug view creator function with a UObject type name, and when an object of that type is selected, that widget will be created and shown by the debugger.
+// systems can register an ICreateGameplayInsightsDebugView implementation with a UObject type name, and when an object of that type is selected, that widget will be created and shown by the debugger.
 class IGameplayInsightsDebugViewCreator
 {
 	public:
-		DECLARE_DELEGATE_RetVal_ThreeParams(TSharedRef<IGameplayInsightsDebugView>, FCreateDebugView, uint64 /*ObjectId*/, double /*Time*/, const TraceServices::IAnalysisSession& InAnalysisSession);
+		// Register a creator for a Type name
+		virtual void RegisterDebugViewCreator(FName TypeName, TSharedPtr<ICreateGameplayInsightsDebugView> Creator) = 0;
+
+		// Create all views for a object Id based on it's type heirarchy
+		virtual void CreateDebugViews(uint64 ObjectId, double CurrentTime, const TraceServices::IAnalysisSession& InAnalysisSession, TArray<TSharedPtr<IGameplayInsightsDebugView>>& OutDebugViews) const = 0;
+
+		// Iterate over all registered creators and call a callback for each one
+		virtual void EnumerateCreators(TFunctionRef<void(const TSharedPtr<ICreateGameplayInsightsDebugView>&)> Callback) const = 0;
 
-		virtual void RegisterDebugViewCreator(FName TypeName, FCreateDebugView Creator) = 0;
-		virtual void CreateDebugViews(uint64 ObjectId, double CurrentTime, const TraceServices::IAnalysisSession& InAnalysisSession, TArray<TSharedPtr<IGameplayInsightsDebugView>>& OutDebugViews) = 0;
+		// Get a creator by its unique name
+		virtual TSharedPtr<ICreateGameplayInsightsDebugView> GetCreator(FName CreatorName) const = 0;
 };
\ No newline at end of file
diff --git a/Engine/Plugins/Animation/GameplayInsights/Source/RewindDebugger/Private/RewindDebuggerStyle.cpp b/Engine/Plugins/Animation/GameplayInsights/Source/RewindDebugger/Private/RewindDebuggerStyle.cpp
--- a/Engine/Plugins/Animation/GameplayInsights/Source/RewindDebugger/Private/RewindDebuggerStyle.cpp
+++ b/Engine/Plugins/Animation/GameplayInsights/Source/RewindDebugger/Private/RewindDebuggerStyle.cpp
@@ -18,24 +18,27 @@ FRewindDebuggerStyle::FRewindDebuggerStyle() :
 	const FVector2D Icon16x16(16.0f, 16.0f);
 
 	SetContentRoot(FPaths::EnginePluginsDir() / TEXT("Animation/GameplayInsights/Content"));
-	SetCoreContentRoot(FPaths::EngineContentDir() / TEXT("Editor/Slate"));
+	SetCoreContentRoot(FPaths::EngineContentDir() / TEXT(""));
 
 	// Toolbar
-	Set("RewindDebugger.StartRecording.small", new CORE_IMAGE_BRUSH("Sequencer/Transport_Bar/Record_24x", Icon24x24));
-	Set("RewindDebugger.StopRecording.small", new CORE_IMAGE_BRUSH("Sequencer/Transport_Bar/Recording_24x", Icon24x24));
-	Set("RewindDebugger.FirstFrame.small", new CORE_IMAGE_BRUSH("Sequencer/Transport_Bar/Go_To_Front_24x", Icon24x24));
-	Set("RewindDebugger.PreviousFrame.small", new CORE_IMAGE_BRUSH("Sequencer/Transport_Bar/Step_Backwards_24x", Icon24x24));
-	Set("RewindDebugger.ReversePlay.small", new CORE_IMAGE_BRUSH("Sequencer/Transport_Bar/Backwards_24x", Icon24x24));
-	Set("RewindDebugger.Pause.small", new CORE_IMAGE_BRUSH("Sequencer/Transport_Bar/Pause_24x", Icon24x24));
-	Set("RewindDebugger.Play.small", new CORE_IMAGE_BRUSH("Sequencer/Transport_Bar/Play_24x", Icon24x24));
-	Set("RewindDebugger.NextFrame.small", new CORE_IMAGE_BRUSH("Sequencer/Transport_Bar/Step_Forward_24x", Icon24x24));
-	Set("RewindDebugger.LastFrame.small", new CORE_IMAGE_BRUSH("Sequencer/Transport_Bar/Go_To_End_24x", Icon24x24));
+	Set("RewindDebugger.StartRecording.small", new CORE_IMAGE_BRUSH("Editor/Slate/Sequencer/Transport_Bar/Record_24x", Icon24x24));
+	Set("RewindDebugger.StopRecording.small", new CORE_IMAGE_BRUSH("Editor/Slate/Sequencer/Transport_Bar/Recording_24x", Icon24x24));
+	Set("RewindDebugger.FirstFrame.small", new CORE_IMAGE_BRUSH("Editor/Slate/Sequencer/Transport_Bar/Go_To_Front_24x", Icon24x24));
+	Set("RewindDebugger.PreviousFrame.small", new CORE_IMAGE_BRUSH("Editor/Slate/Sequencer/Transport_Bar/Step_Backwards_24x", Icon24x24));
+	Set("RewindDebugger.ReversePlay.small", new CORE_IMAGE_BRUSH("Editor/Slate/Sequencer/Transport_Bar/Backwards_24x", Icon24x24));
+	Set("RewindDebugger.Pause.small", new CORE_IMAGE_BRUSH("Editor/Slate/Sequencer/Transport_Bar/Pause_24x", Icon24x24));
+	Set("RewindDebugger.Play.small", new CORE_IMAGE_BRUSH("Editor/Slate/Sequencer/Transport_Bar/Play_24x", Icon24x24));
+	Set("RewindDebugger.NextFrame.small", new CORE_IMAGE_BRUSH("Editor/Slate/Sequencer/Transport_Bar/Step_Forward_24x", Icon24x24));
+	Set("RewindDebugger.LastFrame.small", new CORE_IMAGE_BRUSH("Editor/Slate/Sequencer/Transport_Bar/Go_To_End_24x", Icon24x24));
 
 	// Actor Picker
-	Set("RewindDebugger.SelectActor", new CORE_IMAGE_BRUSH("Icons/eyedropper_16px", Icon16x16));
+	Set("RewindDebugger.SelectActor", new CORE_IMAGE_BRUSH("Editor/Slate/Icons/eyedropper_16px", Icon16x16));
 
 	// tab icon
 	Set("RewindDebugger.RewindIcon", new IMAGE_BRUSH("Rewind_24x", Icon16x16));
+
+	// menu icon
+	Set("RewindDebugger.MenuIcon", new CORE_IMAGE_BRUSH_SVG("Slate/Starship/Common/menu", Icon16x16));
 }
 
 void FRewindDebuggerStyle::Initialize()
diff --git a/Engine/Plugins/Animation/GameplayInsights/Source/RewindDebugger/Private/SRewindDebugger.cpp b/Engine/Plugins/Animation/GameplayInsights/Source/RewindDebugger/Private/SRewindDebugger.cpp
--- a/Engine/Plugins/Animation/GameplayInsights/Source/RewindDebugger/Private/SRewindDebugger.cpp
+++ b/Engine/Plugins/Animation/GameplayInsights/Source/RewindDebugger/Private/SRewindDebugger.cpp
@@ -7,6 +7,7 @@
 #include "Editor.h"
 #include "Editor/EditorEngine.h"
 #include "Framework/Application/SlateApplication.h"
+#include "Framework/Docking/LayoutService.h"
 #include "Framework/MultiBox/MultiBoxBuilder.h"
 #include "Insights/IUnrealInsightsModule.h"
 #include "IGameplayInsightsModule.h"
@@ -20,6 +21,7 @@
 #include "SceneOutlinerModule.h"
 #include "Selection.h"
 #include "Styling/SlateIconFinder.h"
+#include "Widgets/Docking/SDockTab.h"
 #include "Widgets/Images/SImage.h"
 #include "Widgets/Input/SComboButton.h"
 #include "Widgets/Input/SButton.h"
@@ -27,8 +29,7 @@
 #include "Widgets/Layout/SExpandableArea.h"
 #include "Widgets/Layout/SScrollBox.h"
 
-
-#define LOCTEXT_NAMESPACE "SAnimationInsights"
+#define LOCTEXT_NAMESPACE "SRewindDebugger"
 
 SRewindDebugger::SRewindDebugger() 
 	: SCompoundWidget()
@@ -38,7 +39,6 @@ SRewindDebugger::SRewindDebugger()
 
 SRewindDebugger::~SRewindDebugger() 
 {
-
 }
 
 void SRewindDebugger::TrackCursor(bool bReverse)
@@ -238,19 +238,80 @@ void SRewindDebugger::Construct(const FArguments& InArgs, TSharedRef<FUICommandL
 				.OnSelectionChanged(this, &SRewindDebugger::OnComponentSelectionChanged);
 
 
-	DebugViewContainer = SNew(SScrollBox);
-
 	TraceTime.OnPropertyChanged = TraceTime.OnPropertyChanged.CreateRaw(this, &SRewindDebugger::TraceTimeChanged);
 
+	// Tab Manager 
+	TabManager = FGlobalTabmanager::Get()->NewTabManager(ConstructUnderMajorTab);
+
+	TabManager->SetOnPersistLayout(
+		FTabManager::FOnPersistLayout::CreateStatic(
+			[](const TSharedRef<FTabManager::FLayout>& InLayout)
+			{
+				if (InLayout->GetPrimaryArea().Pin().IsValid())
+				{
+					FLayoutSaveRestore::SaveToConfig(GEditorLayoutIni, InLayout);
+				}
+			}
+		)
+	);
+
+	// Default Layout: all tabs in one stack, inside the rewind debugger tab
+
+	TSharedPtr<FTabManager::FStack> MainTabStack = FTabManager::NewStack();
+
+	IGameplayInsightsModule& GameplayInsightsModule = FModuleManager::LoadModuleChecked<IGameplayInsightsModule>("GameplayInsights");
+	GameplayInsightsModule.GetDebugViewCreator()->EnumerateCreators([this, MainTabStack](const TSharedPtr<ICreateGameplayInsightsDebugView>& Creator)
+		{
+			FName TabName = Creator->GetName();
+			TabNames.Add(TabName);
+			// Add a closed tabs to the main tab stack in the default layout, so that the first time they won't pop up in their own window
+			TabManager->RegisterTabSpawner(TabName, FOnSpawnTab::CreateRaw(this, &SRewindDebugger::SpawnTab, TabName),
+													FCanSpawnTab::CreateRaw(this, &SRewindDebugger::CanSpawnTab, TabName))
+				.SetDisplayName(Creator->GetTitle())
+				.SetIcon(Creator->GetIcon());
+
+			MainTabStack->AddTab(TabName, ETabState::ClosedTab);
+		}
+	);
+
+	TSharedRef<FTabManager::FLayout> Layout = FTabManager::NewLayout("RewindDebuggerLayout") 
+		->AddArea
+		(
+			FTabManager::NewPrimaryArea()
+			->Split
+			(
+				MainTabStack.ToSharedRef()
+			)
+		);
+
+	// load saved layout if it exists
+	Layout = FLayoutSaveRestore::LoadFromConfig(GEditorLayoutIni, Layout);
+
+
 	ChildSlot
 	[
 		SNew(SSplitter)
-		+SSplitter::Slot().MinSize(270).Value(0)
+		+SSplitter::Slot().MinSize(280).Value(0)
 		[
 			SNew(SVerticalBox)
 			+SVerticalBox::Slot().AutoHeight()
 			[
-				ToolBarBuilder.MakeWidget()
+				SNew(SHorizontalBox)
+				+SHorizontalBox::Slot().AutoWidth()
+				[
+					SNew(SComboButton)
+						.ComboButtonStyle(FAppStyle::Get(), "SimpleComboButton")
+						.OnGetMenuContent(this, &SRewindDebugger::MakeMainMenu)
+						.ButtonContent()
+					[
+						SNew(SImage)
+						.Image(FRewindDebuggerStyle::Get().GetBrush("RewindDebugger.MenuIcon"))
+					]
+				]
+				+SHorizontalBox::Slot().AutoWidth()
+				[
+					ToolBarBuilder.MakeWidget()
+				]
 			]
 			+SVerticalBox::Slot().AutoHeight()
 			[
@@ -315,7 +376,7 @@ void SRewindDebugger::Construct(const FArguments& InArgs, TSharedRef<FUICommandL
 					]
 				]
 			]
-			+SVerticalBox::Slot().VAlign(VAlign_Top) .FillHeight(1.0f)
+			+SVerticalBox::Slot().FillHeight(1.0f)
 			[
 				ComponentTreeView.ToSharedRef()
 			]
@@ -323,16 +384,22 @@ void SRewindDebugger::Construct(const FArguments& InArgs, TSharedRef<FUICommandL
 		+SSplitter::Slot() 
 		[
 			SNew(SVerticalBox)
-			+SVerticalBox::Slot() .VAlign(VAlign_Top) .AutoHeight()
+			+SVerticalBox::Slot().AutoHeight()
 			[
 				TimeSlider
 			]
-			+SVerticalBox::Slot() .VAlign(VAlign_Top) .FillHeight(1.0f)
+			+SVerticalBox::Slot().FillHeight(1.0f)
 			[
-				DebugViewContainer.ToSharedRef()
+				TabManager->RestoreFrom(Layout, ConstructUnderWindow).ToSharedRef()
 			]
 		]
 	];
+
+	// close all tabs that may be open from restoring the saved layout config
+	for(FName& TabName : TabNames)
+	{
+		CloseTab(TabName);
+	}
 }
 
 void SRewindDebugger::RefreshDebugComponents()
@@ -346,42 +413,207 @@ void SRewindDebugger::TraceTimeChanged(double Time)
 	{
 		DebugView->SetTimeMarker(Time);
 	}
+	for(TSharedPtr<IGameplayInsightsDebugView>& DebugView : PinnedDebugViews)
+	{
+		DebugView->SetTimeMarker(Time);
+	}
 }
 
+TSharedRef<SDockTab> SRewindDebugger::SpawnTab(const FSpawnTabArgs& Args, FName ViewName)
+{
+	TSharedPtr<IGameplayInsightsDebugView>* View = DebugViews.FindByPredicate([ViewName](TSharedPtr<IGameplayInsightsDebugView>& View) { return View->GetName() == ViewName; } );
+	if (View)
+	{
+		HiddenTabs.Remove(ViewName);
 
-void SRewindDebugger::OnComponentSelectionChanged(TSharedPtr<FDebugObjectInfo> SelectedItem, ESelectInfo::Type SelectInfo)
+		return SNew(SDockTab)
+		[
+			View->ToSharedRef()
+		]
+		.OnExtendContextMenu(this, &SRewindDebugger::ExtendTabMenu, *View)
+		.OnTabClosed_Lambda([this,ViewName](TSharedRef<SDockTab>)
+		{
+			if(!bInternalClosingTab) // skip this if the tab is being closed by our own code
+			{
+				HiddenTabs.Add(ViewName);
+			}
+		});
+	}
+
+	return SNew(SDockTab);
+}
+
+// returns true if DebugViews contains a view for the ViewName, but there is no matching Pinned view already open
+bool SRewindDebugger::CanSpawnTab(const FSpawnTabArgs& Args, FName ViewName)
+{
+	TSharedPtr<IGameplayInsightsDebugView>* View = DebugViews.FindByPredicate([ViewName](TSharedPtr<IGameplayInsightsDebugView>& View) { return View->GetName() == ViewName; } );
+
+	if (View!=nullptr)
+	{
+		bool bPinned = nullptr != PinnedDebugViews.FindByPredicate(
+			[View](TSharedPtr<IGameplayInsightsDebugView>& PinnedView)
+			{
+				return PinnedView->GetName() == (*View)->GetName() && PinnedView->GetObjectId() == (*View)->GetObjectId();
+			}
+		);
+
+		return !bPinned;
+	}
+	return false;
+}
+
+void SRewindDebugger::OnPinnedTabClosed(TSharedRef<SDockTab> Tab)
+{
+	// remove view from list of pinned views
+	TSharedRef<IGameplayInsightsDebugView> View = StaticCastSharedRef<IGameplayInsightsDebugView>(Tab->GetContent());
+	PinnedDebugViews.Remove(View);
+
+	// recreate non-pinned tabs, so when closing a pinned tab for the currently selected component, the non-pinned one will appear
+	CreateDebugTabs();
+}
+
+ void SRewindDebugger::PinTab(TSharedPtr<IGameplayInsightsDebugView> View)
+ {
+	if (PinnedDebugViews.Find(View) != INDEX_NONE)
+	{
+		return;
+	}
+
+	FName TabName = View->GetName();
+
+	CloseTab(TabName);
+
+	FSlateIcon TabIcon;
+	FText TabLabel;
+
+	IGameplayInsightsModule& GameplayInsightsModule = FModuleManager::LoadModuleChecked<IGameplayInsightsModule>("GameplayInsights");
+	TSharedPtr<ICreateGameplayInsightsDebugView> Creator = GameplayInsightsModule.GetDebugViewCreator()->GetCreator(TabName);
+	if (Creator.IsValid())
+	{
+		TabIcon = Creator->GetIcon();
+		TabLabel = Creator->GetTitle();
+	}
+
+	TSharedPtr<SDockTab> NewTab = SNew(SDockTab)
+	[
+		// add a wrapper widget here, that says the name of the object/component for pinned tabs 
+		View.ToSharedRef()
+	]
+	.OnExtendContextMenu(this, &SRewindDebugger::ExtendTabMenu, View)
+	.OnTabClosed(this, &SRewindDebugger::OnPinnedTabClosed)
+	.Label(TabLabel)
+	.LabelSuffix(LOCTEXT(" (Locked)", " \xD83D\xDD12")); // unicode lock image
+
+	NewTab->SetTabIcon(TabIcon.GetIcon());
+
+
+	static const FName RewindDebuggerPinnedTab = "RewindDebuggerPinnedTabName";
+	TabManager->InsertNewDocumentTab(TabName, RewindDebuggerPinnedTab, FTabManager::FRequireClosedTab(), NewTab.ToSharedRef());
+
+	PinnedDebugViews.Add(View);
+ }
+
+TSharedRef<SWidget> SRewindDebugger::MakeMainMenu()
+{
+	FMenuBuilder MenuBuilder(true, nullptr);
+	MakeViewsMenu(MenuBuilder);
+	return MenuBuilder.MakeWidget();
+}
+
+void SRewindDebugger::MakeViewsMenu(FMenuBuilder& MenuBuilder)
+{
+	MenuBuilder.BeginSection("Views", LOCTEXT("Views", "Views"));
+
+	TabManager->PopulateLocalTabSpawnerMenu(MenuBuilder);
+
+	MenuBuilder.AddMenuEntry(LOCTEXT("Show All Views", "Show All Views"),
+							 LOCTEXT("Show All tooltip", "Show all debug views that are relevant to the selected object type"),
+							 FSlateIcon(),
+							 FExecuteAction::CreateLambda([this]() { ShowAllViews(); }));
+
+	MenuBuilder.EndSection();
+}
+
+void SRewindDebugger::ExtendTabMenu(FMenuBuilder& MenuBuilder, TSharedPtr<IGameplayInsightsDebugView> View)
+{
+	MenuBuilder.BeginSection("RewindDebugger", LOCTEXT("Rewind Debugger", "Rewind Debugger"));
+
+	MenuBuilder.AddMenuEntry(LOCTEXT("Keep View Open", "Keep View Open"),
+							 LOCTEXT("Keep View Open tooltip", "Keep this debug view open even while selected component/actor changes"),
+							 FSlateIcon(),
+							 FUIAction(
+								FExecuteAction::CreateRaw(this, &SRewindDebugger::PinTab, View),
+								FCanExecuteAction::CreateLambda([this,View](){ return PinnedDebugViews.Find(View) == INDEX_NONE; })
+							 )); 
+
+
+	MenuBuilder.EndSection();
+
+	MakeViewsMenu(MenuBuilder);
+}
+
+void SRewindDebugger::ShowAllViews()
+{
+	HiddenTabs.Empty();
+	CreateDebugTabs();
+}
+
+void SRewindDebugger::CreateDebugViews()
 {
-	DebugViewContainer->ClearChildren();
 	DebugViews.Empty();
 
-	if (SelectedItem.IsValid())
+	if (SelectedComponent.IsValid())
 	{
-		// everything related to creating SAnimGraphSchematicView should be moved to a separate file
 		IUnrealInsightsModule &UnrealInsightsModule = FModuleManager::LoadModuleChecked<IUnrealInsightsModule>("TraceInsights");
 		IGameplayInsightsModule& GameplayInsightsModule = FModuleManager::LoadModuleChecked<IGameplayInsightsModule>("GameplayInsights");
 		TSharedPtr<const TraceServices::IAnalysisSession> Session = UnrealInsightsModule.GetAnalysisSession();
 
-		GameplayInsightsModule.GetDebugViewCreator()->CreateDebugViews(SelectedItem->ObjectId, TraceTime.Get(), *Session, DebugViews);
+		GameplayInsightsModule.GetDebugViewCreator()->CreateDebugViews(SelectedComponent->ObjectId, TraceTime.Get(), *Session, DebugViews);
+	}
+}
+
+void SRewindDebugger::CloseTab(FName TabName)
+{
+	// using bInternalClosingTab to distinguish between procedural, and user initiated tab closing in the OnTabClosed callback
+	bInternalClosingTab = true;
+	TSharedPtr<SDockTab> Tab = TabManager->FindExistingLiveTab(TabName);
+	if (Tab.IsValid())
+	{
+		Tab->RequestCloseTab();
+	}
+	bInternalClosingTab = false;
+}
+
+void SRewindDebugger::CreateDebugTabs()
+{
+	for(FName& TabName : TabNames)
+	{
+		CloseTab(TabName);
+	}
+
+	for(TSharedPtr<IGameplayInsightsDebugView> DebugView : DebugViews)
+	{
+		FName ViewName = DebugView->GetName();
+		uint64 ObjectId = DebugView->GetObjectId();
+
+	    bool bPinned = nullptr != PinnedDebugViews.FindByPredicate([ViewName, ObjectId](TSharedPtr<IGameplayInsightsDebugView>& View) { return View->GetName() == ViewName && View->GetObjectId() == ObjectId; } );
+	    bool bHidden = HiddenTabs.Find(ViewName) != INDEX_NONE;
 
-		for(TSharedPtr<IGameplayInsightsDebugView>& DebugView : DebugViews)
+		if(!bPinned && !bHidden)
 		{
-			// for now, just add all the widgets to a vertical box.
-			// this will be replaced by something better once there are more DebugViews implemented
-			DebugViewContainer->AddSlot()
-			[
-				SNew(SExpandableArea)
-				.InitiallyCollapsed(false)
-				.Padding(10.0f)
-				.AreaTitle(DebugView->GetTitle())
-				.BodyContent()
-				[
-					DebugView.ToSharedRef()
-				]
-			];
+			TabManager->TryInvokeTab(ViewName);
 		}
 	}
 }
 
+void SRewindDebugger::OnComponentSelectionChanged(TSharedPtr<FDebugObjectInfo> SelectedItem, ESelectInfo::Type SelectInfo)
+{
+	SelectedComponent = SelectedItem;
+
+	CreateDebugViews();
+	CreateDebugTabs();
+}
+
 FReply SRewindDebugger::OnSelectActorClicked()
 {
 		FActorPickerModeModule& ActorPickerMode = FModuleManager::Get().GetModuleChecked<FActorPickerModeModule>("ActorPickerMode");
diff --git a/Engine/Plugins/Animation/GameplayInsights/Source/RewindDebugger/Private/SRewindDebugger.h b/Engine/Plugins/Animation/GameplayInsights/Source/RewindDebugger/Private/SRewindDebugger.h
--- a/Engine/Plugins/Animation/GameplayInsights/Source/RewindDebugger/Private/SRewindDebugger.h
+++ b/Engine/Plugins/Animation/GameplayInsights/Source/RewindDebugger/Private/SRewindDebugger.h
@@ -13,8 +13,7 @@
 #include "Widgets/Input/SSearchBox.h"
 #include "Widgets/SCompoundWidget.h"
 
-
-class SScrollBox;
+class SDockTab;
 
 class SRewindDebugger : public SCompoundWidget
 {
@@ -72,12 +71,30 @@ private:
 
 	void TraceTimeChanged(double Time);
 
+	TSharedRef<SWidget> MakeMainMenu();
+	void MakeViewsMenu(FMenuBuilder& MenuBuilder);
+	
 	// component tree view
 	TArray<TSharedPtr<FDebugObjectInfo>>* DebugComponents;
+	TSharedPtr<FDebugObjectInfo> SelectedComponent;
     void OnComponentSelectionChanged(TSharedPtr<FDebugObjectInfo> SelectedItem, ESelectInfo::Type SelectInfo);
 	TSharedPtr<SRewindDebuggerComponentTree> ComponentTreeView;
 
-	// anim graph view 
+	// Debug View Tabs
+	TSharedRef<SDockTab> SpawnTab(const FSpawnTabArgs& Args, FName ViewName);
+	bool CanSpawnTab(const FSpawnTabArgs& Args, FName ViewName);
+	void CloseTab(FName TabName);
+	void OnPinnedTabClosed(TSharedRef<SDockTab> Tab);
+	void ExtendTabMenu(FMenuBuilder& MenuBuilder, TSharedPtr<IGameplayInsightsDebugView> View);
+	void PinTab(TSharedPtr<IGameplayInsightsDebugView> View);
+	void ShowAllViews();
+	void CreateDebugViews();
+	void CreateDebugTabs();
 	TArray<TSharedPtr<IGameplayInsightsDebugView>> DebugViews;
-	TSharedPtr<SScrollBox> DebugViewContainer;
+	TArray<TSharedPtr<IGameplayInsightsDebugView>> PinnedDebugViews;
+	TArray<FName> TabNames;
+	TArray<FName> HiddenTabs;  // keep track of tabs that have been closed so we don't automatically reopen them when switching components
+	bool bInternalClosingTab = false;
+
+	TSharedPtr<FTabManager> TabManager;
 };
