Fixed Unreal Insights build error.

diff --git a/Engine/Plugins/Animation/GameplayInsights/Source/GameplayInsights/Private/SAnimGraphSchematicView.cpp b/Engine/Plugins/Animation/GameplayInsights/Source/GameplayInsights/Private/SAnimGraphSchematicView.cpp
--- a/Engine/Plugins/Animation/GameplayInsights/Source/GameplayInsights/Private/SAnimGraphSchematicView.cpp
+++ b/Engine/Plugins/Animation/GameplayInsights/Source/GameplayInsights/Private/SAnimGraphSchematicView.cpp
@@ -1,7 +1,6 @@
 // Copyright Epic Games, Inc. All Rights Reserved.
 
 #include "SAnimGraphSchematicView.h"
-#include "Animation/AnimInstance.h"
 #include "AnimationProvider.h"
 #include "Widgets/Input/SSearchBox.h"
 #include "TraceServices/Model/AnalysisSession.h"
@@ -20,6 +19,10 @@
 #include "Widgets/Input/SComboButton.h"
 #include "Widgets/Images/SImage.h"
 
+#if WITH_ENGINE
+#include "Animation/AnimInstance.h"
+#endif
+
 #if WITH_EDITOR
 #include "Subsystems/AssetEditorSubsystem.h"
 #include "Editor.h"
@@ -1015,7 +1018,11 @@ FText FAnimGraphSchematicViewCreator::GetTitle() const
 
 FSlateIcon FAnimGraphSchematicViewCreator::GetIcon() const
 {
+#if WITH_ENGINE
 	return FSlateIconFinder::FindIconForClass(UAnimInstance::StaticClass());
+#else
+	return FSlateIcon();
+#endif
 }
 
 TSharedPtr<IGameplayInsightsDebugView> FAnimGraphSchematicViewCreator::CreateDebugView(uint64 ObjectId, double CurrentTime, const TraceServices::IAnalysisSession& AnalysisSession) const
