Adding various animation property views to rewind debugger

diff --git a/Engine/Plugins/Animation/GameplayInsights/Source/GameplayInsights/Private/GameplayInsightsModule.cpp b/Engine/Plugins/Animation/GameplayInsights/Source/GameplayInsights/Private/GameplayInsightsModule.cpp
--- a/Engine/Plugins/Animation/GameplayInsights/Source/GameplayInsights/Private/GameplayInsightsModule.cpp
+++ b/Engine/Plugins/Animation/GameplayInsights/Source/GameplayInsights/Private/GameplayInsightsModule.cpp
@@ -17,7 +17,10 @@
 #include "Stats/Stats.h"
 #include "ObjectPropertyTrace.h"
 #include "SAnimGraphSchematicView.h"
+#include "SAnimationCurvesView.h"
+#include "SBlendWeightsView.h"
 #include "SObjectPropertiesView.h"
+#include "SNotifiesView.h"
 
 #if WITH_EDITOR
 #include "IAnimationBlueprintEditorModule.h"
@@ -63,6 +66,9 @@ void FGameplayInsightsModule::StartupModule()
 
 	DebugViewCreator.RegisterDebugViewCreator("Object", MakeShared<FObjectPropertiesViewCreator>());
 	DebugViewCreator.RegisterDebugViewCreator("AnimInstance", MakeShared<FAnimGraphSchematicViewCreator>());
+	DebugViewCreator.RegisterDebugViewCreator("AnimInstance", MakeShared<FBlendWeightsViewCreator>());
+	DebugViewCreator.RegisterDebugViewCreator("AnimInstance", MakeShared<FNotifiesViewCreator>());
+	DebugViewCreator.RegisterDebugViewCreator("SkeletalMeshComponent", MakeShared<FAnimationCurvesViewCreator>());
 
 #if WITH_EDITOR
 
diff --git a/Engine/Plugins/Animation/GameplayInsights/Source/GameplayInsights/Private/SAnimGraphSchematicView.cpp b/Engine/Plugins/Animation/GameplayInsights/Source/GameplayInsights/Private/SAnimGraphSchematicView.cpp
--- a/Engine/Plugins/Animation/GameplayInsights/Source/GameplayInsights/Private/SAnimGraphSchematicView.cpp
+++ b/Engine/Plugins/Animation/GameplayInsights/Source/GameplayInsights/Private/SAnimGraphSchematicView.cpp
@@ -9,7 +9,6 @@
 #include "Widgets/Layout/SBorder.h"
 #include "Widgets/SBoxPanel.h"
 #include "Styling/CoreStyle.h"
-#include "Styling/SlateIconFinder.h"
 #include "GameplayInsightsStyle.h"
 #include "Widgets/Layout/SScrollBorder.h"
 #include "Framework/MultiBox/MultiBoxBuilder.h"
@@ -19,11 +18,9 @@
 #include "Widgets/Input/SComboButton.h"
 #include "Widgets/Images/SImage.h"
 
-#if WITH_ENGINE
-#include "Animation/AnimInstance.h"
-#endif
-
 #if WITH_EDITOR
+#include "Animation/AnimInstance.h"
+#include "Styling/SlateIconFinder.h"
 #include "Subsystems/AssetEditorSubsystem.h"
 #include "Editor.h"
 #endif
@@ -1018,7 +1015,7 @@ FText FAnimGraphSchematicViewCreator::GetTitle() const
 
 FSlateIcon FAnimGraphSchematicViewCreator::GetIcon() const
 {
-#if WITH_ENGINE
+#if WITH_EDITOR
 	return FSlateIconFinder::FindIconForClass(UAnimInstance::StaticClass());
 #else
 	return FSlateIcon();
diff --git a/Engine/Plugins/Animation/GameplayInsights/Source/GameplayInsights/Private/SAnimationCurvesView.cpp b/Engine/Plugins/Animation/GameplayInsights/Source/GameplayInsights/Private/SAnimationCurvesView.cpp
new file mode 100644
--- /dev/null
+++ b/Engine/Plugins/Animation/GameplayInsights/Source/GameplayInsights/Private/SAnimationCurvesView.cpp
@@ -0,0 +1,79 @@
+// Copyright Epic Games, Inc. All Rights Reserved.
+
+#include "SAnimationCurvesView.h"
+#include "AnimationProvider.h"
+#include "TraceServices/Model/AnalysisSession.h"
+#include "GameplayProvider.h"
+#include "Styling/SlateIconFinder.h"
+#include "TraceServices/Model/Frames.h"
+#include "VariantTreeNode.h"
+
+#if WITH_EDITOR
+#include "Animation/AnimInstance.h"
+#endif
+
+#define LOCTEXT_NAMESPACE "SAnimationCurvesView"
+
+void SAnimationCurvesView::GetVariantsAtFrame(const TraceServices::FFrame& InFrame, TArray<TSharedRef<FVariantTreeNode>>& OutVariants) const
+{
+	const FGameplayProvider* GameplayProvider = AnalysisSession->ReadProvider<FGameplayProvider>(FGameplayProvider::ProviderName);
+	const FAnimationProvider* AnimationProvider = AnalysisSession->ReadProvider<FAnimationProvider>(FAnimationProvider::ProviderName);
+
+	if(GameplayProvider && AnimationProvider)
+	{
+		TraceServices::FAnalysisSessionReadScope SessionReadScope(*AnalysisSession);
+
+		TSharedRef<FVariantTreeNode> Header = OutVariants.Add_GetRef(FVariantTreeNode::MakeHeader(LOCTEXT("Animation Curves", "Animation Curves"), 0));
+
+		AnimationProvider->ReadSkeletalMeshPoseTimeline(ObjectId, [&Header, &AnimationProvider, &InFrame](const FAnimationProvider::SkeletalMeshPoseTimeline& InTimeline, bool bInHasCurves)
+		{
+			InTimeline.EnumerateEvents(InFrame.StartTime, InFrame.EndTime, [&Header, &AnimationProvider, &InFrame](double InStartTime, double InEndTime, uint32 InDepth, const FSkeletalMeshPoseMessage& InMessage)
+			{
+				if(InStartTime >= InFrame.StartTime && InStartTime <= InFrame.EndTime)
+				{
+					AnimationProvider->EnumerateSkeletalMeshCurves(InMessage, [&Header, &AnimationProvider](const FSkeletalMeshNamedCurve& InCurve)
+					{
+						const TCHAR* CurveName = AnimationProvider->GetName(InCurve.Id);
+						Header->AddChild(FVariantTreeNode::MakeFloat(FText::FromString(CurveName), InCurve.Value));
+					});
+				}
+				return TraceServices::EEventEnumerate::Continue;
+			});
+		});
+	}
+}
+
+
+static const FName AnimationCurvesName("AnimationCurves");
+
+FName SAnimationCurvesView::GetName() const
+{
+	return AnimationCurvesName;
+}
+
+FName FAnimationCurvesViewCreator::GetName() const
+{
+	return AnimationCurvesName;
+}
+
+FText FAnimationCurvesViewCreator::GetTitle() const
+{
+	return LOCTEXT("Animation Curves", "Animation Curves");
+}
+
+FSlateIcon FAnimationCurvesViewCreator::GetIcon() const
+{
+#if WITH_EDITOR
+	return FSlateIconFinder::FindIconForClass(UAnimInstance::StaticClass());
+#else
+	return FSlateIcon();
+#endif
+}
+
+TSharedPtr<IGameplayInsightsDebugView> FAnimationCurvesViewCreator::CreateDebugView(uint64 ObjectId, double CurrentTime, const TraceServices::IAnalysisSession& AnalysisSession) const
+{
+	return SNew(SAnimationCurvesView, ObjectId, CurrentTime, AnalysisSession);
+}
+
+
+#undef LOCTEXT_NAMESPACE
diff --git a/Engine/Plugins/Animation/GameplayInsights/Source/GameplayInsights/Private/SAnimationCurvesView.h b/Engine/Plugins/Animation/GameplayInsights/Source/GameplayInsights/Private/SAnimationCurvesView.h
new file mode 100644
--- /dev/null
+++ b/Engine/Plugins/Animation/GameplayInsights/Source/GameplayInsights/Private/SAnimationCurvesView.h
@@ -0,0 +1,26 @@
+// Copyright Epic Games, Inc. All Rights Reserved.
+
+#pragma once
+
+#include "CoreMinimal.h"
+#include "IGameplayInsightsDebugView.h"
+#include "IGameplayInsightsDebugViewCreator.h"
+#include "SPropertiesDebugViewBase.h"
+
+namespace TraceServices { class IAnalysisSession; }
+
+class SAnimationCurvesView : public SPropertiesDebugViewBase
+{
+public:
+	virtual void GetVariantsAtFrame(const TraceServices::FFrame& InFrame, TArray<TSharedRef<FVariantTreeNode>>& OutVariants) const override;
+	virtual FName GetName() const override;
+};
+
+class FAnimationCurvesViewCreator : public ICreateGameplayInsightsDebugView
+{
+	public:
+		virtual FName GetName() const override;
+		virtual FText GetTitle() const override;
+		virtual FSlateIcon GetIcon() const override;
+		virtual TSharedPtr<IGameplayInsightsDebugView> CreateDebugView(uint64 ObjectId, double CurrentTime, const TraceServices::IAnalysisSession& InAnalysisSession) const override;
+};
\ No newline at end of file
diff --git a/Engine/Plugins/Animation/GameplayInsights/Source/GameplayInsights/Private/SBlendWeightsView.cpp b/Engine/Plugins/Animation/GameplayInsights/Source/GameplayInsights/Private/SBlendWeightsView.cpp
new file mode 100644
--- /dev/null
+++ b/Engine/Plugins/Animation/GameplayInsights/Source/GameplayInsights/Private/SBlendWeightsView.cpp
@@ -0,0 +1,86 @@
+// Copyright Epic Games, Inc. All Rights Reserved.
+
+#include "SBlendWeightsView.h"
+#include "AnimationProvider.h"
+#include "TraceServices/Model/AnalysisSession.h"
+#include "GameplayProvider.h"
+#include "Styling/SlateIconFinder.h"
+#include "TraceServices/Model/Frames.h"
+#include "VariantTreeNode.h"
+
+#if WITH_EDITOR
+#include "Animation/AnimInstance.h"
+#endif
+
+#define LOCTEXT_NAMESPACE "SBlendWeightsView"
+
+void SBlendWeightsView::GetVariantsAtFrame(const TraceServices::FFrame& InFrame, TArray<TSharedRef<FVariantTreeNode>>& OutVariants) const
+{
+	const FGameplayProvider* GameplayProvider = AnalysisSession->ReadProvider<FGameplayProvider>(FGameplayProvider::ProviderName);
+	const FAnimationProvider* AnimationProvider = AnalysisSession->ReadProvider<FAnimationProvider>(FAnimationProvider::ProviderName);
+
+	if(GameplayProvider && AnimationProvider)
+	{
+		TraceServices::FAnalysisSessionReadScope SessionReadScope(*AnalysisSession);
+
+		AnimationProvider->ReadTickRecordTimeline(ObjectId, [&GameplayProvider, &OutVariants, &InFrame](const FAnimationProvider::TickRecordTimeline& InTimeline)
+		{
+			InTimeline.EnumerateEvents(InFrame.StartTime, InFrame.EndTime, [&GameplayProvider, &OutVariants, &InFrame](double InStartTime, double InEndTime, uint32 InDepth, const FTickRecordMessage& InMessage)
+			{
+				if(InStartTime >= InFrame.StartTime && InEndTime <= InFrame.EndTime)
+				{
+					const FClassInfo& ClassInfo = GameplayProvider->GetClassInfoFromObject(InMessage.AssetId);
+					TSharedRef<FVariantTreeNode> Header = OutVariants.Add_GetRef(FVariantTreeNode::MakeObject(FText::FromString(ClassInfo.Name), InMessage.AssetId, InMessage.AssetId));
+					
+					Header->AddChild(FVariantTreeNode::MakeFloat(LOCTEXT("BlendWeight", "Blend Weight"), InMessage.BlendWeight));
+					Header->AddChild(FVariantTreeNode::MakeFloat(LOCTEXT("PlaybackTime", "Playback Time"), InMessage.PlaybackTime));
+					Header->AddChild(FVariantTreeNode::MakeFloat(LOCTEXT("RootMotionWeight", "Root Motion Weight"), InMessage.RootMotionWeight));
+					Header->AddChild(FVariantTreeNode::MakeFloat(LOCTEXT("PlayRate", "Play Rate"), InMessage.PlayRate));
+					if(InMessage.bIsBlendSpace)
+					{
+						Header->AddChild(FVariantTreeNode::MakeFloat(LOCTEXT("BlendSpacePositionX", "Blend Space Position X"), InMessage.BlendSpacePositionX));
+						Header->AddChild(FVariantTreeNode::MakeFloat(LOCTEXT("BlendSpacePositionY", "Blend Space Position Y"), InMessage.BlendSpacePositionY));
+						Header->AddChild(FVariantTreeNode::MakeFloat(LOCTEXT("BlendSpaceFilteredPositionX", "Blend Space Filtered Position X"), InMessage.BlendSpaceFilteredPositionX));
+						Header->AddChild(FVariantTreeNode::MakeFloat(LOCTEXT("BlendSpaceFilteredPositionY", "Blend Space Filtered Position Y"), InMessage.BlendSpaceFilteredPositionY));
+					}
+				}
+				return TraceServices::EEventEnumerate::Continue;
+			});
+		});
+	}
+}
+
+
+static const FName BlendWeightsName("BlendWeights");
+
+FName SBlendWeightsView::GetName() const
+{
+	return BlendWeightsName;
+}
+
+FName FBlendWeightsViewCreator::GetName() const
+{
+	return BlendWeightsName;
+}
+
+FText FBlendWeightsViewCreator::GetTitle() const
+{
+	return LOCTEXT("Blend Weights", "Blend Weights");
+}
+
+FSlateIcon FBlendWeightsViewCreator::GetIcon() const
+{
+#if WITH_EDITOR
+	return FSlateIconFinder::FindIconForClass(UAnimInstance::StaticClass());
+#else
+	return FSlateIcon();
+#endif
+}
+
+TSharedPtr<IGameplayInsightsDebugView> FBlendWeightsViewCreator::CreateDebugView(uint64 ObjectId, double CurrentTime, const TraceServices::IAnalysisSession& AnalysisSession) const
+{
+	return SNew(SBlendWeightsView, ObjectId, CurrentTime, AnalysisSession);
+}
+
+
+#undef LOCTEXT_NAMESPACE
diff --git a/Engine/Plugins/Animation/GameplayInsights/Source/GameplayInsights/Private/SBlendWeightsView.h b/Engine/Plugins/Animation/GameplayInsights/Source/GameplayInsights/Private/SBlendWeightsView.h
new file mode 100644
--- /dev/null
+++ b/Engine/Plugins/Animation/GameplayInsights/Source/GameplayInsights/Private/SBlendWeightsView.h
@@ -0,0 +1,26 @@
+// Copyright Epic Games, Inc. All Rights Reserved.
+
+#pragma once
+
+#include "CoreMinimal.h"
+#include "IGameplayInsightsDebugView.h"
+#include "IGameplayInsightsDebugViewCreator.h"
+#include "SPropertiesDebugViewBase.h"
+
+namespace TraceServices { class IAnalysisSession; }
+
+class SBlendWeightsView : public SPropertiesDebugViewBase
+{
+public:
+	virtual void GetVariantsAtFrame(const TraceServices::FFrame& InFrame, TArray<TSharedRef<FVariantTreeNode>>& OutVariants) const override;
+	virtual FName GetName() const override;
+};
+
+class FBlendWeightsViewCreator : public ICreateGameplayInsightsDebugView
+{
+	public:
+		virtual FName GetName() const override;
+		virtual FText GetTitle() const override;
+		virtual FSlateIcon GetIcon() const override;
+		virtual TSharedPtr<IGameplayInsightsDebugView> CreateDebugView(uint64 ObjectId, double CurrentTime, const TraceServices::IAnalysisSession& InAnalysisSession) const override;
+};
\ No newline at end of file
diff --git a/Engine/Plugins/Animation/GameplayInsights/Source/GameplayInsights/Private/SNotifiesView.cpp b/Engine/Plugins/Animation/GameplayInsights/Source/GameplayInsights/Private/SNotifiesView.cpp
new file mode 100644
--- /dev/null
+++ b/Engine/Plugins/Animation/GameplayInsights/Source/GameplayInsights/Private/SNotifiesView.cpp
@@ -0,0 +1,140 @@
+// Copyright Epic Games, Inc. All Rights Reserved.
+
+#include "SNotifiesView.h"
+#include "AnimationProvider.h"
+#include "TraceServices/Model/AnalysisSession.h"
+#include "GameplayProvider.h"
+#include "Styling/SlateIconFinder.h"
+#include "TraceServices/Model/Frames.h"
+#include "VariantTreeNode.h"
+
+#if WITH_EDITOR
+#include "Animation/AnimInstance.h"
+#endif
+
+#define LOCTEXT_NAMESPACE "SNotifiesView"
+
+static const TCHAR* GetNotifyType(EAnimNotifyMessageType InType, bool bInForEvent)
+{
+	switch (InType)
+	{
+	case EAnimNotifyMessageType::Begin:
+	case EAnimNotifyMessageType::End:
+	case EAnimNotifyMessageType::Tick:
+	{
+		static FText StateType = LOCTEXT("AnimNotifyState", "Anim Notify State");
+		static FText BeginType = LOCTEXT("AnimNotifyStateBegin", "Anim Notify State (Begin)");
+		static FText EndType = LOCTEXT("AnimNotifyStateEnd", "Anim Notify State (End)");
+		static FText TickType = LOCTEXT("AnimNotifyStateTick", "Anim Notify State (Tick)");
+		if (bInForEvent)
+		{
+			switch (InType)
+			{
+			case EAnimNotifyMessageType::Begin:
+				return *BeginType.ToString();
+			case EAnimNotifyMessageType::End:
+				return *EndType.ToString();
+			case EAnimNotifyMessageType::Tick:
+				return *TickType.ToString();
+			}
+		}
+		else
+		{
+			return *StateType.ToString();
+		}
+	}
+	case EAnimNotifyMessageType::Event:
+	{
+		static FText Type = LOCTEXT("AnimNotify", "Anim Notify");
+		return *Type.ToString();
+	}
+	case EAnimNotifyMessageType::SyncMarker:
+	{
+		static FText Type = LOCTEXT("SyncMarker", "Sync Marker");
+		return *Type.ToString();
+	}
+	}
+
+	static FText UnknownType = LOCTEXT("UnknownType", "Unknown");
+	return *UnknownType.ToString();
+}
+
+void SNotifiesView::GetVariantsAtFrame(const TraceServices::FFrame& InFrame, TArray<TSharedRef<FVariantTreeNode>>& OutVariants) const
+{
+	TSharedRef<FVariantTreeNode> Header = OutVariants.Add_GetRef(FVariantTreeNode::MakeHeader(LOCTEXT("AnimNotifiesHeader", "Notifies and Sync Markers"), 0));
+
+	const FGameplayProvider* GameplayProvider = AnalysisSession->ReadProvider<FGameplayProvider>(FGameplayProvider::ProviderName);
+	const FAnimationProvider* AnimationProvider = AnalysisSession->ReadProvider<FAnimationProvider>(FAnimationProvider::ProviderName);
+
+	if(GameplayProvider && AnimationProvider)
+	{
+		TraceServices::FAnalysisSessionReadScope SessionReadScope(*AnalysisSession);
+		bool bForEvent = true;
+		auto ProcessEvent = [this, &AnimationProvider, &GameplayProvider, &Header, &bForEvent](double InStartTime, double InEndTime, uint32 InDepth, const FAnimNotifyMessage& InMessage)
+		{
+			const TCHAR* Name = AnimationProvider->GetName(InMessage.NameId);
+			TSharedRef<FVariantTreeNode> NotifyHeader = Header->AddChild(FVariantTreeNode::MakeHeader(FText::FromString(Name), InMessage.NameId));
+
+			NotifyHeader->AddChild(FVariantTreeNode::MakeString(LOCTEXT("EventType", "Type"), GetNotifyType(InMessage.NotifyEventType, bForEvent)));
+			NotifyHeader->AddChild(FVariantTreeNode::MakeFloat(LOCTEXT("EventDuration", "Duration"), InEndTime - InStartTime));
+			NotifyHeader->AddChild(FVariantTreeNode::MakeFloat(LOCTEXT("EventTime", "Time"), InStartTime));
+
+			if(InMessage.NotifyEventType != EAnimNotifyMessageType::SyncMarker)
+			{
+				NotifyHeader->AddChild(FVariantTreeNode::MakeObject(LOCTEXT("Asset", "Asset"), InMessage.AssetId));
+
+				const FObjectInfo& NotifyInfo = GameplayProvider->GetObjectInfo(InMessage.NotifyId);
+				bool bIsState = InMessage.NotifyEventType == EAnimNotifyMessageType::Begin || InMessage.NotifyEventType == EAnimNotifyMessageType::End || InMessage.NotifyEventType == EAnimNotifyMessageType::Tick;
+				NotifyHeader->AddChild(FVariantTreeNode::MakeClass((bIsState ? LOCTEXT("Notify State Class", "Notify State Class") : LOCTEXT("NotifyClass", "Notify Class")), NotifyInfo.ClassId));
+			}
+
+			return TraceServices::EEventEnumerate::Continue;
+		};
+
+		AnimationProvider->ReadNotifyTimeline(ObjectId, [&InFrame, &ProcessEvent](const FAnimationProvider::AnimNotifyTimeline& InTimeline)
+		{
+			InTimeline.EnumerateEvents(InFrame.StartTime, InFrame.EndTime, ProcessEvent);
+		});
+
+		bForEvent = false;
+		AnimationProvider->EnumerateNotifyStateTimelines(ObjectId, [&InFrame, &ProcessEvent](uint32 InNotifyNameId, const FAnimationProvider::AnimNotifyTimeline& InTimeline)
+		{
+			InTimeline.EnumerateEvents(InFrame.StartTime, InFrame.EndTime, ProcessEvent);
+		});
+	}
+}
+
+
+static const FName NotifiesName("Notifies");
+
+FName SNotifiesView::GetName() const
+{
+	return NotifiesName;
+}
+
+FName FNotifiesViewCreator::GetName() const
+{
+	return NotifiesName;
+}
+
+FText FNotifiesViewCreator::GetTitle() const
+{
+	return LOCTEXT("Notifies", "Notifies");
+}
+
+FSlateIcon FNotifiesViewCreator::GetIcon() const
+{
+#if WITH_EDITOR
+	return FSlateIconFinder::FindIconForClass(UAnimInstance::StaticClass());
+#else
+	return FSlateIcon();
+#endif
+}
+
+TSharedPtr<IGameplayInsightsDebugView> FNotifiesViewCreator::CreateDebugView(uint64 ObjectId, double CurrentTime, const TraceServices::IAnalysisSession& AnalysisSession) const
+{
+	return SNew(SNotifiesView, ObjectId, CurrentTime, AnalysisSession);
+}
+
+
+#undef LOCTEXT_NAMESPACE
diff --git a/Engine/Plugins/Animation/GameplayInsights/Source/GameplayInsights/Private/SNotifiesView.h b/Engine/Plugins/Animation/GameplayInsights/Source/GameplayInsights/Private/SNotifiesView.h
new file mode 100644
--- /dev/null
+++ b/Engine/Plugins/Animation/GameplayInsights/Source/GameplayInsights/Private/SNotifiesView.h
@@ -0,0 +1,26 @@
+// Copyright Epic Games, Inc. All Rights Reserved.
+
+#pragma once
+
+#include "CoreMinimal.h"
+#include "IGameplayInsightsDebugView.h"
+#include "IGameplayInsightsDebugViewCreator.h"
+#include "SPropertiesDebugViewBase.h"
+
+namespace TraceServices { class IAnalysisSession; }
+
+class SNotifiesView : public SPropertiesDebugViewBase
+{
+public:
+	virtual void GetVariantsAtFrame(const TraceServices::FFrame& InFrame, TArray<TSharedRef<FVariantTreeNode>>& OutVariants) const override;
+	virtual FName GetName() const override;
+};
+
+class FNotifiesViewCreator : public ICreateGameplayInsightsDebugView
+{
+	public:
+		virtual FName GetName() const override;
+		virtual FText GetTitle() const override;
+		virtual FSlateIcon GetIcon() const override;
+		virtual TSharedPtr<IGameplayInsightsDebugView> CreateDebugView(uint64 ObjectId, double CurrentTime, const TraceServices::IAnalysisSession& InAnalysisSession) const override;
+};
\ No newline at end of file
diff --git a/Engine/Plugins/Animation/GameplayInsights/Source/GameplayInsights/Private/SObjectPropertiesView.cpp b/Engine/Plugins/Animation/GameplayInsights/Source/GameplayInsights/Private/SObjectPropertiesView.cpp
--- a/Engine/Plugins/Animation/GameplayInsights/Source/GameplayInsights/Private/SObjectPropertiesView.cpp
+++ b/Engine/Plugins/Animation/GameplayInsights/Source/GameplayInsights/Private/SObjectPropertiesView.cpp
@@ -10,34 +10,6 @@
 
 #define LOCTEXT_NAMESPACE "SObjectPropertiesView"
 
-void SObjectPropertiesView::Construct(const FArguments& InArgs, uint64 InObjectId, double InTimeMarker, const TraceServices::IAnalysisSession& InAnalysisSession)
-{
-	ObjectId = InObjectId;
-	AnalysisSession = &InAnalysisSession;
-
-	View = SNew(SVariantValueView, InAnalysisSession).OnGetVariantValues(this, &SObjectPropertiesView::GetVariantsAtFrame);
-
-	SetTimeMarker(InTimeMarker);
-
-	ChildSlot
-	[
-		View.ToSharedRef()
-	];
-
-}
-
-void SObjectPropertiesView::SetTimeMarker(double Time)
-{
-	TimeMarker = Time;
-	TraceServices::FAnalysisSessionReadScope SessionReadScope(*AnalysisSession);
-	const TraceServices::IFrameProvider& FramesProvider = TraceServices::ReadFrameProvider(*AnalysisSession);
-	TraceServices::FFrame MarkerFrame;
-	if(FramesProvider.GetFrameFromTime(ETraceFrameType::TraceFrameType_Game, TimeMarker, MarkerFrame))
-	{
-		View->RequestRefresh(MarkerFrame);
-	}
-}
-
 void SObjectPropertiesView::GetVariantsAtFrame(const TraceServices::FFrame& InFrame, TArray<TSharedRef<FVariantTreeNode>>& OutVariants) const
 {
 	if (const FGameplayProvider* GameplayProvider = AnalysisSession->ReadProvider<FGameplayProvider>(FGameplayProvider::ProviderName))
diff --git a/Engine/Plugins/Animation/GameplayInsights/Source/GameplayInsights/Private/SObjectPropertiesView.h b/Engine/Plugins/Animation/GameplayInsights/Source/GameplayInsights/Private/SObjectPropertiesView.h
--- a/Engine/Plugins/Animation/GameplayInsights/Source/GameplayInsights/Private/SObjectPropertiesView.h
+++ b/Engine/Plugins/Animation/GameplayInsights/Source/GameplayInsights/Private/SObjectPropertiesView.h
@@ -5,28 +5,15 @@
 #include "CoreMinimal.h"
 #include "IGameplayInsightsDebugView.h"
 #include "IGameplayInsightsDebugViewCreator.h"
-#include "SVariantValueView.h"
+#include "SPropertiesDebugViewBase.h"
 
 namespace TraceServices { class IAnalysisSession; }
 
-class SObjectPropertiesView : public IGameplayInsightsDebugView
+class SObjectPropertiesView : public SPropertiesDebugViewBase
 {
-	SLATE_BEGIN_ARGS(SObjectPropertiesView) {}
-	SLATE_END_ARGS()
 public:
-	void Construct(const FArguments& InArgs, uint64 InObjectId, double InTimeMarker, const TraceServices::IAnalysisSession& InAnalysisSession);
-
-	virtual void SetTimeMarker(double InTimeMarker) override;
+	virtual void GetVariantsAtFrame(const TraceServices::FFrame& InFrame, TArray<TSharedRef<FVariantTreeNode>>& OutVariants) const override;
 	virtual FName GetName() const override;
-	virtual uint64 GetObjectId() const override { return ObjectId; }
-
-	void GetVariantsAtFrame(const TraceServices::FFrame& InFrame, TArray<TSharedRef<FVariantTreeNode>>& OutVariants) const;
-private:
-	TSharedPtr<SVariantValueView> View;
-
-	uint64 ObjectId;
-	double TimeMarker;
-	const TraceServices::IAnalysisSession* AnalysisSession;
 };
 
 class FObjectPropertiesViewCreator : public ICreateGameplayInsightsDebugView
diff --git a/Engine/Plugins/Animation/GameplayInsights/Source/GameplayInsights/Private/SPropertiesDebugViewBase.cpp b/Engine/Plugins/Animation/GameplayInsights/Source/GameplayInsights/Private/SPropertiesDebugViewBase.cpp
new file mode 100644
--- /dev/null
+++ b/Engine/Plugins/Animation/GameplayInsights/Source/GameplayInsights/Private/SPropertiesDebugViewBase.cpp
@@ -0,0 +1,41 @@
+// Copyright Epic Games, Inc. All Rights Reserved.
+
+#include "SPropertiesDebugViewBase.h"
+#include "AnimationProvider.h"
+#include "TraceServices/Model/AnalysisSession.h"
+#include "GameplayProvider.h"
+#include "Styling/SlateIconFinder.h"
+#include "TraceServices/Model/Frames.h"
+#include "VariantTreeNode.h"
+
+#define LOCTEXT_NAMESPACE "SPropertiesDebugViewBase"
+
+void SPropertiesDebugViewBase::Construct(const FArguments& InArgs, uint64 InObjectId, double InTimeMarker, const TraceServices::IAnalysisSession& InAnalysisSession)
+{
+	ObjectId = InObjectId;
+	AnalysisSession = &InAnalysisSession;
+
+	View = SNew(SVariantValueView, InAnalysisSession).OnGetVariantValues(this, &SPropertiesDebugViewBase::GetVariantsAtFrame);
+
+	SetTimeMarker(InTimeMarker);
+
+	ChildSlot
+	[
+		View.ToSharedRef()
+	];
+
+}
+
+void SPropertiesDebugViewBase::SetTimeMarker(double Time)
+{
+	TimeMarker = Time;
+	TraceServices::FAnalysisSessionReadScope SessionReadScope(*AnalysisSession);
+	const TraceServices::IFrameProvider& FramesProvider = TraceServices::ReadFrameProvider(*AnalysisSession);
+	TraceServices::FFrame MarkerFrame;
+	if(FramesProvider.GetFrameFromTime(ETraceFrameType::TraceFrameType_Game, TimeMarker, MarkerFrame))
+	{
+		View->RequestRefresh(MarkerFrame);
+	}
+}
+
+#undef LOCTEXT_NAMESPACE
diff --git a/Engine/Plugins/Animation/GameplayInsights/Source/GameplayInsights/Private/SPropertiesDebugViewBase.h b/Engine/Plugins/Animation/GameplayInsights/Source/GameplayInsights/Private/SPropertiesDebugViewBase.h
new file mode 100644
--- /dev/null
+++ b/Engine/Plugins/Animation/GameplayInsights/Source/GameplayInsights/Private/SPropertiesDebugViewBase.h
@@ -0,0 +1,30 @@
+// Copyright Epic Games, Inc. All Rights Reserved.
+
+#pragma once
+
+#include "CoreMinimal.h"
+#include "IGameplayInsightsDebugView.h"
+#include "IGameplayInsightsDebugViewCreator.h"
+#include "SVariantValueView.h"
+
+namespace TraceServices { class IAnalysisSession; }
+
+class SPropertiesDebugViewBase : public IGameplayInsightsDebugView
+{
+	SLATE_BEGIN_ARGS(SPropertiesDebugViewBase) {}
+	SLATE_END_ARGS()
+public:
+	void Construct(const FArguments& InArgs, uint64 InObjectId, double InTimeMarker, const TraceServices::IAnalysisSession& InAnalysisSession);
+
+	virtual void SetTimeMarker(double InTimeMarker) override;
+	virtual uint64 GetObjectId() const override { return ObjectId; }
+
+	virtual void GetVariantsAtFrame(const TraceServices::FFrame& InFrame, TArray<TSharedRef<FVariantTreeNode>>& OutVariants) const = 0;
+
+protected:
+	TSharedPtr<SVariantValueView> View;
+
+	uint64 ObjectId;
+	double TimeMarker;
+	const TraceServices::IAnalysisSession* AnalysisSession;
+};
\ No newline at end of file
