bump the tab layout name for SRewindDebugger so that the new tabs don't all come up in individual windows

diff --git a/Engine/Plugins/Animation/GameplayInsights/Source/RewindDebugger/Private/SRewindDebugger.cpp b/Engine/Plugins/Animation/GameplayInsights/Source/RewindDebugger/Private/SRewindDebugger.cpp
--- a/Engine/Plugins/Animation/GameplayInsights/Source/RewindDebugger/Private/SRewindDebugger.cpp
+++ b/Engine/Plugins/Animation/GameplayInsights/Source/RewindDebugger/Private/SRewindDebugger.cpp
@@ -274,7 +274,7 @@ void SRewindDebugger::Construct(const FArguments& InArgs, TSharedRef<FUICommandL
 		}
 	);
 
-	TSharedRef<FTabManager::FLayout> Layout = FTabManager::NewLayout("RewindDebuggerLayout") 
+	TSharedRef<FTabManager::FLayout> Layout = FTabManager::NewLayout("RewindDebuggerLayout1.0") 
 		->AddArea
 		(
 			FTabManager::NewPrimaryArea()
