Removed a copy-paste ensure() that was included in error.

diff --git a/Engine/Plugins/Animation/GameplayInsights/Source/GameplayInsights/Private/GameplayInsightsModule.cpp b/Engine/Plugins/Animation/GameplayInsights/Source/GameplayInsights/Private/GameplayInsightsModule.cpp
--- a/Engine/Plugins/Animation/GameplayInsights/Source/GameplayInsights/Private/GameplayInsightsModule.cpp
+++ b/Engine/Plugins/Animation/GameplayInsights/Source/GameplayInsights/Private/GameplayInsightsModule.cpp
@@ -186,7 +186,7 @@ void FGameplayInsightsModule::StartupModule()
 				UE::Trace::SendTo(TEXT("localhost"), StoreService->GetRecorderPort());
 #else
 				FTraceAuxiliary::LaunchUnrealTrace();
-				ensure(UnrealInsightsModule.ConnectToStore(TEXT("127.0.0.1")));
+				UnrealInsightsModule.ConnectToStore(TEXT("127.0.0.1"));
 				UE::Trace::SendTo(TEXT("127.0.0.1"));
 #endif // WITH_TRACE_STORE
 
