rewind debugger camera support

diff --git a/Engine/Plugins/Animation/GameplayInsights/Source/GameplayInsights/Private/GameplayAnalyzer.cpp b/Engine/Plugins/Animation/GameplayInsights/Source/GameplayInsights/Private/GameplayAnalyzer.cpp
--- a/Engine/Plugins/Animation/GameplayInsights/Source/GameplayInsights/Private/GameplayAnalyzer.cpp
+++ b/Engine/Plugins/Animation/GameplayInsights/Source/GameplayInsights/Private/GameplayAnalyzer.cpp
@@ -19,6 +19,7 @@ void FGameplayAnalyzer::OnAnalysisBegin(const FOnAnalysisContext& Context)
 	Builder.RouteEvent(RouteId_Object, "Object", "Object");
 	Builder.RouteEvent(RouteId_ObjectEvent, "Object", "ObjectEvent");
 	Builder.RouteEvent(RouteId_World, "Object", "World");
+	Builder.RouteEvent(RouteId_View, "Object", "View");
 	Builder.RouteEvent(RouteId_ClassPropertyStringId, "Object", "ClassPropertyStringId");
 	Builder.RouteEvent(RouteId_ClassProperty, "Object", "ClassProperty");
 	Builder.RouteEvent(RouteId_PropertiesStart, "Object", "PropertiesStart");
@@ -84,6 +85,17 @@ bool FGameplayAnalyzer::OnEvent(uint16 RouteId, EStyle Style, const FOnEventCont
 		GameplayProvider.AppendObjectEvent(Id, Context.EventTime.AsSeconds(Cycle), Event);
 		break;
 	}
+	case RouteId_View:
+	{
+		uint64 Cycle = EventData.GetValue<uint64>("Cycle");
+		uint64 PlayerId = EventData.GetValue<uint64>("PlayerId");
+		FVector Position( EventData.GetValue<double>("PosX"), EventData.GetValue<double>("PosY"), EventData.GetValue<double>("PosZ"));
+		FRotator Rotation( EventData.GetValue<float>("Pitch"), EventData.GetValue<float>("Yaw"), EventData.GetValue<float>("Roll"));
+		float Fov = EventData.GetValue<float>("Fov");
+		float AspectRatio  = EventData.GetValue<float>("AspectRatio");
+		GameplayProvider.AppendView(PlayerId, Context.EventTime.AsSeconds(Cycle), Position, Rotation, Fov, AspectRatio);
+		break;
+	}
 	case RouteId_ClassPropertyStringId:
 	{
 		uint32 Id = EventData.GetValue<uint32>("Id");
diff --git a/Engine/Plugins/Animation/GameplayInsights/Source/GameplayInsights/Private/GameplayAnalyzer.h b/Engine/Plugins/Animation/GameplayInsights/Source/GameplayInsights/Private/GameplayAnalyzer.h
--- a/Engine/Plugins/Animation/GameplayInsights/Source/GameplayInsights/Private/GameplayAnalyzer.h
+++ b/Engine/Plugins/Animation/GameplayInsights/Source/GameplayInsights/Private/GameplayAnalyzer.h
@@ -23,6 +23,7 @@ private:
 		RouteId_Class,
 		RouteId_Object,
 		RouteId_ObjectEvent,
+		RouteId_View,
 		RouteId_ClassPropertyStringId,
 		RouteId_ClassProperty,
 		RouteId_PropertiesStart,
diff --git a/Engine/Plugins/Animation/GameplayInsights/Source/GameplayInsights/Private/GameplayProvider.cpp b/Engine/Plugins/Animation/GameplayInsights/Source/GameplayInsights/Private/GameplayProvider.cpp
--- a/Engine/Plugins/Animation/GameplayInsights/Source/GameplayInsights/Private/GameplayProvider.cpp
+++ b/Engine/Plugins/Animation/GameplayInsights/Source/GameplayInsights/Private/GameplayProvider.cpp
@@ -316,6 +316,39 @@ void FGameplayProvider::AppendObjectEvent(uint64 InObjectId, double InTime, cons
 	Session.UpdateDurationSeconds(InTime);
 }
 
+void FGameplayProvider::ReadViewTimeline(TFunctionRef<void(const IGameplayProvider::ViewTimeline&)> Callback) const
+{
+	Session.ReadAccessCheck();
+
+	if (ViewTimeline.IsValid())
+	{
+		Callback(*ViewTimeline);
+	}
+}
+
+void FGameplayProvider::AppendView(uint64 InPlayerId, double InTime, const FVector& InPosition, const FRotator& InRotation, float InFov, float InAspectRatio)
+{
+	Session.WriteAccessCheck();
+
+	if (!ViewTimeline.IsValid())
+	{
+		ViewTimeline = MakeShared<TraceServices::TPointTimeline<FViewMessage>>(Session.GetLinearAllocator());
+	}
+
+	bHasAnyData = true;
+
+	FViewMessage Message;
+	Message.PlayerId = InPlayerId;
+	Message.Position = InPosition;
+	Message.Rotation = InRotation;
+	Message.Fov = InFov;
+	Message.AspectRatio = InAspectRatio;
+
+	ViewTimeline->AppendEvent(InTime, Message);
+
+	Session.UpdateDurationSeconds(InTime);
+}
+
 void FGameplayProvider::AppendWorld(uint64 InObjectId, int32 InPIEInstanceId, uint8 InType, uint8 InNetMode, bool bInIsSimulating)
 {
 	Session.WriteAccessCheck();
diff --git a/Engine/Plugins/Animation/GameplayInsights/Source/GameplayInsights/Private/GameplayProvider.h b/Engine/Plugins/Animation/GameplayInsights/Source/GameplayInsights/Private/GameplayProvider.h
--- a/Engine/Plugins/Animation/GameplayInsights/Source/GameplayInsights/Private/GameplayProvider.h
+++ b/Engine/Plugins/Animation/GameplayInsights/Source/GameplayInsights/Private/GameplayProvider.h
@@ -36,6 +36,7 @@ public:
 	virtual FOnObjectEndPlay& OnObjectEndPlay() override { return OnObjectEndPlayDelegate; }
 	virtual const TCHAR* GetPropertyName(uint32 InPropertyStringId) const override;
 	virtual const RecordingInfoTimeline* GetRecordingInfo(uint32 RecordingId) const override; 
+	virtual void ReadViewTimeline(TFunctionRef<void(const ViewTimeline&)> Callback) const override;
 
 	/** Add a class message */
 	void AppendClass(uint64 InClassId, uint64 InSuperId, const TCHAR* InClassName, const TCHAR* InClassPathName);
@@ -46,6 +47,9 @@ public:
 	/** Add an object event message */
 	void AppendObjectEvent(uint64 InObjectId, double InTime, const TCHAR* InEvent);
 
+	/** Add a view message */
+	void AppendView(uint64 InObjectId, double InTime, const FVector& InPosition, const FRotator& InRotation, float InFov, float InAspectRatio);
+
 	/** Add a world message */
 	void AppendWorld(uint64 InObjectId, int32 InPIEInstanceId, uint8 InType, uint8 InNetMode, bool bInIsSimulating);
 
@@ -110,6 +114,7 @@ private:
 	/** Message storage */
 	TArray<TSharedRef<TraceServices::TPointTimeline<FObjectEventMessage>>> EventTimelines;
 	TArray<TSharedRef<FObjectPropertiesStorage>> PropertiesStorage;
+	TSharedPtr<TraceServices::TPointTimeline<FViewMessage>> ViewTimeline;
 
 	/** Map of class path name to ClassInfo index */
 	TMap<FStringView, int32> ClassPathNameToIndexMap;
diff --git a/Engine/Plugins/Animation/GameplayInsights/Source/GameplayInsights/Public/IGameplayProvider.h b/Engine/Plugins/Animation/GameplayInsights/Source/GameplayInsights/Public/IGameplayProvider.h
--- a/Engine/Plugins/Animation/GameplayInsights/Source/GameplayInsights/Public/IGameplayProvider.h
+++ b/Engine/Plugins/Animation/GameplayInsights/Source/GameplayInsights/Public/IGameplayProvider.h
@@ -59,6 +59,14 @@ struct FRecordingInfoMessage
 	uint32 RecordingIndex;
 };
 
+struct FViewMessage
+{
+	uint64 PlayerId;
+	FVector Position;
+	FRotator Rotation;
+	float Fov;
+	float AspectRatio;
+};
 
 struct FWorldInfo
 {
@@ -102,6 +110,7 @@ public:
 	typedef TraceServices::ITimeline<FObjectEventMessage> ObjectEventsTimeline;
 	typedef TraceServices::ITimeline<FObjectPropertiesMessage> ObjectPropertiesTimeline;
 	typedef TraceServices::ITimeline<FRecordingInfoMessage> RecordingInfoTimeline;
+	typedef TraceServices::ITimeline<FViewMessage> ViewTimeline;
 
 	virtual bool ReadObjectEventsTimeline(uint64 InObjectId, TFunctionRef<void(const ObjectEventsTimeline&)> Callback) const = 0;
 	virtual bool ReadObjectEvent(uint64 InObjectId, uint64 InMessageId, TFunctionRef<void(const FObjectEventMessage&)> Callback) const = 0;
@@ -120,5 +129,6 @@ public:
 	virtual FOnObjectEndPlay& OnObjectEndPlay() = 0;
 	virtual const TCHAR* GetPropertyName(uint32 InPropertyStringId) const = 0;
 	virtual const RecordingInfoTimeline* GetRecordingInfo(uint32 RecordingId) const = 0; 
+	virtual void ReadViewTimeline(TFunctionRef<void(const ViewTimeline&)> Callback) const = 0;
 };
 
diff --git a/Engine/Plugins/Animation/GameplayInsights/Source/RewindDebugger/Private/RewindDebugger.cpp b/Engine/Plugins/Animation/GameplayInsights/Source/RewindDebugger/Private/RewindDebugger.cpp
--- a/Engine/Plugins/Animation/GameplayInsights/Source/RewindDebugger/Private/RewindDebugger.cpp
+++ b/Engine/Plugins/Animation/GameplayInsights/Source/RewindDebugger/Private/RewindDebugger.cpp
@@ -9,9 +9,13 @@
 #include "Modules/ModuleManager.h"
 #include "ObjectTrace.h"
 #include "TraceServices/ITraceServicesModule.h"
+#include "LevelEditor.h"
+#include "SLevelViewport.h"
+#include "IRewindDebuggerExtension.h"
 
 
 FRewindDebugger* FRewindDebugger::InternalInstance = nullptr;
+const FName IRewindDebuggerExtension::ModularFeatureName = "RewindDebuggerExtension";
 
 FRewindDebugger::FRewindDebugger()  :
 	ControlState(FRewindDebugger::EControlState::Pause),
@@ -21,10 +25,16 @@ FRewindDebugger::FRewindDebugger()  :
 	bRecording(false),
 	PlaybackRate(1),
 	CurrentScrubTime(0),
-	RecordingIndex(0)
+	RecordingIndex(0),
+	bTargetActorPositionValid(false)
 {
 	RecordingDuration.Set(0);
 
+	if (GEditor->bIsSimulatingInEditor || GEditor->PlayWorld)
+	{
+		OnPIEStarted(true);
+	}
+
 	FEditorDelegates::PreBeginPIE.AddRaw(this, &FRewindDebugger::OnPIEStarted);
 	FEditorDelegates::PausePIE.AddRaw(this, &FRewindDebugger::OnPIEPaused);
 	FEditorDelegates::ResumePIE.AddRaw(this, &FRewindDebugger::OnPIEResumed);
@@ -170,7 +180,7 @@ bool FRewindDebugger::UpdateComponentList(uint64 ParentId, TArray<TSharedPtr<FDe
 		// todo: filter components based on creation/destruction frame, and the current scrubbing time (so dynamically created and destroyed components won't add up)
 
 		const FObjectInfo* ObjectInfo = &InObjectInfo;
-		if(ObjectInfo->OuterId == ParentId)
+		if (ObjectInfo->OuterId == ParentId)
 		{
 			int32 FoundIndex = ComponentList.FindLastByPredicate([ObjectInfo](const TSharedPtr<FDebugObjectInfo>& Info) { return Info->ObjectName == ObjectInfo->Name; });
 
@@ -191,23 +201,25 @@ bool FRewindDebugger::UpdateComponentList(uint64 ParentId, TArray<TSharedPtr<FDe
 	return bChanged;
 }
 
-void FRewindDebugger::RefreshDebugComponents()
+bool FRewindDebugger::GetTargetActorPosition(FVector& OutPosition) const
 {
-	if (const TraceServices::IAnalysisSession* Session = GetAnalysisSession())
+	OutPosition = TargetActorPosition;
+	return bTargetActorPositionValid;
+}
+
+uint64 FRewindDebugger::GetTargetActorId() const
+{
+	if (DebugTargetActor.Get() == "")
 	{
-		TraceServices::FAnalysisSessionReadScope SessionReadScope(*Session);
-		UWorld* World = GetWorldToVisualize();
+		return 0;
+	}
+
+	uint64 TargetActorId = 0;
 
+	if (const TraceServices::IAnalysisSession* Session = GetAnalysisSession())
+	{
 		if (const IGameplayProvider* GameplayProvider = Session->ReadProvider<IGameplayProvider>("GameplayProvider"))
 		{
-			if (DebugTargetActor.Get() == "")
-			{
-				return;
-			}
-
-			uint64 TargetActorId = 0;
-
-			// find the Id of the selected actor (move this to a function or cache it)
 			GameplayProvider->EnumerateObjects([this,&TargetActorId](const FObjectInfo& InObjectInfo)
 			{
 				if (DebugTargetActor.Get() == InObjectInfo.Name)
@@ -215,6 +227,27 @@ void FRewindDebugger::RefreshDebugComponents()
 					TargetActorId = InObjectInfo.Id;
 				}
 			});
+		}
+	}
+
+	return TargetActorId;
+}
+
+void FRewindDebugger::RefreshDebugComponents()
+{
+	if (const TraceServices::IAnalysisSession* Session = GetAnalysisSession())
+	{
+		TraceServices::FAnalysisSessionReadScope SessionReadScope(*Session);
+		UWorld* World = GetWorldToVisualize();
+
+		if (const IGameplayProvider* GameplayProvider = Session->ReadProvider<IGameplayProvider>("GameplayProvider"))
+		{
+			uint64 TargetActorId = GetTargetActorId();
+
+			if (TargetActorId == 0)
+			{
+				return;
+			}
 
 			bool bChanged = false;
 
@@ -462,15 +495,11 @@ void FRewindDebugger::UpdateTraceTime()
 	}
 }
 
-const TraceServices::IAnalysisSession* FRewindDebugger::GetAnalysisSession()
+const TraceServices::IAnalysisSession* FRewindDebugger::GetAnalysisSession() const
 {
 	if (UnrealInsightsModule == nullptr)
 	{
 		UnrealInsightsModule = &FModuleManager::LoadModuleChecked<IUnrealInsightsModule>("TraceInsights");
-	    if (UnrealInsightsModule == nullptr)
-		{
-			return nullptr; 
-		}
 	}
 
 	return UnrealInsightsModule->GetAnalysisSession().Get();
@@ -520,13 +549,21 @@ void FRewindDebugger::Tick(float DeltaTime)
 					}
 
 					double CurrentTraceTime = TraceTime.Get();
+					uint64 TargetActorId = GetTargetActorId();
 
-					// update pose on all SkeletalMeshComponents
+					// update pose on all SkeletalMeshComponents (todo: move pose updating into an extension)
 					ULevel* CurLevel = World->GetCurrentLevel();
 					for( AActor* LevelActor : CurLevel->Actors )
 					{
 						if (LevelActor)
 						{
+							bool bIsTargetActor = TargetActorId == FObjectTrace::GetObjectId(LevelActor);
+
+							if (bIsTargetActor)
+							{
+								bTargetActorPositionValid = false;
+							}
+
 							TInlineComponentArray<USkeletalMeshComponent*> SkeletalMeshComponents;
 							LevelActor->GetComponents(SkeletalMeshComponents);
 
@@ -534,7 +571,7 @@ void FRewindDebugger::Tick(float DeltaTime)
 							{
 								int64 ObjectId = FObjectTrace::GetObjectId(MeshComponent);
 
-								AnimationProvider->ReadSkeletalMeshPoseTimeline(ObjectId, [this, ObjectId, CurrentTraceTime, MeshComponent, AnimationProvider](const IAnimationProvider::SkeletalMeshPoseTimeline& TimelineData, bool bHasCurves)
+								AnimationProvider->ReadSkeletalMeshPoseTimeline(ObjectId, [this, bIsTargetActor, ObjectId, CurrentTraceTime, MeshComponent, AnimationProvider](const IAnimationProvider::SkeletalMeshPoseTimeline& TimelineData, bool bHasCurves)
 								{
 									double PrecedingPoseTime;
 									double FollowingPoseTime;
@@ -563,6 +600,13 @@ void FRewindDebugger::Tick(float DeltaTime)
 										// todo: we probably need to take into account tick order requirements for attached objects here
 										MeshComponent->SetWorldTransform(ComponentWorldTransform);
 										MeshComponent->SetForcedLOD(PoseMessage.LodIndex + 1);
+
+										if (bIsTargetActor && !bTargetActorPositionValid)
+										{
+											// until we have actor transforms traced out, the first skeletal mesh component transform will be used as as the actor position 
+											bTargetActorPositionValid = true;
+											TargetActorPosition = ComponentWorldTransform.GetTranslation();
+										}
 									}
 								});
 
@@ -572,6 +616,16 @@ void FRewindDebugger::Tick(float DeltaTime)
 							}
 						}
 					}
+
+					// update extensions
+					IModularFeatures& ModularFeatures = IModularFeatures::Get();
+					
+					const int32 NumExtensions = ModularFeatures.GetModularFeatureImplementationCount( IRewindDebuggerExtension::ModularFeatureName );
+					for (int32 ExtensionIndex = 0; ExtensionIndex < NumExtensions; ++ExtensionIndex)
+					{
+						IRewindDebuggerExtension* Extension = static_cast<IRewindDebuggerExtension *>(ModularFeatures.GetModularFeatureImplementation(IRewindDebuggerExtension::ModularFeatureName, ExtensionIndex));
+						Extension->UpdatePlayback(DeltaTime, this);
+					}
 				}
 			}
 		}
diff --git a/Engine/Plugins/Animation/GameplayInsights/Source/RewindDebugger/Private/RewindDebugger.h b/Engine/Plugins/Animation/GameplayInsights/Source/RewindDebugger/Private/RewindDebugger.h
--- a/Engine/Plugins/Animation/GameplayInsights/Source/RewindDebugger/Private/RewindDebugger.h
+++ b/Engine/Plugins/Animation/GameplayInsights/Source/RewindDebugger/Private/RewindDebugger.h
@@ -2,6 +2,7 @@
 #pragma once
 
 #include "CoreMinimal.h"
+#include "IRewindDebugger.h"
 #include "BindableProperty.h"
 #include "UObject/WeakObjectPtr.h"
 #include "RewindDebuggerModule.h"
@@ -18,11 +19,18 @@ namespace TraceServices
 //  Start/Stop recording
 //  Keeping track of the current Debug Target actor, and outputing a list of it's Components for the UI
 
-class FRewindDebugger
+class FRewindDebugger : public IRewindDebugger
 {
 public:
 	FRewindDebugger();
-	~FRewindDebugger();
+	virtual ~FRewindDebugger();
+
+	// IRewindDebugger interface
+	virtual double CurrentTraceTime() const override { return TraceTime.Get(); }
+	virtual const TraceServices::IAnalysisSession* GetAnalysisSession() const override;
+	virtual uint64 GetTargetActorId() const override;
+	virtual bool GetTargetActorPosition(FVector& OutPosition) const override;
+	virtual UWorld* GetWorldToVisualize() const override;
 
 	// create singleton instance
 	static void Initialize();
@@ -92,8 +100,6 @@ private:
 	void OnPIEStopped(bool bSimulating);
 	void OnPIESingleStepped(bool bSimulating);
 
-	const TraceServices::IAnalysisSession* GetAnalysisSession();
-	UWorld* GetWorldToVisualize() const;
 
 	bool UpdateComponentList(uint64 ParentId, TArray<TSharedPtr<FDebugObjectInfo>>& NewComponentList);
 	void SetCurrentScrubTime(float Time);
@@ -137,6 +143,9 @@ private:
 
 	TMap<uint64, FMeshComponentResetData> MeshComponentsToReset;
 
-	class IUnrealInsightsModule *UnrealInsightsModule;
+	mutable class IUnrealInsightsModule *UnrealInsightsModule;
 	FDelegateHandle TickerHandle;
+
+	bool bTargetActorPositionValid;
+	FVector TargetActorPosition;
 };
diff --git a/Engine/Plugins/Animation/GameplayInsights/Source/RewindDebugger/Private/RewindDebuggerCamera.cpp b/Engine/Plugins/Animation/GameplayInsights/Source/RewindDebugger/Private/RewindDebuggerCamera.cpp
new file mode 100644
--- /dev/null
+++ b/Engine/Plugins/Animation/GameplayInsights/Source/RewindDebugger/Private/RewindDebuggerCamera.cpp
@@ -0,0 +1,160 @@
+// Copyright Epic Games, Inc. All Rights Reserved.
+
+#include "RewindDebuggerCamera.h"
+#include "Camera/CameraActor.h"
+#include "Camera/CameraComponent.h"
+#include "Editor.h"
+#include "IGameplayProvider.h"
+#include "Insights/IUnrealInsightsModule.h"
+#include "IRewindDebugger.h"
+#include "LevelEditor.h"
+#include "Modules/ModuleManager.h"
+#include "SLevelViewport.h"
+#include "ToolMenu.h"
+#include "ToolMenus.h"
+
+#define LOCTEXT_NAMESPACE "RewindDebuggerCamera"
+
+FRewindDebuggerCamera::FRewindDebuggerCamera()
+	: LastPositionValid(false)
+{
+	UToolMenu* Menu = UToolMenus::Get()->RegisterMenu("RewindDebugger.MainMenu");
+
+	Menu->AddSection("Camera Mode", LOCTEXT("Camera Mode","Camera Mode"));
+	Menu->AddMenuEntry("Camera Mode",
+			FToolMenuEntry::InitMenuEntry("CameraModeDisabled",
+				LOCTEXT("Camera Mode Dosan;ed", "Disabled"),
+				FText(),
+				FSlateIcon(),
+				FUIAction(
+					FExecuteAction::CreateRaw(this, &FRewindDebuggerCamera::SetCameraMode, ECameraMode::Disabled),
+					FCanExecuteAction(),
+					FGetActionCheckState::CreateLambda([this] { return Mode == ECameraMode::Disabled ? ECheckBoxState::Checked : ECheckBoxState::Unchecked; } )),
+				EUserInterfaceActionType::Check
+				)
+			);
+
+	Menu->AddMenuEntry("Camera Mode",
+			FToolMenuEntry::InitMenuEntry("CameraModeFollow",
+				LOCTEXT("Camera Mode Follow", "Follow Target Actor"),
+				FText(),
+				FSlateIcon(),
+				FUIAction( FExecuteAction::CreateRaw(this, &FRewindDebuggerCamera::SetCameraMode, ECameraMode::FollowTargetActor),
+					FCanExecuteAction(),
+					FGetActionCheckState::CreateLambda([this] { return Mode == ECameraMode::FollowTargetActor ? ECheckBoxState::Checked : ECheckBoxState::Unchecked; } )),
+				EUserInterfaceActionType::Check
+				)
+			);
+
+	Menu->AddMenuEntry("Camera Mode",
+			FToolMenuEntry::InitMenuEntry("CameraModeReplay",
+				LOCTEXT("Camera Mode Recorded", "Replay Recorded Camera"),
+				FText(),
+				FSlateIcon(),
+				FUIAction( FExecuteAction::CreateRaw(this, &FRewindDebuggerCamera::SetCameraMode, ECameraMode::Replay),
+					FCanExecuteAction(),
+					FGetActionCheckState::CreateLambda([this] { return Mode == ECameraMode::Replay ? ECheckBoxState::Checked : ECheckBoxState::Unchecked; } )),
+				EUserInterfaceActionType::Check
+				)
+			);
+}
+
+void FRewindDebuggerCamera::SetCameraMode(ECameraMode InMode)
+{
+	FLevelEditorModule& LevelEditor = FModuleManager::GetModuleChecked<FLevelEditorModule>("LevelEditor");
+	SLevelViewport* LevelViewport = LevelEditor.GetFirstActiveLevelViewport().Get();
+	FLevelEditorViewportClient& LevelViewportClient = LevelViewport->GetLevelViewportClient();
+
+	if (Mode == ECameraMode::Replay && InMode != ECameraMode::Replay)
+	{
+		LevelViewportClient.SetActorLock(nullptr);
+	}
+	else if (Mode == ECameraMode::Replay)
+	{
+		if (CameraActor.IsValid())
+		{
+			LevelViewportClient.SetActorLock(CameraActor.Get());
+		}
+	}
+
+	Mode = InMode;
+}
+
+void FRewindDebuggerCamera::UpdatePlayback(float DeltaTime, IRewindDebugger* RewindDebugger)
+{
+	if (const TraceServices::IAnalysisSession* Session = RewindDebugger->GetAnalysisSession())
+	{
+		double CurrentTraceTime = RewindDebugger->CurrentTraceTime();
+		
+		static double LastCameraScrubTime = 0.0f;
+		if (CurrentTraceTime != LastCameraScrubTime)
+		{
+			FLevelEditorModule& LevelEditor = FModuleManager::GetModuleChecked<FLevelEditorModule>("LevelEditor");
+			SLevelViewport* LevelViewport = LevelEditor.GetFirstActiveLevelViewport().Get();
+			FLevelEditorViewportClient& LevelViewportClient = LevelViewport->GetLevelViewportClient();
+
+			FVector TargetActorPosition;
+			bool bTargetActorPositionValid = RewindDebugger->GetTargetActorPosition(TargetActorPosition);
+
+			// only update camera in playback or scrubbing when the time has changed (allow free movement when paused)
+			LastCameraScrubTime = CurrentTraceTime;
+
+			if (Mode == ECameraMode::FollowTargetActor)
+			{
+				// Follow Actor mode: apply position changes from the target actor to the camera
+				if (bTargetActorPositionValid)
+				{
+					if(LastPositionValid)
+					{
+						LevelViewportClient.SetViewLocation(LevelViewportClient.GetViewLocation() + TargetActorPosition - LastPosition);
+					}
+				}
+			}
+
+			// always update the camera actor to the replay values even if it isn't locked
+			if (const IGameplayProvider* GameplayProvider = Session->ReadProvider<IGameplayProvider>("GameplayProvider"))
+			{
+				GameplayProvider->ReadViewTimeline([this, RewindDebugger, CurrentTraceTime](const IGameplayProvider::ViewTimeline& TimelineData)
+				{
+					double PrecedingTime;
+					double FollowingTime;
+					const FViewMessage* PrecedingView;
+					const FViewMessage* FollowingView;
+					
+					TimelineData.FindNearestEvents(CurrentTraceTime, PrecedingView, PrecedingTime, FollowingView, FollowingTime);
+
+					if (FollowingView || PrecedingView)
+					{
+						const FViewMessage& ViewMessage = FollowingView ? *FollowingView : *PrecedingView;
+
+						if (!CameraActor.IsValid())
+						{
+							FActorSpawnParameters SpawnParameters;
+							SpawnParameters.Name = "RewindDebuggerCamera"; 
+							CameraActor = RewindDebugger->GetWorldToVisualize()->SpawnActor<ACameraActor>(ViewMessage.Position, ViewMessage.Rotation, SpawnParameters);
+							UCameraComponent* Camera = CameraActor->GetCameraComponent();
+						}
+
+						UCameraComponent* Camera = CameraActor->GetCameraComponent();
+						Camera->SetWorldLocationAndRotation(ViewMessage.Position, ViewMessage.Rotation);
+						Camera->SetFieldOfView(ViewMessage.Fov);
+						Camera->SetAspectRatio(ViewMessage.AspectRatio);
+					}
+				});
+			}
+
+			if (Mode == ECameraMode::Replay)
+			{
+				if (CameraActor.IsValid())
+				{
+					LevelViewportClient.SetActorLock(CameraActor.Get());
+				}
+			}
+
+			LastPosition = TargetActorPosition;
+			LastPositionValid = bTargetActorPositionValid;
+		}
+	}
+}
+
+#undef LOCTEXT_NAMESPACE
\ No newline at end of file
diff --git a/Engine/Plugins/Animation/GameplayInsights/Source/RewindDebugger/Private/RewindDebuggerCamera.h b/Engine/Plugins/Animation/GameplayInsights/Source/RewindDebugger/Private/RewindDebuggerCamera.h
new file mode 100644
--- /dev/null
+++ b/Engine/Plugins/Animation/GameplayInsights/Source/RewindDebugger/Private/RewindDebuggerCamera.h
@@ -0,0 +1,38 @@
+// Copyright Epic Games, Inc. All Rights Reserved.
+#pragma once
+
+#include "CoreMinimal.h"
+#include "IRewindDebuggerExtension.h"
+#include "UObject/WeakObjectPtr.h"
+
+class ACameraActor;
+
+// Rewind debugger extension for camera support
+//  replay of recorded camera data
+//  follow selected actor
+
+class FRewindDebuggerCamera : public IRewindDebuggerExtension
+{
+public:
+	enum class ECameraMode
+	{
+		Replay,
+		FollowTargetActor,
+		Disabled,
+	};
+
+
+	FRewindDebuggerCamera();
+	virtual ~FRewindDebuggerCamera() {};
+
+	virtual void UpdatePlayback(float DeltaTime, IRewindDebugger* RewindDebugger) override;
+
+	ECameraMode CameraMode();
+	void SetCameraMode(ECameraMode Mode);
+private:
+	bool LastPositionValid;
+	FVector LastPosition;
+
+	ECameraMode Mode = ECameraMode::Replay;
+	TWeakObjectPtr<ACameraActor> CameraActor; 
+};
diff --git a/Engine/Plugins/Animation/GameplayInsights/Source/RewindDebugger/Private/RewindDebuggerModule.cpp b/Engine/Plugins/Animation/GameplayInsights/Source/RewindDebugger/Private/RewindDebuggerModule.cpp
--- a/Engine/Plugins/Animation/GameplayInsights/Source/RewindDebugger/Private/RewindDebuggerModule.cpp
+++ b/Engine/Plugins/Animation/GameplayInsights/Source/RewindDebugger/Private/RewindDebuggerModule.cpp
@@ -126,10 +126,13 @@ void FRewindDebuggerModule::StartupModule()
 		.SetTooltipText(LOCTEXT("TooltipText", "Opens Rewind Debugger."));
 
 
+	IModularFeatures::Get().RegisterModularFeature(IRewindDebuggerExtension::ModularFeatureName, &RewindDebuggerCameraExtension);
 }
 
 void FRewindDebuggerModule::ShutdownModule()
 {
+	IModularFeatures::Get().UnregisterModularFeature(IRewindDebuggerExtension::ModularFeatureName, &RewindDebuggerCameraExtension);
+
 	FRewindDebuggerCommands::Unregister();
 	FRewindDebuggerStyle::Shutdown();
 	FRewindDebugger::Shutdown();
diff --git a/Engine/Plugins/Animation/GameplayInsights/Source/RewindDebugger/Private/RewindDebuggerModule.h b/Engine/Plugins/Animation/GameplayInsights/Source/RewindDebugger/Private/RewindDebuggerModule.h
--- a/Engine/Plugins/Animation/GameplayInsights/Source/RewindDebugger/Private/RewindDebuggerModule.h
+++ b/Engine/Plugins/Animation/GameplayInsights/Source/RewindDebugger/Private/RewindDebuggerModule.h
@@ -5,6 +5,7 @@
 #include "CoreMinimal.h"
 #include "Modules/ModuleInterface.h"
 #include "Framework/Docking/TabManager.h"
+#include "RewindDebuggerCamera.h"
 
 class SRewindDebugger;
 
@@ -32,4 +33,6 @@ public:
 
 private:
 	TSharedPtr<SRewindDebugger> RewindDebuggerWidget;
+
+	FRewindDebuggerCamera RewindDebuggerCameraExtension;
 };
diff --git a/Engine/Plugins/Animation/GameplayInsights/Source/RewindDebugger/Private/SRewindDebugger.cpp b/Engine/Plugins/Animation/GameplayInsights/Source/RewindDebugger/Private/SRewindDebugger.cpp
--- a/Engine/Plugins/Animation/GameplayInsights/Source/RewindDebugger/Private/SRewindDebugger.cpp
+++ b/Engine/Plugins/Animation/GameplayInsights/Source/RewindDebugger/Private/SRewindDebugger.cpp
@@ -21,6 +21,8 @@
 #include "SceneOutlinerModule.h"
 #include "Selection.h"
 #include "Styling/SlateIconFinder.h"
+#include "ToolMenu.h"
+#include "ToolMenus.h"
 #include "Widgets/Docking/SDockTab.h"
 #include "Widgets/Images/SImage.h"
 #include "Widgets/Input/SComboButton.h"
@@ -287,6 +289,15 @@ void SRewindDebugger::Construct(const FArguments& InArgs, TSharedRef<FUICommandL
 	// load saved layout if it exists
 	Layout = FLayoutSaveRestore::LoadFromConfig(GEditorLayoutIni, Layout);
 
+	UToolMenu* Menu = UToolMenus::Get()->RegisterMenu("RewindDebugger.MainMenu");
+
+	FToolMenuSection& Section = Menu->AddSection("ViewsSection", LOCTEXT("Views", "Views"));
+
+	Section.AddDynamicEntry("ViewsSection", FNewToolMenuDelegateLegacy::CreateLambda([this](FMenuBuilder& InMenuBuilder, UToolMenu* InMenu)
+	{
+		MakeViewsMenu(InMenuBuilder);
+	}));
+
 
 	ChildSlot
 	[
@@ -515,14 +526,11 @@ void SRewindDebugger::OnPinnedTabClosed(TSharedRef<SDockTab> Tab)
 
 TSharedRef<SWidget> SRewindDebugger::MakeMainMenu()
 {
-	FMenuBuilder MenuBuilder(true, nullptr);
-	MakeViewsMenu(MenuBuilder);
-	return MenuBuilder.MakeWidget();
+	return UToolMenus::Get()->GenerateWidget("RewindDebugger.MainMenu", FToolMenuContext());
 }
 
 void SRewindDebugger::MakeViewsMenu(FMenuBuilder& MenuBuilder)
 {
-	MenuBuilder.BeginSection("Views", LOCTEXT("Views", "Views"));
 
 	TabManager->PopulateLocalTabSpawnerMenu(MenuBuilder);
 
@@ -530,8 +538,6 @@ void SRewindDebugger::MakeViewsMenu(FMenuBuilder& MenuBuilder)
 							 LOCTEXT("Show All tooltip", "Show all debug views that are relevant to the selected object type"),
 							 FSlateIcon(),
 							 FExecuteAction::CreateLambda([this]() { ShowAllViews(); }));
-
-	MenuBuilder.EndSection();
 }
 
 void SRewindDebugger::ExtendTabMenu(FMenuBuilder& MenuBuilder, TSharedPtr<IGameplayInsightsDebugView> View)
@@ -548,8 +554,6 @@ void SRewindDebugger::ExtendTabMenu(FMenuBuilder& MenuBuilder, TSharedPtr<IGamep
 
 
 	MenuBuilder.EndSection();
-
-	MakeViewsMenu(MenuBuilder);
 }
 
 void SRewindDebugger::ShowAllViews()
@@ -630,6 +634,4 @@ FReply SRewindDebugger::OnSelectActorClicked()
 	return FReply::Handled();
 }
 
-
-
 #undef LOCTEXT_NAMESPACE
diff --git a/Engine/Plugins/Animation/GameplayInsights/Source/RewindDebugger/Public/IRewindDebugger.h b/Engine/Plugins/Animation/GameplayInsights/Source/RewindDebugger/Public/IRewindDebugger.h
new file mode 100644
--- /dev/null
+++ b/Engine/Plugins/Animation/GameplayInsights/Source/RewindDebugger/Public/IRewindDebugger.h
@@ -0,0 +1,32 @@
+// Copyright Epic Games, Inc. All Rights Reserved.
+#pragma once
+
+#include "CoreMinimal.h"
+
+// IRewindDebugger
+//
+// Public interface to rewind debugger
+
+namespace TraceServices
+{
+	class IAnalysisSession;
+}
+
+class IRewindDebugger
+{
+public:
+	// get the time the debugger is scrubbed to, in seconds since the capture started (or the recording duration while the game is running)
+	virtual double CurrentTraceTime() const = 0;
+
+	// get the current analysis session
+	virtual const TraceServices::IAnalysisSession* GetAnalysisSession() const = 0;
+
+	// get insights id for the selected target actor
+	virtual uint64 GetTargetActorId() const = 0;
+
+	// get posiotion of the selected target actor (returns true if position is valid)
+	virtual bool GetTargetActorPosition(FVector& OutPosition) const = 0;
+
+	// get the world that the debugger is replaying in
+	virtual UWorld* GetWorldToVisualize() const = 0;
+};
diff --git a/Engine/Plugins/Animation/GameplayInsights/Source/RewindDebugger/Public/IRewindDebuggerExtension.h b/Engine/Plugins/Animation/GameplayInsights/Source/RewindDebugger/Public/IRewindDebuggerExtension.h
new file mode 100644
--- /dev/null
+++ b/Engine/Plugins/Animation/GameplayInsights/Source/RewindDebugger/Public/IRewindDebuggerExtension.h
@@ -0,0 +1,23 @@
+// Copyright Epic Games, Inc. All Rights Reserved.
+#pragma once
+
+#include "CoreMinimal.h"
+#include "Features/IModularFeatures.h"
+
+
+// IRewindDebuggerExtension 
+//
+// interface class for extensions which add functionality to the rewind debugger
+// these will get updated on scrub/playback to handle things like updating the world state to match recorded data from that time for a particular system
+//
+
+class IRewindDebugger;
+
+class IRewindDebuggerExtension : public IModularFeature
+{
+public:
+	static REWINDDEBUGGER_API const FName ModularFeatureName;
+
+	// called while scrubbing, playing back, or paused
+	virtual void UpdatePlayback(float DeltaTime, IRewindDebugger* RewindDebugger) {};
+};
diff --git a/Engine/Plugins/Animation/GameplayInsights/Source/RewindDebugger/RewindDebugger.Build.cs b/Engine/Plugins/Animation/GameplayInsights/Source/RewindDebugger/RewindDebugger.Build.cs
--- a/Engine/Plugins/Animation/GameplayInsights/Source/RewindDebugger/RewindDebugger.Build.cs
+++ b/Engine/Plugins/Animation/GameplayInsights/Source/RewindDebugger/RewindDebugger.Build.cs
@@ -23,6 +23,7 @@ namespace UnrealBuildTool.Rules
 				"GameplayInsightsEditor",
 				"InputCore",
 				"Kismet",
+				"LevelEditor",
 				"Persona",
 				"SceneOutliner",
 				"SequencerWidgets",
diff --git a/Engine/Source/Runtime/Engine/Private/GameViewportClient.cpp b/Engine/Source/Runtime/Engine/Private/GameViewportClient.cpp
--- a/Engine/Source/Runtime/Engine/Private/GameViewportClient.cpp
+++ b/Engine/Source/Runtime/Engine/Private/GameViewportClient.cpp
@@ -63,6 +63,7 @@
 #include "IImageWrapperModule.h"
 #include "HAL/PlatformApplicationMisc.h"
 #include "CustomStaticScreenPercentage.h"
+#include "ObjectTrace.h"
 
 #if WITH_EDITOR
 #include "Settings/LevelEditorPlaySettings.h"
@@ -1561,6 +1562,15 @@ void UGameViewportClient::Draw(FViewport* InViewport, FCanvas* SceneCanvas)
 	UpdateCsvCameraStats(PlayerViewMap);
 #endif
 
+#if OBJECT_TRACE_ENABLED 
+	for (TMap<ULocalPlayer*, FSceneView*>::TConstIterator It(PlayerViewMap); It; ++It)
+	{
+		ULocalPlayer* LocalPlayer = It.Key();
+		FSceneView* SceneView = It.Value();
+		TRACE_VIEW(LocalPlayer, SceneView);	
+	}
+#endif
+
 	FinalizeViews(&ViewFamily, PlayerViewMap);
 
 	// Update level streaming.
diff --git a/Engine/Source/Runtime/Engine/Private/ObjectTrace.cpp b/Engine/Source/Runtime/Engine/Private/ObjectTrace.cpp
--- a/Engine/Source/Runtime/Engine/Private/ObjectTrace.cpp
+++ b/Engine/Source/Runtime/Engine/Private/ObjectTrace.cpp
@@ -16,6 +16,7 @@
 #include "HAL/IConsoleManager.h"
 #include "Engine/World.h"
 #include "TraceFilter.h"
+#include "SceneView.h"
 #if WITH_EDITOR
 #include "Editor.h"
 #endif
@@ -57,6 +58,22 @@ UE_TRACE_EVENT_BEGIN(Object, RecordingInfo)
 	UE_TRACE_EVENT_FIELD(double, ElapsedTime)
 UE_TRACE_EVENT_END()
 
+UE_TRACE_EVENT_BEGIN(Object, View)
+	UE_TRACE_EVENT_FIELD(uint64, PlayerId)
+	UE_TRACE_EVENT_FIELD(uint64, Cycle)
+
+	UE_TRACE_EVENT_FIELD(double, PosX)
+	UE_TRACE_EVENT_FIELD(double, PosY)
+	UE_TRACE_EVENT_FIELD(double, PosZ)
+
+	UE_TRACE_EVENT_FIELD(float, Pitch)
+	UE_TRACE_EVENT_FIELD(float, Yaw)
+	UE_TRACE_EVENT_FIELD(float, Roll)
+
+	UE_TRACE_EVENT_FIELD(float, Fov)
+	UE_TRACE_EVENT_FIELD(float, AspectRatio)
+UE_TRACE_EVENT_END()
+
 // Object annotations used for tracing
 struct FTracedObjectAnnotation
 {
@@ -289,6 +306,35 @@ void FObjectTrace::OutputClass(const UClass* InClass)
 
 }
 
+void FObjectTrace::OutputView(const UObject* InPlayer, const FSceneView* InView)
+{
+	bool bChannelEnabled = UE_TRACE_CHANNELEXPR_IS_ENABLED(ObjectChannel);
+	if (!bChannelEnabled || InPlayer == nullptr)
+	{
+		return;
+	}
+
+	if (CANNOT_TRACE_OBJECT(InPlayer->GetWorld()))
+	{
+		return;
+	}
+
+	const FIntRect& ViewRect = InView->CameraConstrainedViewRect;
+	float AspectRatio = (float)ViewRect.Width()/(float)ViewRect.Height();
+
+	UE_TRACE_LOG(Object, View, ObjectChannel)
+		<< View.Cycle(FPlatformTime::Cycles64())
+		<< View.PlayerId(GetObjectId(InPlayer))
+		<< View.PosX(InView->ViewLocation.X)
+		<< View.PosY(InView->ViewLocation.Y)
+		<< View.PosZ(InView->ViewLocation.Z)
+		<< View.Pitch(InView->ViewRotation.Pitch)
+		<< View.Yaw(InView->ViewRotation.Yaw)
+		<< View.Roll(InView->ViewRotation.Roll)
+		<< View.Fov(InView->FOV)
+		<< View.AspectRatio(AspectRatio);
+}
+
 void FObjectTrace::OutputObject(const UObject* InObject)
 {
 	bool bChannelEnabled = UE_TRACE_CHANNELEXPR_IS_ENABLED(ObjectChannel);
diff --git a/Engine/Source/Runtime/Engine/Public/ObjectTrace.h b/Engine/Source/Runtime/Engine/Public/ObjectTrace.h
--- a/Engine/Source/Runtime/Engine/Public/ObjectTrace.h
+++ b/Engine/Source/Runtime/Engine/Public/ObjectTrace.h
@@ -39,6 +39,7 @@ public:
 class UClass;
 class UObject;
 class UWorld;
+class FSceneView;
 
 struct FObjectTrace
 {
@@ -54,6 +55,9 @@ struct FObjectTrace
 	/** Helper function to output an object */
 	ENGINE_API static void OutputObject(const UObject* InObject);
 
+	/** Helper function to output camera information for a player */
+	ENGINE_API static void OutputView(const UObject* LocalPlayer, const FSceneView* View);
+
 	/** Helper function to output an object event */
 	ENGINE_API static void OutputObjectEvent(const UObject* InObject, const TCHAR* InEvent);
 
@@ -95,6 +99,9 @@ struct FObjectTrace
 #define TRACE_OBJECT(Object) \
 	FObjectTrace::OutputObject(Object);
 
+#define TRACE_VIEW(Player, View) \
+	FObjectTrace::OutputView(Player, View);
+
 #if TRACE_FILTERING_ENABLED
 
 #define TRACE_OBJECT_EVENT(Object, Event) \
@@ -119,5 +126,6 @@ struct FObjectTrace
 #define TRACE_OBJECT(Object)
 #define TRACE_OBJECT_EVENT(Object, Event)
 #define TRACE_WORLD(World)
+#define TRACE_VIEW(Player, View)
 
 #endif
\ No newline at end of file
