Rewind Debugger interface changes to support visual logger prototype

diff --git a/Engine/Plugins/Animation/GameplayInsights/Source/RewindDebugger/Private/RewindDebugger.cpp b/Engine/Plugins/Animation/GameplayInsights/Source/RewindDebugger/Private/RewindDebugger.cpp
--- a/Engine/Plugins/Animation/GameplayInsights/Source/RewindDebugger/Private/RewindDebugger.cpp
+++ b/Engine/Plugins/Animation/GameplayInsights/Source/RewindDebugger/Private/RewindDebugger.cpp
@@ -17,6 +17,19 @@
 FRewindDebugger* FRewindDebugger::InternalInstance = nullptr;
 const FName IRewindDebuggerExtension::ModularFeatureName = "RewindDebuggerExtension";
 
+static void IterateExtensions(TFunction<void(IRewindDebuggerExtension* Extension)> IteratorFunction)
+{
+	// update extensions
+	IModularFeatures& ModularFeatures = IModularFeatures::Get();
+
+	const int32 NumExtensions = ModularFeatures.GetModularFeatureImplementationCount(IRewindDebuggerExtension::ModularFeatureName);
+	for (int32 ExtensionIndex = 0; ExtensionIndex < NumExtensions; ++ExtensionIndex)
+	{
+		IRewindDebuggerExtension* Extension = static_cast<IRewindDebuggerExtension*>(ModularFeatures.GetModularFeatureImplementation(IRewindDebuggerExtension::ModularFeatureName, ExtensionIndex));
+		IteratorFunction(Extension);
+	}
+}
+
 FRewindDebugger::FRewindDebugger()  :
 	ControlState(FRewindDebugger::EControlState::Pause),
 	bPIEStarted(false),
@@ -291,6 +304,13 @@ void FRewindDebugger::StartRecording()
 	UE::Trace::ToggleChannel(TEXT("Animation"), true);
 	UE::Trace::ToggleChannel(TEXT("Frame"), true);
 
+	// update extensions
+	IterateExtensions([this](IRewindDebuggerExtension* Extension)
+		{
+			Extension->RecordingStarted(this);
+		}
+	);
+
 	RecordingDuration.Set(0);
 	RecordingIndex++;
 	bRecording = true;
@@ -311,6 +331,13 @@ void FRewindDebugger::StopRecording()
 		UE::Trace::ToggleChannel(TEXT("Animation"), false);
 		UE::Trace::ToggleChannel(TEXT("Frame"), false);
 
+		// update extensions
+		IterateExtensions([this](IRewindDebuggerExtension* Extension)
+			{
+				Extension->RecordingStopped(this);
+			}
+		);
+
 		bRecording = false;
 	}
 }
@@ -616,18 +643,15 @@ void FRewindDebugger::Tick(float DeltaTime)
 							}
 						}
 					}
-
-					// update extensions
-					IModularFeatures& ModularFeatures = IModularFeatures::Get();
-					
-					const int32 NumExtensions = ModularFeatures.GetModularFeatureImplementationCount( IRewindDebuggerExtension::ModularFeatureName );
-					for (int32 ExtensionIndex = 0; ExtensionIndex < NumExtensions; ++ExtensionIndex)
-					{
-						IRewindDebuggerExtension* Extension = static_cast<IRewindDebuggerExtension *>(ModularFeatures.GetModularFeatureImplementation(IRewindDebuggerExtension::ModularFeatureName, ExtensionIndex));
-						Extension->UpdatePlayback(DeltaTime, this);
-					}
 				}
 			}
 		}
+
+		// update extensions
+		IterateExtensions([DeltaTime, this](IRewindDebuggerExtension* Extension)
+			{
+				Extension->Update(DeltaTime, this);
+			}
+		);
 	}
 }
diff --git a/Engine/Plugins/Animation/GameplayInsights/Source/RewindDebugger/Private/RewindDebugger.h b/Engine/Plugins/Animation/GameplayInsights/Source/RewindDebugger/Private/RewindDebugger.h
--- a/Engine/Plugins/Animation/GameplayInsights/Source/RewindDebugger/Private/RewindDebugger.h
+++ b/Engine/Plugins/Animation/GameplayInsights/Source/RewindDebugger/Private/RewindDebugger.h
@@ -5,8 +5,8 @@
 #include "IRewindDebugger.h"
 #include "BindableProperty.h"
 #include "UObject/WeakObjectPtr.h"
-#include "RewindDebuggerModule.h"
 
+class USkeletalMeshComponent;
 
 namespace TraceServices
 {
@@ -31,6 +31,9 @@ public:
 	virtual uint64 GetTargetActorId() const override;
 	virtual bool GetTargetActorPosition(FVector& OutPosition) const override;
 	virtual UWorld* GetWorldToVisualize() const override;
+	virtual bool IsRecording() const override { return bRecording; }
+	virtual bool IsPIESimulating() const override { return bPIESimulating; }
+	virtual double GetRecordingDuration() const override { return RecordingDuration.Get(); }
 
 	// create singleton instance
 	static void Initialize();
@@ -44,7 +47,6 @@ public:
 	// Start a new Recording:  Start tracing Object + Animation data, increment the current recording index, and reset the recording elapsed time to 0
 	void StartRecording();
 
-	bool IsRecording() const { return bRecording; }
 	bool CanStartRecording() const { return !IsRecording() && bPIESimulating; }
 
 	bool AutoRecord() const { return bAutoRecord; }
@@ -81,7 +83,7 @@ public:
 	// update the list of components for the currently selected debug target
 	void RefreshDebugComponents();
 
-	TArray<TSharedPtr<FDebugObjectInfo>>& GetDebugComponents() { return DebugComponents; };
+	TArray<TSharedPtr<FDebugObjectInfo>>& GetDebugComponents() override { return DebugComponents; };
 
 	DECLARE_DELEGATE(FOnComponentListChanged)
 	void OnComponentListChanged(const FOnComponentListChanged& ComponentListChangedCallback);
diff --git a/Engine/Plugins/Animation/GameplayInsights/Source/RewindDebugger/Private/RewindDebuggerCamera.cpp b/Engine/Plugins/Animation/GameplayInsights/Source/RewindDebugger/Private/RewindDebuggerCamera.cpp
--- a/Engine/Plugins/Animation/GameplayInsights/Source/RewindDebugger/Private/RewindDebuggerCamera.cpp
+++ b/Engine/Plugins/Animation/GameplayInsights/Source/RewindDebugger/Private/RewindDebuggerCamera.cpp
@@ -80,10 +80,16 @@ void FRewindDebuggerCamera::SetCameraMode(ECameraMode InMode)
 	Mode = InMode;
 }
 
-void FRewindDebuggerCamera::UpdatePlayback(float DeltaTime, IRewindDebugger* RewindDebugger)
+void FRewindDebuggerCamera::Update(float DeltaTime, IRewindDebugger* RewindDebugger)
 {
+	if (RewindDebugger->IsPIESimulating() || RewindDebugger->GetRecordingDuration() == 0.0)
+	{
+		return;
+	}
+
 	if (const TraceServices::IAnalysisSession* Session = RewindDebugger->GetAnalysisSession())
 	{
+		TraceServices::FAnalysisSessionReadScope SessionReadScope(*Session);
 		double CurrentTraceTime = RewindDebugger->CurrentTraceTime();
 		
 		static double LastCameraScrubTime = 0.0f;
@@ -130,9 +136,10 @@ void FRewindDebuggerCamera::UpdatePlayback(float DeltaTime, IRewindDebugger* Rew
 						if (!CameraActor.IsValid())
 						{
 							FActorSpawnParameters SpawnParameters;
-							SpawnParameters.Name = "RewindDebuggerCamera"; 
+							SpawnParameters.ObjectFlags |= RF_Transient;
 							CameraActor = RewindDebugger->GetWorldToVisualize()->SpawnActor<ACameraActor>(ViewMessage.Position, ViewMessage.Rotation, SpawnParameters);
 							UCameraComponent* Camera = CameraActor->GetCameraComponent();
+							CameraActor->SetActorLabel("RewindDebuggerCamera"); 
 						}
 
 						UCameraComponent* Camera = CameraActor->GetCameraComponent();
diff --git a/Engine/Plugins/Animation/GameplayInsights/Source/RewindDebugger/Private/RewindDebuggerCamera.h b/Engine/Plugins/Animation/GameplayInsights/Source/RewindDebugger/Private/RewindDebuggerCamera.h
--- a/Engine/Plugins/Animation/GameplayInsights/Source/RewindDebugger/Private/RewindDebuggerCamera.h
+++ b/Engine/Plugins/Animation/GameplayInsights/Source/RewindDebugger/Private/RewindDebuggerCamera.h
@@ -25,7 +25,7 @@ public:
 	FRewindDebuggerCamera();
 	virtual ~FRewindDebuggerCamera() {};
 
-	virtual void UpdatePlayback(float DeltaTime, IRewindDebugger* RewindDebugger) override;
+	virtual void Update(float DeltaTime, IRewindDebugger* RewindDebugger) override;
 
 	ECameraMode CameraMode();
 	void SetCameraMode(ECameraMode Mode);
diff --git a/Engine/Plugins/Animation/GameplayInsights/Source/RewindDebugger/Private/RewindDebuggerModule.h b/Engine/Plugins/Animation/GameplayInsights/Source/RewindDebugger/Private/RewindDebuggerModule.h
--- a/Engine/Plugins/Animation/GameplayInsights/Source/RewindDebugger/Private/RewindDebuggerModule.h
+++ b/Engine/Plugins/Animation/GameplayInsights/Source/RewindDebugger/Private/RewindDebuggerModule.h
@@ -9,19 +9,6 @@
 
 class SRewindDebugger;
 
-struct FDebugObjectInfo
-{
-	FDebugObjectInfo(uint64 Id, const FString& Name): ObjectId(Id), ObjectName(Name), bExpanded(true)
-	{
-	}
-
-	uint64 ObjectId;
-	FString ObjectName;
-	bool bExpanded;
-
-	TArray<TSharedPtr<FDebugObjectInfo>> Children;
-};
-
 class FRewindDebuggerModule : public IModuleInterface
 {
 public:
diff --git a/Engine/Plugins/Animation/GameplayInsights/Source/RewindDebugger/Private/SRewindDebuggerComponentTree.h b/Engine/Plugins/Animation/GameplayInsights/Source/RewindDebugger/Private/SRewindDebuggerComponentTree.h
--- a/Engine/Plugins/Animation/GameplayInsights/Source/RewindDebugger/Private/SRewindDebuggerComponentTree.h
+++ b/Engine/Plugins/Animation/GameplayInsights/Source/RewindDebugger/Private/SRewindDebuggerComponentTree.h
@@ -2,7 +2,7 @@
 #pragma once
 
 #include "CoreMinimal.h"
-#include "RewindDebuggerModule.h"
+#include "IRewindDebugger.h"
 #include "Widgets/DeclarativeSyntaxSupport.h"
 #include "Widgets/SCompoundWidget.h"
 #include "Widgets/Views/STreeView.h"
diff --git a/Engine/Plugins/Animation/GameplayInsights/Source/RewindDebugger/Public/IRewindDebugger.h b/Engine/Plugins/Animation/GameplayInsights/Source/RewindDebugger/Public/IRewindDebugger.h
--- a/Engine/Plugins/Animation/GameplayInsights/Source/RewindDebugger/Public/IRewindDebugger.h
+++ b/Engine/Plugins/Animation/GameplayInsights/Source/RewindDebugger/Public/IRewindDebugger.h
@@ -12,7 +12,20 @@ namespace TraceServices
 	class IAnalysisSession;
 }
 
-class IRewindDebugger
+struct REWINDDEBUGGER_API FDebugObjectInfo
+{
+	FDebugObjectInfo(uint64 Id, const FString& Name): ObjectId(Id), ObjectName(Name), bExpanded(true)
+	{
+	}
+
+	uint64 ObjectId;
+	FString ObjectName;
+	bool bExpanded;
+
+	TArray<TSharedPtr<FDebugObjectInfo>> Children;
+};
+
+class REWINDDEBUGGER_API IRewindDebugger
 {
 public:
 	// get the time the debugger is scrubbed to, in seconds since the capture started (or the recording duration while the game is running)
@@ -24,9 +37,21 @@ public:
 	// get insights id for the selected target actor
 	virtual uint64 GetTargetActorId() const = 0;
 
+	// get a list of all components of the selected target actor (with the actor as the first element in the list)
+	virtual TArray<TSharedPtr<FDebugObjectInfo>>& GetDebugComponents() = 0;
+
 	// get posiotion of the selected target actor (returns true if position is valid)
 	virtual bool GetTargetActorPosition(FVector& OutPosition) const = 0;
 
 	// get the world that the debugger is replaying in
 	virtual UWorld* GetWorldToVisualize() const = 0;
+
+	// returns true if recording is active
+	virtual bool IsRecording() const = 0;
+
+	// returns true if PIE is running and not paused
+	virtual bool IsPIESimulating() const = 0;
+
+	// returns the length of the current recording
+	virtual double GetRecordingDuration() const = 0;
 };
diff --git a/Engine/Plugins/Animation/GameplayInsights/Source/RewindDebugger/Public/IRewindDebuggerExtension.h b/Engine/Plugins/Animation/GameplayInsights/Source/RewindDebugger/Public/IRewindDebuggerExtension.h
--- a/Engine/Plugins/Animation/GameplayInsights/Source/RewindDebugger/Public/IRewindDebuggerExtension.h
+++ b/Engine/Plugins/Animation/GameplayInsights/Source/RewindDebugger/Public/IRewindDebuggerExtension.h
@@ -13,11 +13,17 @@
 
 class IRewindDebugger;
 
-class IRewindDebuggerExtension : public IModularFeature
+class REWINDDEBUGGER_API IRewindDebuggerExtension : public IModularFeature
 {
 public:
-	static REWINDDEBUGGER_API const FName ModularFeatureName;
+	static const FName ModularFeatureName;
 
 	// called while scrubbing, playing back, or paused
-	virtual void UpdatePlayback(float DeltaTime, IRewindDebugger* RewindDebugger) {};
+	virtual void Update(float DeltaTime, IRewindDebugger* RewindDebugger) {};
+
+	// called when recording has started
+	virtual void RecordingStarted(IRewindDebugger* RewindDebugger) {};
+
+	// called when recording has ended
+	virtual void RecordingStopped(IRewindDebugger* RewindDebugger) {};
 };
