rewind debugger interfaces refactor
- moving interfaces to engine from plugins
- get rewind debugger debug view interface out of GameplayInsights plugin
- change view system to use modular features instead of having it's own registration management

diff --git a/Engine/Plugins/Animation/GameplayInsights/Source/GameplayInsights/GameplayInsights.Build.cs b/Engine/Plugins/Animation/GameplayInsights/Source/GameplayInsights/GameplayInsights.Build.cs
--- a/Engine/Plugins/Animation/GameplayInsights/Source/GameplayInsights/GameplayInsights.Build.cs
+++ b/Engine/Plugins/Animation/GameplayInsights/Source/GameplayInsights/GameplayInsights.Build.cs
@@ -20,6 +20,7 @@ namespace UnrealBuildTool.Rules
 				"TraceInsights",
 				"AssetRegistry",
 				"ApplicationCore",
+				"RewindDebuggerInterface"
 			});
 
 			if (Target.bCompileAgainstEngine)
diff --git a/Engine/Plugins/Animation/GameplayInsights/Source/GameplayInsights/Private/GameplayInsightsDebugViewCreator.cpp b/Engine/Plugins/Animation/GameplayInsights/Source/GameplayInsights/Private/GameplayInsightsDebugViewCreator.cpp
deleted file mode 100644
diff --git a/Engine/Plugins/Animation/GameplayInsights/Source/GameplayInsights/Private/GameplayInsightsDebugViewCreator.h b/Engine/Plugins/Animation/GameplayInsights/Source/GameplayInsights/Private/GameplayInsightsDebugViewCreator.h
deleted file mode 100644
diff --git a/Engine/Plugins/Animation/GameplayInsights/Source/GameplayInsights/Private/GameplayInsightsModule.cpp b/Engine/Plugins/Animation/GameplayInsights/Source/GameplayInsights/Private/GameplayInsightsModule.cpp
--- a/Engine/Plugins/Animation/GameplayInsights/Source/GameplayInsights/Private/GameplayInsightsModule.cpp
+++ b/Engine/Plugins/Animation/GameplayInsights/Source/GameplayInsights/Private/GameplayInsightsModule.cpp
@@ -64,13 +64,19 @@ void FGameplayInsightsModule::StartupModule()
 		}
 	});
 
-	DebugViewCreator.RegisterDebugViewCreator("Object", MakeShared<FObjectPropertiesViewCreator>());
-	DebugViewCreator.RegisterDebugViewCreator("AnimInstance", MakeShared<FAnimGraphSchematicViewCreator>());
-	DebugViewCreator.RegisterDebugViewCreator("AnimInstance", MakeShared<FBlendWeightsViewCreator>());
-	DebugViewCreator.RegisterDebugViewCreator("AnimInstance", MakeShared<FNotifiesViewCreator>());
-	DebugViewCreator.RegisterDebugViewCreator("SkeletalMeshComponent", MakeShared<FAnimationCurvesViewCreator>());
-
 #if WITH_EDITOR
+	// register rewind debugger view creators
+	static FObjectPropertiesViewCreator ObjectPropertiesViewCreator;
+	IModularFeatures::Get().RegisterModularFeature(IRewindDebuggerViewCreator::ModularFeatureName, &ObjectPropertiesViewCreator);
+	static FAnimGraphSchematicViewCreator AnimGraphSchematicViewCreator;
+	IModularFeatures::Get().RegisterModularFeature(IRewindDebuggerViewCreator::ModularFeatureName, &AnimGraphSchematicViewCreator);
+	static FBlendWeightsViewCreator BlendWeightsViewCreator;
+	IModularFeatures::Get().RegisterModularFeature(IRewindDebuggerViewCreator::ModularFeatureName, &BlendWeightsViewCreator);
+	static FNotifiesViewCreator NotifiesViewCreator;
+	IModularFeatures::Get().RegisterModularFeature(IRewindDebuggerViewCreator::ModularFeatureName, &NotifiesViewCreator);
+	static FAnimationCurvesViewCreator AnimationCurvesViewCreator;
+	IModularFeatures::Get().RegisterModularFeature(IRewindDebuggerViewCreator::ModularFeatureName, &AnimationCurvesViewCreator);
+
 
 	if (!IsRunningCommandlet())
 	{
diff --git a/Engine/Plugins/Animation/GameplayInsights/Source/GameplayInsights/Private/GameplayInsightsModule.h b/Engine/Plugins/Animation/GameplayInsights/Source/GameplayInsights/Private/GameplayInsightsModule.h
--- a/Engine/Plugins/Animation/GameplayInsights/Source/GameplayInsights/Private/GameplayInsightsModule.h
+++ b/Engine/Plugins/Animation/GameplayInsights/Source/GameplayInsights/Private/GameplayInsightsModule.h
@@ -8,7 +8,6 @@
 #include "GameplayTimingViewExtender.h"
 #include "Framework/Docking/TabManager.h"
 #include "IGameplayInsightsModule.h"
-#include "GameplayInsightsDebugViewCreator.h"
 #include "Trace/StoreService.h"
 
 struct FInsightsMajorTabExtender;
@@ -31,9 +30,6 @@ public:
 	virtual void StartupModule() override;
 	virtual void ShutdownModule() override;
 
-	// IGameplayInsightsModule interface
-	virtual IGameplayInsightsDebugViewCreator* GetDebugViewCreator() override { return &DebugViewCreator; }
-
 #if WITH_EDITOR
 	virtual void EnableObjectPropertyTrace(UObject* Object, bool bEnable = true) override;
 	virtual bool IsObjectPropertyTraceEnabled(UObject* Object) override;
@@ -49,8 +45,6 @@ protected:
 #endif
 
 private:
-	FGameplayInsightsDebugViewCreator DebugViewCreator;
-
 	FGameplayTraceModule GameplayTraceModule;
 
 	FGameplayTimingViewExtender GameplayTimingViewExtender;
diff --git a/Engine/Plugins/Animation/GameplayInsights/Source/GameplayInsights/Private/SAnimGraphSchematicView.cpp b/Engine/Plugins/Animation/GameplayInsights/Source/GameplayInsights/Private/SAnimGraphSchematicView.cpp
--- a/Engine/Plugins/Animation/GameplayInsights/Source/GameplayInsights/Private/SAnimGraphSchematicView.cpp
+++ b/Engine/Plugins/Animation/GameplayInsights/Source/GameplayInsights/Private/SAnimGraphSchematicView.cpp
@@ -1003,6 +1003,12 @@ void SAnimGraphSchematicView::RefreshDetails(const TArray<TSharedRef<FAnimGraphS
 	}
 }
 
+FName FAnimGraphSchematicViewCreator::GetTargetTypeName() const
+{
+	static FName TargetTypeName = "AnimInstance";
+	return TargetTypeName;
+}
+
 FName FAnimGraphSchematicViewCreator::GetName() const
 {
 	return AnimGraphSchematicName;
@@ -1022,7 +1028,7 @@ FSlateIcon FAnimGraphSchematicViewCreator::GetIcon() const
 #endif
 }
 
-TSharedPtr<IGameplayInsightsDebugView> FAnimGraphSchematicViewCreator::CreateDebugView(uint64 ObjectId, double CurrentTime, const TraceServices::IAnalysisSession& AnalysisSession) const
+TSharedPtr<IRewindDebuggerView> FAnimGraphSchematicViewCreator::CreateDebugView(uint64 ObjectId, double CurrentTime, const TraceServices::IAnalysisSession& AnalysisSession) const
 {
 	return SNew(SAnimGraphSchematicView, ObjectId, CurrentTime, AnalysisSession);
 }
diff --git a/Engine/Plugins/Animation/GameplayInsights/Source/GameplayInsights/Private/SAnimGraphSchematicView.h b/Engine/Plugins/Animation/GameplayInsights/Source/GameplayInsights/Private/SAnimGraphSchematicView.h
--- a/Engine/Plugins/Animation/GameplayInsights/Source/GameplayInsights/Private/SAnimGraphSchematicView.h
+++ b/Engine/Plugins/Animation/GameplayInsights/Source/GameplayInsights/Private/SAnimGraphSchematicView.h
@@ -3,8 +3,8 @@
 #pragma once
 
 #include "CoreMinimal.h"
-#include "IGameplayInsightsDebugView.h"
-#include "IGameplayInsightsDebugViewCreator.h"
+#include "IRewindDebuggerView.h"
+#include "IRewindDebuggerViewCreator.h"
 #include "Widgets/SCompoundWidget.h"
 #include "Widgets/Views/STreeView.h"
 
@@ -21,7 +21,7 @@ class SComboButton;
 class SSplitter;
 enum class EAnimGraphSchematicFilterState;
 
-class SAnimGraphSchematicView : public IGameplayInsightsDebugView
+class SAnimGraphSchematicView : public IRewindDebuggerView
 {
 	SLATE_BEGIN_ARGS(SAnimGraphSchematicView) {}
 	SLATE_END_ARGS()
@@ -116,11 +116,12 @@ private:
 	uint64 AnimInstanceId;
 };
 
-class FAnimGraphSchematicViewCreator : public ICreateGameplayInsightsDebugView
+class FAnimGraphSchematicViewCreator : public IRewindDebuggerViewCreator
 {
 	public:
+		virtual FName GetTargetTypeName() const;
 		virtual FName GetName() const override;
 		virtual FText GetTitle() const override;
 		virtual FSlateIcon GetIcon() const override;
-		virtual TSharedPtr<IGameplayInsightsDebugView> CreateDebugView(uint64 ObjectId, double CurrentTime, const TraceServices::IAnalysisSession& InAnalysisSession) const override;
+		virtual TSharedPtr<IRewindDebuggerView> CreateDebugView(uint64 ObjectId, double CurrentTime, const TraceServices::IAnalysisSession& InAnalysisSession) const override;
 };
\ No newline at end of file
diff --git a/Engine/Plugins/Animation/GameplayInsights/Source/GameplayInsights/Private/SAnimationCurvesView.cpp b/Engine/Plugins/Animation/GameplayInsights/Source/GameplayInsights/Private/SAnimationCurvesView.cpp
--- a/Engine/Plugins/Animation/GameplayInsights/Source/GameplayInsights/Private/SAnimationCurvesView.cpp
+++ b/Engine/Plugins/Animation/GameplayInsights/Source/GameplayInsights/Private/SAnimationCurvesView.cpp
@@ -43,6 +43,11 @@ void SAnimationCurvesView::GetVariantsAtFrame(const TraceServices::FFrame& InFra
 	}
 }
 
+FName FAnimationCurvesViewCreator::GetTargetTypeName() const
+{
+	static FName TargetTypeName = "SkeletalMeshComponent";
+	return TargetTypeName;
+}
 
 static const FName AnimationCurvesName("AnimationCurves");
 
@@ -70,7 +75,7 @@ FSlateIcon FAnimationCurvesViewCreator::GetIcon() const
 #endif
 }
 
-TSharedPtr<IGameplayInsightsDebugView> FAnimationCurvesViewCreator::CreateDebugView(uint64 ObjectId, double CurrentTime, const TraceServices::IAnalysisSession& AnalysisSession) const
+TSharedPtr<IRewindDebuggerView> FAnimationCurvesViewCreator::CreateDebugView(uint64 ObjectId, double CurrentTime, const TraceServices::IAnalysisSession& AnalysisSession) const
 {
 	return SNew(SAnimationCurvesView, ObjectId, CurrentTime, AnalysisSession);
 }
diff --git a/Engine/Plugins/Animation/GameplayInsights/Source/GameplayInsights/Private/SAnimationCurvesView.h b/Engine/Plugins/Animation/GameplayInsights/Source/GameplayInsights/Private/SAnimationCurvesView.h
--- a/Engine/Plugins/Animation/GameplayInsights/Source/GameplayInsights/Private/SAnimationCurvesView.h
+++ b/Engine/Plugins/Animation/GameplayInsights/Source/GameplayInsights/Private/SAnimationCurvesView.h
@@ -3,8 +3,8 @@
 #pragma once
 
 #include "CoreMinimal.h"
-#include "IGameplayInsightsDebugView.h"
-#include "IGameplayInsightsDebugViewCreator.h"
+#include "IRewindDebuggerView.h"
+#include "IRewindDebuggerViewCreator.h"
 #include "SPropertiesDebugViewBase.h"
 
 namespace TraceServices { class IAnalysisSession; }
@@ -16,11 +16,13 @@ public:
 	virtual FName GetName() const override;
 };
 
-class FAnimationCurvesViewCreator : public ICreateGameplayInsightsDebugView
+class FAnimationCurvesViewCreator : public IRewindDebuggerViewCreator
 {
 	public:
+
+		virtual FName GetTargetTypeName() const;
 		virtual FName GetName() const override;
 		virtual FText GetTitle() const override;
 		virtual FSlateIcon GetIcon() const override;
-		virtual TSharedPtr<IGameplayInsightsDebugView> CreateDebugView(uint64 ObjectId, double CurrentTime, const TraceServices::IAnalysisSession& InAnalysisSession) const override;
+		virtual TSharedPtr<IRewindDebuggerView> CreateDebugView(uint64 ObjectId, double CurrentTime, const TraceServices::IAnalysisSession& InAnalysisSession) const override;
 };
\ No newline at end of file
diff --git a/Engine/Plugins/Animation/GameplayInsights/Source/GameplayInsights/Private/SBlendWeightsView.cpp b/Engine/Plugins/Animation/GameplayInsights/Source/GameplayInsights/Private/SBlendWeightsView.cpp
--- a/Engine/Plugins/Animation/GameplayInsights/Source/GameplayInsights/Private/SBlendWeightsView.cpp
+++ b/Engine/Plugins/Animation/GameplayInsights/Source/GameplayInsights/Private/SBlendWeightsView.cpp
@@ -58,6 +58,12 @@ FName SBlendWeightsView::GetName() const
 	return BlendWeightsName;
 }
 
+FName FBlendWeightsViewCreator::GetTargetTypeName() const
+{
+	static FName TargetTypeName = "AnimInstance";
+	return TargetTypeName;
+}
+
 FName FBlendWeightsViewCreator::GetName() const
 {
 	return BlendWeightsName;
@@ -77,7 +83,7 @@ FSlateIcon FBlendWeightsViewCreator::GetIcon() const
 #endif
 }
 
-TSharedPtr<IGameplayInsightsDebugView> FBlendWeightsViewCreator::CreateDebugView(uint64 ObjectId, double CurrentTime, const TraceServices::IAnalysisSession& AnalysisSession) const
+TSharedPtr<IRewindDebuggerView> FBlendWeightsViewCreator::CreateDebugView(uint64 ObjectId, double CurrentTime, const TraceServices::IAnalysisSession& AnalysisSession) const
 {
 	return SNew(SBlendWeightsView, ObjectId, CurrentTime, AnalysisSession);
 }
diff --git a/Engine/Plugins/Animation/GameplayInsights/Source/GameplayInsights/Private/SBlendWeightsView.h b/Engine/Plugins/Animation/GameplayInsights/Source/GameplayInsights/Private/SBlendWeightsView.h
--- a/Engine/Plugins/Animation/GameplayInsights/Source/GameplayInsights/Private/SBlendWeightsView.h
+++ b/Engine/Plugins/Animation/GameplayInsights/Source/GameplayInsights/Private/SBlendWeightsView.h
@@ -3,8 +3,8 @@
 #pragma once
 
 #include "CoreMinimal.h"
-#include "IGameplayInsightsDebugView.h"
-#include "IGameplayInsightsDebugViewCreator.h"
+#include "IRewindDebuggerView.h"
+#include "IRewindDebuggerViewCreator.h"
 #include "SPropertiesDebugViewBase.h"
 
 namespace TraceServices { class IAnalysisSession; }
@@ -16,11 +16,12 @@ public:
 	virtual FName GetName() const override;
 };
 
-class FBlendWeightsViewCreator : public ICreateGameplayInsightsDebugView
+class FBlendWeightsViewCreator : public IRewindDebuggerViewCreator
 {
 	public:
+		virtual FName GetTargetTypeName() const;
 		virtual FName GetName() const override;
 		virtual FText GetTitle() const override;
 		virtual FSlateIcon GetIcon() const override;
-		virtual TSharedPtr<IGameplayInsightsDebugView> CreateDebugView(uint64 ObjectId, double CurrentTime, const TraceServices::IAnalysisSession& InAnalysisSession) const override;
+		virtual TSharedPtr<IRewindDebuggerView> CreateDebugView(uint64 ObjectId, double CurrentTime, const TraceServices::IAnalysisSession& InAnalysisSession) const override;
 };
\ No newline at end of file
diff --git a/Engine/Plugins/Animation/GameplayInsights/Source/GameplayInsights/Private/SNotifiesView.cpp b/Engine/Plugins/Animation/GameplayInsights/Source/GameplayInsights/Private/SNotifiesView.cpp
--- a/Engine/Plugins/Animation/GameplayInsights/Source/GameplayInsights/Private/SNotifiesView.cpp
+++ b/Engine/Plugins/Animation/GameplayInsights/Source/GameplayInsights/Private/SNotifiesView.cpp
@@ -112,6 +112,12 @@ FName SNotifiesView::GetName() const
 	return NotifiesName;
 }
 
+FName FNotifiesViewCreator::GetTargetTypeName() const
+{
+	static FName TargetTypeName = "AnimInstance";
+	return TargetTypeName;
+}
+
 FName FNotifiesViewCreator::GetName() const
 {
 	return NotifiesName;
@@ -131,7 +137,7 @@ FSlateIcon FNotifiesViewCreator::GetIcon() const
 #endif
 }
 
-TSharedPtr<IGameplayInsightsDebugView> FNotifiesViewCreator::CreateDebugView(uint64 ObjectId, double CurrentTime, const TraceServices::IAnalysisSession& AnalysisSession) const
+TSharedPtr<IRewindDebuggerView> FNotifiesViewCreator::CreateDebugView(uint64 ObjectId, double CurrentTime, const TraceServices::IAnalysisSession& AnalysisSession) const
 {
 	return SNew(SNotifiesView, ObjectId, CurrentTime, AnalysisSession);
 }
diff --git a/Engine/Plugins/Animation/GameplayInsights/Source/GameplayInsights/Private/SNotifiesView.h b/Engine/Plugins/Animation/GameplayInsights/Source/GameplayInsights/Private/SNotifiesView.h
--- a/Engine/Plugins/Animation/GameplayInsights/Source/GameplayInsights/Private/SNotifiesView.h
+++ b/Engine/Plugins/Animation/GameplayInsights/Source/GameplayInsights/Private/SNotifiesView.h
@@ -3,8 +3,6 @@
 #pragma once
 
 #include "CoreMinimal.h"
-#include "IGameplayInsightsDebugView.h"
-#include "IGameplayInsightsDebugViewCreator.h"
 #include "SPropertiesDebugViewBase.h"
 
 namespace TraceServices { class IAnalysisSession; }
@@ -16,11 +14,12 @@ public:
 	virtual FName GetName() const override;
 };
 
-class FNotifiesViewCreator : public ICreateGameplayInsightsDebugView
+class FNotifiesViewCreator : public IRewindDebuggerViewCreator
 {
 	public:
+		virtual FName GetTargetTypeName() const;
 		virtual FName GetName() const override;
 		virtual FText GetTitle() const override;
 		virtual FSlateIcon GetIcon() const override;
-		virtual TSharedPtr<IGameplayInsightsDebugView> CreateDebugView(uint64 ObjectId, double CurrentTime, const TraceServices::IAnalysisSession& InAnalysisSession) const override;
+		virtual TSharedPtr<IRewindDebuggerView> CreateDebugView(uint64 ObjectId, double CurrentTime, const TraceServices::IAnalysisSession& InAnalysisSession) const override;
 };
\ No newline at end of file
diff --git a/Engine/Plugins/Animation/GameplayInsights/Source/GameplayInsights/Private/SObjectPropertiesView.cpp b/Engine/Plugins/Animation/GameplayInsights/Source/GameplayInsights/Private/SObjectPropertiesView.cpp
--- a/Engine/Plugins/Animation/GameplayInsights/Source/GameplayInsights/Private/SObjectPropertiesView.cpp
+++ b/Engine/Plugins/Animation/GameplayInsights/Source/GameplayInsights/Private/SObjectPropertiesView.cpp
@@ -60,7 +60,6 @@ void SObjectPropertiesView::GetVariantsAtFrame(const TraceServices::FFrame& InFr
 	}
 }
 
-
 static const FName ObjectPropertiesName("ObjectProperties");
 
 FName SObjectPropertiesView::GetName() const
@@ -68,6 +67,12 @@ FName SObjectPropertiesView::GetName() const
 	return ObjectPropertiesName;
 }
 
+FName FObjectPropertiesViewCreator::GetTargetTypeName() const
+{
+	static FName TargetTypeName = "Object";
+	return TargetTypeName;
+}
+
 FName FObjectPropertiesViewCreator::GetName() const
 {
 	return ObjectPropertiesName;
@@ -83,7 +88,7 @@ FSlateIcon FObjectPropertiesViewCreator::GetIcon() const
 	return FSlateIconFinder::FindIconForClass(UObject::StaticClass());
 }
 
-TSharedPtr<IGameplayInsightsDebugView> FObjectPropertiesViewCreator::CreateDebugView(uint64 ObjectId, double CurrentTime, const TraceServices::IAnalysisSession& AnalysisSession) const
+TSharedPtr<IRewindDebuggerView> FObjectPropertiesViewCreator::CreateDebugView(uint64 ObjectId, double CurrentTime, const TraceServices::IAnalysisSession& AnalysisSession) const
 {
 	return SNew(SObjectPropertiesView, ObjectId, CurrentTime, AnalysisSession);
 }
diff --git a/Engine/Plugins/Animation/GameplayInsights/Source/GameplayInsights/Private/SObjectPropertiesView.h b/Engine/Plugins/Animation/GameplayInsights/Source/GameplayInsights/Private/SObjectPropertiesView.h
--- a/Engine/Plugins/Animation/GameplayInsights/Source/GameplayInsights/Private/SObjectPropertiesView.h
+++ b/Engine/Plugins/Animation/GameplayInsights/Source/GameplayInsights/Private/SObjectPropertiesView.h
@@ -3,8 +3,6 @@
 #pragma once
 
 #include "CoreMinimal.h"
-#include "IGameplayInsightsDebugView.h"
-#include "IGameplayInsightsDebugViewCreator.h"
 #include "SPropertiesDebugViewBase.h"
 
 namespace TraceServices { class IAnalysisSession; }
@@ -16,11 +14,12 @@ public:
 	virtual FName GetName() const override;
 };
 
-class FObjectPropertiesViewCreator : public ICreateGameplayInsightsDebugView
+class FObjectPropertiesViewCreator : public IRewindDebuggerViewCreator
 {
 	public:
+		virtual FName GetTargetTypeName() const;
 		virtual FName GetName() const override;
 		virtual FText GetTitle() const override;
 		virtual FSlateIcon GetIcon() const override;
-		virtual TSharedPtr<IGameplayInsightsDebugView> CreateDebugView(uint64 ObjectId, double CurrentTime, const TraceServices::IAnalysisSession& InAnalysisSession) const override;
+		virtual TSharedPtr<IRewindDebuggerView> CreateDebugView(uint64 ObjectId, double CurrentTime, const TraceServices::IAnalysisSession& InAnalysisSession) const override;
 };
\ No newline at end of file
diff --git a/Engine/Plugins/Animation/GameplayInsights/Source/GameplayInsights/Private/SPropertiesDebugViewBase.h b/Engine/Plugins/Animation/GameplayInsights/Source/GameplayInsights/Private/SPropertiesDebugViewBase.h
--- a/Engine/Plugins/Animation/GameplayInsights/Source/GameplayInsights/Private/SPropertiesDebugViewBase.h
+++ b/Engine/Plugins/Animation/GameplayInsights/Source/GameplayInsights/Private/SPropertiesDebugViewBase.h
@@ -3,13 +3,13 @@
 #pragma once
 
 #include "CoreMinimal.h"
-#include "IGameplayInsightsDebugView.h"
-#include "IGameplayInsightsDebugViewCreator.h"
+#include "IRewindDebuggerView.h"
+#include "IRewindDebuggerViewCreator.h"
 #include "SVariantValueView.h"
 
 namespace TraceServices { class IAnalysisSession; }
 
-class SPropertiesDebugViewBase : public IGameplayInsightsDebugView
+class SPropertiesDebugViewBase : public IRewindDebuggerView
 {
 	SLATE_BEGIN_ARGS(SPropertiesDebugViewBase) {}
 	SLATE_END_ARGS()
diff --git a/Engine/Plugins/Animation/GameplayInsights/Source/GameplayInsights/Public/IGameplayInsightsDebugView.h b/Engine/Plugins/Animation/GameplayInsights/Source/GameplayInsights/Public/IGameplayInsightsDebugView.h
deleted file mode 100644
diff --git a/Engine/Plugins/Animation/GameplayInsights/Source/GameplayInsights/Public/IGameplayInsightsDebugViewCreator.h b/Engine/Plugins/Animation/GameplayInsights/Source/GameplayInsights/Public/IGameplayInsightsDebugViewCreator.h
deleted file mode 100644
diff --git a/Engine/Plugins/Animation/GameplayInsights/Source/GameplayInsights/Public/IGameplayInsightsModule.h b/Engine/Plugins/Animation/GameplayInsights/Source/GameplayInsights/Public/IGameplayInsightsModule.h
--- a/Engine/Plugins/Animation/GameplayInsights/Source/GameplayInsights/Public/IGameplayInsightsModule.h
+++ b/Engine/Plugins/Animation/GameplayInsights/Source/GameplayInsights/Public/IGameplayInsightsModule.h
@@ -3,28 +3,21 @@
 #pragma once
 
 #include "CoreMinimal.h"
-#include "IGameplayInsightsDebugViewCreator.h"
-
 
 namespace GameplayInsightsTabs
 {
 	extern GAMEPLAYINSIGHTS_API const FName DocumentTab;
 };
 
-
 namespace TraceServices
 {
 	class IAnalysisSession;
 }
 
-class IGameplayInsightsDebugViewCreator;
-
 class IGameplayInsightsModule
 	: public IModuleInterface
 {
 public:
-	
-	virtual IGameplayInsightsDebugViewCreator* GetDebugViewCreator() = 0;
 
 #if WITH_EDITOR
 	// start or stop tracing all properties of this Object to Insights
diff --git a/Engine/Plugins/Animation/GameplayInsights/Source/RewindDebugger/Private/RewindDebugger.cpp b/Engine/Plugins/Animation/GameplayInsights/Source/RewindDebugger/Private/RewindDebugger.cpp
--- a/Engine/Plugins/Animation/GameplayInsights/Source/RewindDebugger/Private/RewindDebugger.cpp
+++ b/Engine/Plugins/Animation/GameplayInsights/Source/RewindDebugger/Private/RewindDebugger.cpp
@@ -15,7 +15,6 @@
 
 
 FRewindDebugger* FRewindDebugger::InternalInstance = nullptr;
-const FName IRewindDebuggerExtension::ModularFeatureName = "RewindDebuggerExtension";
 
 static void IterateExtensions(TFunction<void(IRewindDebuggerExtension* Extension)> IteratorFunction)
 {
diff --git a/Engine/Plugins/Animation/GameplayInsights/Source/RewindDebugger/Private/RewindDebuggerViewCreators.cpp b/Engine/Plugins/Animation/GameplayInsights/Source/RewindDebugger/Private/RewindDebuggerViewCreators.cpp
new file mode 100644
--- /dev/null
+++ b/Engine/Plugins/Animation/GameplayInsights/Source/RewindDebugger/Private/RewindDebuggerViewCreators.cpp
@@ -0,0 +1,61 @@
+// Copyright Epic Games, Inc. All Rights Reserved.
+
+#include "RewindDebuggerViewCreators.h"
+#include "Features/IModularFeatures.h"
+#include "TraceServices/Model/AnalysisSession.h"
+#include "IGameplayProvider.h"
+
+void FRewindDebuggerViewCreators::EnumerateCreators(TFunctionRef<void(const IRewindDebuggerViewCreator*)> Callback)
+{
+	IModularFeatures& ModularFeatures = IModularFeatures::Get();
+	static const FName CreatorFeatureName = IRewindDebuggerViewCreator::ModularFeatureName;
+
+	const int32 NumExtensions = ModularFeatures.GetModularFeatureImplementationCount(CreatorFeatureName);
+	for (int32 ExtensionIndex = 0; ExtensionIndex < NumExtensions; ++ExtensionIndex)
+	{
+		IRewindDebuggerViewCreator* ViewCreator = static_cast<IRewindDebuggerViewCreator*>(ModularFeatures.GetModularFeatureImplementation(CreatorFeatureName, ExtensionIndex));
+        Callback(ViewCreator);
+    }
+}
+
+void FRewindDebuggerViewCreators::CreateDebugViews(uint64 ObjectId, double CurrentTime, const TraceServices::IAnalysisSession& Session, TArray<TSharedPtr<IRewindDebuggerView>>& OutDebugViews)
+{
+	TraceServices::FAnalysisSessionReadScope SessionReadScope(Session);
+
+	const IGameplayProvider* GameplayProvider = Session.ReadProvider<IGameplayProvider>("GameplayProvider");
+	const FObjectInfo& ObjectInfo = GameplayProvider->GetObjectInfo(ObjectId);
+	uint64 ClassId = ObjectInfo.ClassId;
+	TArray<FName> TypeNameHierarchy;
+	while (ClassId != 0)
+	{
+		const FClassInfo& ClassInfo = GameplayProvider->GetClassInfo(ClassId);
+		TypeNameHierarchy.Add(ClassInfo.Name);
+		ClassId = ClassInfo.SuperId;
+	}
+
+	EnumerateCreators([&OutDebugViews, &TypeNameHierarchy, ObjectId, CurrentTime, &Session](const IRewindDebuggerViewCreator* ViewCreator)
+	{
+		if (TypeNameHierarchy.Contains(ViewCreator->GetTargetTypeName()))
+		{
+			OutDebugViews.Add(ViewCreator->CreateDebugView(ObjectId, CurrentTime, Session));
+		}
+	});
+}
+
+const IRewindDebuggerViewCreator* FRewindDebuggerViewCreators::GetCreator(FName CreatorName)
+{
+	IModularFeatures& ModularFeatures = IModularFeatures::Get();
+	static const FName CreatorFeatureName = IRewindDebuggerViewCreator::ModularFeatureName;
+
+	const int32 NumExtensions = ModularFeatures.GetModularFeatureImplementationCount(CreatorFeatureName);
+	for (int32 ExtensionIndex = 0; ExtensionIndex < NumExtensions; ++ExtensionIndex)
+	{
+		IRewindDebuggerViewCreator* ViewCreator = static_cast<IRewindDebuggerViewCreator*>(ModularFeatures.GetModularFeatureImplementation(CreatorFeatureName, ExtensionIndex));
+		if (ViewCreator->GetName() == CreatorName)
+		{
+			return ViewCreator;
+		}
+    }
+
+    return nullptr;
+}
diff --git a/Engine/Plugins/Animation/GameplayInsights/Source/RewindDebugger/Private/RewindDebuggerViewCreators.h b/Engine/Plugins/Animation/GameplayInsights/Source/RewindDebugger/Private/RewindDebuggerViewCreators.h
new file mode 100644
--- /dev/null
+++ b/Engine/Plugins/Animation/GameplayInsights/Source/RewindDebugger/Private/RewindDebuggerViewCreators.h
@@ -0,0 +1,20 @@
+// Copyright Epic Games, Inc. All Rights Reserved.
+
+#pragma once
+
+#include "CoreMinimal.h"
+#include "IRewindDebuggerView.h"
+#include "IRewindDebuggerViewCreator.h"
+
+namespace TraceServices
+{
+	class IRewindDebuggerViewCreator;
+}
+
+class FRewindDebuggerViewCreators
+{
+public:
+	static void EnumerateCreators(TFunctionRef<void(const IRewindDebuggerViewCreator*)> Callback);
+	static const IRewindDebuggerViewCreator* GetCreator(FName CreatorName);
+	static void CreateDebugViews(uint64 ObjectId, double CurrentTime, const TraceServices::IAnalysisSession& Session, TArray<TSharedPtr<IRewindDebuggerView>>& OutDebugViews);
+};
diff --git a/Engine/Plugins/Animation/GameplayInsights/Source/RewindDebugger/Private/SRewindDebugger.cpp b/Engine/Plugins/Animation/GameplayInsights/Source/RewindDebugger/Private/SRewindDebugger.cpp
--- a/Engine/Plugins/Animation/GameplayInsights/Source/RewindDebugger/Private/SRewindDebugger.cpp
+++ b/Engine/Plugins/Animation/GameplayInsights/Source/RewindDebugger/Private/SRewindDebugger.cpp
@@ -30,6 +30,8 @@
 #include "Widgets/Input/SSearchBox.h"
 #include "Widgets/Layout/SExpandableArea.h"
 #include "Widgets/Layout/SScrollBox.h"
+#include "IRewindDebuggerViewCreator.h"
+#include "RewindDebuggerViewCreators.h"
 
 #define LOCTEXT_NAMESPACE "SRewindDebugger"
 
@@ -262,7 +264,7 @@ void SRewindDebugger::Construct(const FArguments& InArgs, TSharedRef<FUICommandL
 	TSharedPtr<FTabManager::FStack> MainTabStack = FTabManager::NewStack();
 
 	IGameplayInsightsModule& GameplayInsightsModule = FModuleManager::LoadModuleChecked<IGameplayInsightsModule>("GameplayInsights");
-	GameplayInsightsModule.GetDebugViewCreator()->EnumerateCreators([this, MainTabStack](const TSharedPtr<ICreateGameplayInsightsDebugView>& Creator)
+	FRewindDebuggerViewCreators::EnumerateCreators([this, MainTabStack](const IRewindDebuggerViewCreator* Creator)
 		{
 			FName TabName = Creator->GetName();
 			TabNames.Add(TabName);
@@ -420,11 +422,11 @@ void SRewindDebugger::RefreshDebugComponents()
 
 void SRewindDebugger::TraceTimeChanged(double Time)
 {
-	for(TSharedPtr<IGameplayInsightsDebugView>& DebugView : DebugViews)
+	for(TSharedPtr<IRewindDebuggerView>& DebugView : DebugViews)
 	{
 		DebugView->SetTimeMarker(Time);
 	}
-	for(TSharedPtr<IGameplayInsightsDebugView>& DebugView : PinnedDebugViews)
+	for(TSharedPtr<IRewindDebuggerView>& DebugView : PinnedDebugViews)
 	{
 		DebugView->SetTimeMarker(Time);
 	}
@@ -432,7 +434,7 @@ void SRewindDebugger::TraceTimeChanged(double Time)
 
 TSharedRef<SDockTab> SRewindDebugger::SpawnTab(const FSpawnTabArgs& Args, FName ViewName)
 {
-	TSharedPtr<IGameplayInsightsDebugView>* View = DebugViews.FindByPredicate([ViewName](TSharedPtr<IGameplayInsightsDebugView>& View) { return View->GetName() == ViewName; } );
+	TSharedPtr<IRewindDebuggerView>* View = DebugViews.FindByPredicate([ViewName](TSharedPtr<IRewindDebuggerView>& View) { return View->GetName() == ViewName; } );
 	if (View)
 	{
 		HiddenTabs.Remove(ViewName);
@@ -457,12 +459,12 @@ TSharedRef<SDockTab> SRewindDebugger::SpawnTab(const FSpawnTabArgs& Args, FName
 // returns true if DebugViews contains a view for the ViewName, but there is no matching Pinned view already open
 bool SRewindDebugger::CanSpawnTab(const FSpawnTabArgs& Args, FName ViewName)
 {
-	TSharedPtr<IGameplayInsightsDebugView>* View = DebugViews.FindByPredicate([ViewName](TSharedPtr<IGameplayInsightsDebugView>& View) { return View->GetName() == ViewName; } );
+	TSharedPtr<IRewindDebuggerView>* View = DebugViews.FindByPredicate([ViewName](TSharedPtr<IRewindDebuggerView>& View) { return View->GetName() == ViewName; } );
 
 	if (View!=nullptr)
 	{
 		bool bPinned = nullptr != PinnedDebugViews.FindByPredicate(
-			[View](TSharedPtr<IGameplayInsightsDebugView>& PinnedView)
+			[View](TSharedPtr<IRewindDebuggerView>& PinnedView)
 			{
 				return PinnedView->GetName() == (*View)->GetName() && PinnedView->GetObjectId() == (*View)->GetObjectId();
 			}
@@ -476,14 +478,14 @@ bool SRewindDebugger::CanSpawnTab(const FSpawnTabArgs& Args, FName ViewName)
 void SRewindDebugger::OnPinnedTabClosed(TSharedRef<SDockTab> Tab)
 {
 	// remove view from list of pinned views
-	TSharedRef<IGameplayInsightsDebugView> View = StaticCastSharedRef<IGameplayInsightsDebugView>(Tab->GetContent());
+	TSharedRef<IRewindDebuggerView> View = StaticCastSharedRef<IRewindDebuggerView>(Tab->GetContent());
 	PinnedDebugViews.Remove(View);
 
 	// recreate non-pinned tabs, so when closing a pinned tab for the currently selected component, the non-pinned one will appear
 	CreateDebugTabs();
 }
 
- void SRewindDebugger::PinTab(TSharedPtr<IGameplayInsightsDebugView> View)
+ void SRewindDebugger::PinTab(TSharedPtr<IRewindDebuggerView> View)
  {
 	if (PinnedDebugViews.Find(View) != INDEX_NONE)
 	{
@@ -498,8 +500,7 @@ void SRewindDebugger::OnPinnedTabClosed(TSharedRef<SDockTab> Tab)
 	FText TabLabel;
 
 	IGameplayInsightsModule& GameplayInsightsModule = FModuleManager::LoadModuleChecked<IGameplayInsightsModule>("GameplayInsights");
-	TSharedPtr<ICreateGameplayInsightsDebugView> Creator = GameplayInsightsModule.GetDebugViewCreator()->GetCreator(TabName);
-	if (Creator.IsValid())
+	if (const IRewindDebuggerViewCreator* Creator = FRewindDebuggerViewCreators::GetCreator(TabName))
 	{
 		TabIcon = Creator->GetIcon();
 		TabLabel = Creator->GetTitle();
@@ -540,7 +541,7 @@ void SRewindDebugger::MakeViewsMenu(FMenuBuilder& MenuBuilder)
 							 FExecuteAction::CreateLambda([this]() { ShowAllViews(); }));
 }
 
-void SRewindDebugger::ExtendTabMenu(FMenuBuilder& MenuBuilder, TSharedPtr<IGameplayInsightsDebugView> View)
+void SRewindDebugger::ExtendTabMenu(FMenuBuilder& MenuBuilder, TSharedPtr<IRewindDebuggerView> View)
 {
 	MenuBuilder.BeginSection("RewindDebugger", LOCTEXT("Rewind Debugger", "Rewind Debugger"));
 
@@ -572,7 +573,7 @@ void SRewindDebugger::CreateDebugViews()
 		IGameplayInsightsModule& GameplayInsightsModule = FModuleManager::LoadModuleChecked<IGameplayInsightsModule>("GameplayInsights");
 		TSharedPtr<const TraceServices::IAnalysisSession> Session = UnrealInsightsModule.GetAnalysisSession();
 
-		GameplayInsightsModule.GetDebugViewCreator()->CreateDebugViews(SelectedComponent->ObjectId, TraceTime.Get(), *Session, DebugViews);
+		FRewindDebuggerViewCreators::CreateDebugViews(SelectedComponent->ObjectId, TraceTime.Get(), *Session, DebugViews);
 	}
 }
 
@@ -595,12 +596,12 @@ void SRewindDebugger::CreateDebugTabs()
 		CloseTab(TabName);
 	}
 
-	for(TSharedPtr<IGameplayInsightsDebugView> DebugView : DebugViews)
+	for(TSharedPtr<IRewindDebuggerView> DebugView : DebugViews)
 	{
 		FName ViewName = DebugView->GetName();
 		uint64 ObjectId = DebugView->GetObjectId();
 
-	    bool bPinned = nullptr != PinnedDebugViews.FindByPredicate([ViewName, ObjectId](TSharedPtr<IGameplayInsightsDebugView>& View) { return View->GetName() == ViewName && View->GetObjectId() == ObjectId; } );
+	    bool bPinned = nullptr != PinnedDebugViews.FindByPredicate([ViewName, ObjectId](TSharedPtr<IRewindDebuggerView>& View) { return View->GetName() == ViewName && View->GetObjectId() == ObjectId; } );
 	    bool bHidden = HiddenTabs.Find(ViewName) != INDEX_NONE;
 
 		if(!bPinned && !bHidden)
diff --git a/Engine/Plugins/Animation/GameplayInsights/Source/RewindDebugger/Private/SRewindDebugger.h b/Engine/Plugins/Animation/GameplayInsights/Source/RewindDebugger/Private/SRewindDebugger.h
--- a/Engine/Plugins/Animation/GameplayInsights/Source/RewindDebugger/Private/SRewindDebugger.h
+++ b/Engine/Plugins/Animation/GameplayInsights/Source/RewindDebugger/Private/SRewindDebugger.h
@@ -5,7 +5,7 @@
 #include "CoreMinimal.h"
 #include "Framework/Commands/UICommandList.h"
 #include "Framework/Docking/TabManager.h"
-#include "IGameplayInsightsDebugView.h"
+#include "IRewindDebuggerView.h"
 #include "RewindDebuggerModule.h"
 #include "RewindDebuggerTimeSliderController.h"
 #include "SRewindDebuggerComponentTree.h"
@@ -85,13 +85,13 @@ private:
 	bool CanSpawnTab(const FSpawnTabArgs& Args, FName ViewName);
 	void CloseTab(FName TabName);
 	void OnPinnedTabClosed(TSharedRef<SDockTab> Tab);
-	void ExtendTabMenu(FMenuBuilder& MenuBuilder, TSharedPtr<IGameplayInsightsDebugView> View);
-	void PinTab(TSharedPtr<IGameplayInsightsDebugView> View);
+	void ExtendTabMenu(FMenuBuilder& MenuBuilder, TSharedPtr<IRewindDebuggerView> View);
+	void PinTab(TSharedPtr<IRewindDebuggerView> View);
 	void ShowAllViews();
 	void CreateDebugViews();
 	void CreateDebugTabs();
-	TArray<TSharedPtr<IGameplayInsightsDebugView>> DebugViews;
-	TArray<TSharedPtr<IGameplayInsightsDebugView>> PinnedDebugViews;
+	TArray<TSharedPtr<IRewindDebuggerView>> DebugViews;
+	TArray<TSharedPtr<IRewindDebuggerView>> PinnedDebugViews;
 	TArray<FName> TabNames;
 	TArray<FName> HiddenTabs;  // keep track of tabs that have been closed so we don't automatically reopen them when switching components
 	bool bInternalClosingTab = false;
diff --git a/Engine/Plugins/Animation/GameplayInsights/Source/RewindDebugger/Public/IRewindDebugger.h b/Engine/Plugins/Animation/GameplayInsights/Source/RewindDebugger/Public/IRewindDebugger.h
deleted file mode 100644
diff --git a/Engine/Plugins/Animation/GameplayInsights/Source/RewindDebugger/Public/IRewindDebuggerExtension.h b/Engine/Plugins/Animation/GameplayInsights/Source/RewindDebugger/Public/IRewindDebuggerExtension.h
deleted file mode 100644
diff --git a/Engine/Plugins/Animation/GameplayInsights/Source/RewindDebugger/RewindDebugger.Build.cs b/Engine/Plugins/Animation/GameplayInsights/Source/RewindDebugger/RewindDebugger.Build.cs
--- a/Engine/Plugins/Animation/GameplayInsights/Source/RewindDebugger/RewindDebugger.Build.cs
+++ b/Engine/Plugins/Animation/GameplayInsights/Source/RewindDebugger/RewindDebugger.Build.cs
@@ -34,7 +34,8 @@ namespace UnrealBuildTool.Rules
 				"TraceAnalysis",
 				"TraceServices",
 				"TraceInsights",
-				"UnrealEd"
+				"UnrealEd",
+				"RewindDebuggerInterface"
 			});
 		}
 	}
diff --git a/Engine/Source/Editor/RewindDebuggerInterface/Private/RewindDebuggerInterface.cpp b/Engine/Source/Editor/RewindDebuggerInterface/Private/RewindDebuggerInterface.cpp
new file mode 100644
--- /dev/null
+++ b/Engine/Source/Editor/RewindDebuggerInterface/Private/RewindDebuggerInterface.cpp
@@ -0,0 +1,16 @@
+// Copyright Epic Games, Inc. All Rights Reserved.
+
+#include "IRewindDebuggerExtension.h"
+#include "IRewindDebuggerViewCreator.h"
+#include "IRewindDebugger.h"
+
+const FName IRewindDebuggerExtension::ModularFeatureName = "RewindDebuggerExtension";
+const FName IRewindDebuggerViewCreator::ModularFeatureName = "RewindDebuggerViewCreator";
+
+IRewindDebugger::~IRewindDebugger()
+{
+}
+
+IRewindDebugger::IRewindDebugger()
+{
+}
\ No newline at end of file
diff --git a/Engine/Source/Editor/RewindDebuggerInterface/Public/IRewindDebugger.h b/Engine/Source/Editor/RewindDebuggerInterface/Public/IRewindDebugger.h
new file mode 100644
--- /dev/null
+++ b/Engine/Source/Editor/RewindDebuggerInterface/Public/IRewindDebugger.h
@@ -0,0 +1,61 @@
+// Copyright Epic Games, Inc. All Rights Reserved.
+#pragma once
+
+#include "CoreMinimal.h"
+
+// IRewindDebugger
+//
+// Public interface to rewind debugger
+
+namespace TraceServices
+{
+	class IAnalysisSession;
+}
+
+struct FDebugObjectInfo
+{
+	FDebugObjectInfo(uint64 Id, const FString& Name): ObjectId(Id), ObjectName(Name), bExpanded(true)
+	{
+	}
+
+	uint64 ObjectId;
+	FString ObjectName;
+	bool bExpanded;
+
+	TArray<TSharedPtr<FDebugObjectInfo>> Children;
+};
+
+class REWINDDEBUGGERINTERFACE_API IRewindDebugger
+{
+public:
+	IRewindDebugger();
+
+	virtual ~IRewindDebugger();
+
+	// get the time the debugger is scrubbed to, in seconds since the capture started (or the recording duration while the game is running)
+	virtual double CurrentTraceTime() const = 0;
+
+	// get the current analysis session
+	virtual const TraceServices::IAnalysisSession* GetAnalysisSession() const = 0;
+
+	// get insights id for the selected target actor
+	virtual uint64 GetTargetActorId() const = 0;
+
+	// get a list of all components of the selected target actor (with the actor as the first element in the list)
+	virtual TArray<TSharedPtr<FDebugObjectInfo>>& GetDebugComponents() = 0;
+
+	// get posiotion of the selected target actor (returns true if position is valid)
+	virtual bool GetTargetActorPosition(FVector& OutPosition) const = 0;
+
+	// get the world that the debugger is replaying in
+	virtual UWorld* GetWorldToVisualize() const = 0;
+
+	// returns true if recording is active
+	virtual bool IsRecording() const = 0;
+
+	// returns true if PIE is running and not paused
+	virtual bool IsPIESimulating() const = 0;
+
+	// returns the length of the current recording
+	virtual double GetRecordingDuration() const = 0;
+};
diff --git a/Engine/Source/Editor/RewindDebuggerInterface/Public/IRewindDebuggerExtension.h b/Engine/Source/Editor/RewindDebuggerInterface/Public/IRewindDebuggerExtension.h
new file mode 100644
--- /dev/null
+++ b/Engine/Source/Editor/RewindDebuggerInterface/Public/IRewindDebuggerExtension.h
@@ -0,0 +1,29 @@
+// Copyright Epic Games, Inc. All Rights Reserved.
+#pragma once
+
+#include "CoreMinimal.h"
+#include "Features/IModularFeatures.h"
+
+
+// IRewindDebuggerExtension 
+//
+// interface class for extensions which add functionality to the rewind debugger
+// these will get updated on scrub/playback to handle things like updating the world state to match recorded data from that time for a particular system
+//
+
+class IRewindDebugger;
+
+class REWINDDEBUGGERINTERFACE_API IRewindDebuggerExtension : public IModularFeature
+{
+public:
+	static const FName ModularFeatureName;
+
+	// called while scrubbing, playing back, or paused
+	virtual void Update(float DeltaTime, IRewindDebugger* RewindDebugger) {};
+
+	// called when recording has started
+	virtual void RecordingStarted(IRewindDebugger* RewindDebugger) {};
+
+	// called when recording has ended
+	virtual void RecordingStopped(IRewindDebugger* RewindDebugger) {};
+};
diff --git a/Engine/Source/Editor/RewindDebuggerInterface/Public/IRewindDebuggerView.h b/Engine/Source/Editor/RewindDebuggerInterface/Public/IRewindDebuggerView.h
new file mode 100644
--- /dev/null
+++ b/Engine/Source/Editor/RewindDebuggerInterface/Public/IRewindDebuggerView.h
@@ -0,0 +1,22 @@
+// Copyright Epic Games, Inc. All Rights Reserved.
+
+#pragma once
+
+#include "CoreMinimal.h"
+#include "Widgets/SCompoundWidget.h"
+
+// Interface class for views of recorded debug data
+// a delegate which creates implementations of this interface can be registered with a UObject type name to IGameplayInsightsDebugViewCreator
+// and they will be shown in the Rewind Debugger when that type of object is selected
+class REWINDDEBUGGERINTERFACE_API IRewindDebuggerView : public SCompoundWidget
+{
+	public:
+		// unique name for widget
+		virtual FName GetName() const = 0;
+
+		// id of target object
+		virtual uint64 GetObjectId() const = 0;
+
+		// called by the debugger when the scrubbing bar position changes
+		virtual void SetTimeMarker(double InTimeMarker) = 0;
+};
diff --git a/Engine/Source/Editor/RewindDebuggerInterface/Public/IRewindDebuggerViewCreator.h b/Engine/Source/Editor/RewindDebuggerInterface/Public/IRewindDebuggerViewCreator.h
new file mode 100644
--- /dev/null
+++ b/Engine/Source/Editor/RewindDebuggerInterface/Public/IRewindDebuggerViewCreator.h
@@ -0,0 +1,34 @@
+// Copyright Epic Games, Inc. All Rights Reserved.
+
+#pragma once
+
+#include "CoreMinimal.h"
+#include "IRewindDebuggerView.h"
+#include "Textures/SlateIcon.h"
+
+namespace TraceServices
+{
+	class IAnalysisSession;
+}
+
+// Interface class which creates debug widgets
+class REWINDDEBUGGERINTERFACE_API IRewindDebuggerViewCreator : public IModularFeature
+{
+public: 
+	static const FName ModularFeatureName;
+
+	// returns the name of a type of UObject for which this debug view will be created
+	virtual FName GetTargetTypeName() const = 0;
+
+	// returns a unique name for identifying this type of widget (same value returned by IRewindDebuggerView::GetName)
+	virtual FName GetName() const = 0;
+
+	// text for tab header
+	virtual FText GetTitle() const = 0;
+
+	// icon for tab header
+	virtual FSlateIcon GetIcon() const = 0;
+
+	// creates and returns a widget, which will be displayed in Rewind Debugger
+	virtual TSharedPtr<IRewindDebuggerView> CreateDebugView(uint64 ObjectId, double CurrentTime, const TraceServices::IAnalysisSession& InAnalysisSession) const = 0;
+};
\ No newline at end of file
diff --git a/Engine/Source/Editor/RewindDebuggerInterface/RewindDebuggerInterface.Build.cs b/Engine/Source/Editor/RewindDebuggerInterface/RewindDebuggerInterface.Build.cs
new file mode 100644
--- /dev/null
+++ b/Engine/Source/Editor/RewindDebuggerInterface/RewindDebuggerInterface.Build.cs
@@ -0,0 +1,16 @@
+// Copyright Epic Games, Inc. All Rights Reserved.
+
+using UnrealBuildTool;
+
+public class RewindDebuggerInterface : ModuleRules
+{
+	public RewindDebuggerInterface(ReadOnlyTargetRules Target) : base(Target)
+	{
+		PrivateDependencyModuleNames.AddRange(
+			new string[] {
+				"Core",
+				"SlateCore",
+			}
+		);
+	}
+}
