fix unclickable links in rewind debugger property views

diff --git a/Engine/Plugins/Animation/GameplayInsights/Source/GameplayInsights/Private/SPropertiesDebugViewBase.cpp b/Engine/Plugins/Animation/GameplayInsights/Source/GameplayInsights/Private/SPropertiesDebugViewBase.cpp
--- a/Engine/Plugins/Animation/GameplayInsights/Source/GameplayInsights/Private/SPropertiesDebugViewBase.cpp
+++ b/Engine/Plugins/Animation/GameplayInsights/Source/GameplayInsights/Private/SPropertiesDebugViewBase.cpp
@@ -28,13 +28,16 @@ void SPropertiesDebugViewBase::Construct(const FArguments& InArgs, uint64 InObje
 
 void SPropertiesDebugViewBase::SetTimeMarker(double Time)
 {
-	TimeMarker = Time;
-	TraceServices::FAnalysisSessionReadScope SessionReadScope(*AnalysisSession);
-	const TraceServices::IFrameProvider& FramesProvider = TraceServices::ReadFrameProvider(*AnalysisSession);
-	TraceServices::FFrame MarkerFrame;
-	if(FramesProvider.GetFrameFromTime(ETraceFrameType::TraceFrameType_Game, TimeMarker, MarkerFrame))
+	if (TimeMarker != Time)
 	{
-		View->RequestRefresh(MarkerFrame);
+		TimeMarker = Time;
+		TraceServices::FAnalysisSessionReadScope SessionReadScope(*AnalysisSession);
+		const TraceServices::IFrameProvider& FramesProvider = TraceServices::ReadFrameProvider(*AnalysisSession);
+		TraceServices::FFrame MarkerFrame;
+		if(FramesProvider.GetFrameFromTime(ETraceFrameType::TraceFrameType_Game, TimeMarker, MarkerFrame))
+		{
+			View->RequestRefresh(MarkerFrame);
+		}
 	}
 }
 
