fix crashes in Object Properties View

diff --git a/Engine/Plugins/Animation/GameplayInsights/Source/GameplayInsights/Private/GameplayAnalyzer.cpp b/Engine/Plugins/Animation/GameplayInsights/Source/GameplayInsights/Private/GameplayAnalyzer.cpp
--- a/Engine/Plugins/Animation/GameplayInsights/Source/GameplayInsights/Private/GameplayAnalyzer.cpp
+++ b/Engine/Plugins/Animation/GameplayInsights/Source/GameplayInsights/Private/GameplayAnalyzer.cpp
@@ -22,10 +22,9 @@ void FGameplayAnalyzer::OnAnalysisBegin(const FOnAnalysisContext& Context)
 	Builder.RouteEvent(RouteId_World, "Object", "World");
 	Builder.RouteEvent(RouteId_View, "Object", "View");
 	Builder.RouteEvent(RouteId_ClassPropertyStringId, "Object", "ClassPropertyStringId");
-	Builder.RouteEvent(RouteId_ClassProperty, "Object", "ClassProperty");
 	Builder.RouteEvent(RouteId_PropertiesStart, "Object", "PropertiesStart");
 	Builder.RouteEvent(RouteId_PropertiesEnd, "Object", "PropertiesEnd");
-	Builder.RouteEvent(RouteId_PropertyValue, "Object", "PropertyValue");
+	Builder.RouteEvent(RouteId_PropertyValue, "Object", "PropertyValue2");
 }
 
 bool FGameplayAnalyzer::OnEvent(uint16 RouteId, EStyle Style, const FOnEventContext& Context)
@@ -121,16 +120,6 @@ bool FGameplayAnalyzer::OnEvent(uint16 RouteId, EStyle Style, const FOnEventCont
 		GameplayProvider.AppendClassPropertyStringId(Id, Value);
 		break;
 	}
-	case RouteId_ClassProperty:
-	{
-		uint64 ClassId = EventData.GetValue<uint64>("ClassId");
-		int32 Id = EventData.GetValue<int32>("Id");
-		int32 ParentId = EventData.GetValue<int32>("ParentId");
-		uint32 TypeId = EventData.GetValue<uint32>("TypeId");
-		uint32 KeyId = EventData.GetValue<uint32>("KeyId");
-		GameplayProvider.AppendClassProperty(ClassId, Id, ParentId, TypeId, KeyId);
-		break;
-	}
 	case RouteId_PropertiesStart:
 	{
 		uint64 Cycle = EventData.GetValue<uint64>("Cycle");
@@ -149,9 +138,12 @@ bool FGameplayAnalyzer::OnEvent(uint16 RouteId, EStyle Style, const FOnEventCont
 	{
 		uint64 Cycle = EventData.GetValue<uint64>("Cycle");
 		uint64 ObjectId = EventData.GetValue<uint64>("ObjectId");
-		int32 PropertyId = EventData.GetValue<int32>("PropertyId");
+		int32 ParentId = EventData.GetValue<int32>("ParentId");
+		uint32 TypeId = EventData.GetValue<uint32>("TypeId");
+		uint32 KeyId = EventData.GetValue<uint32>("KeyId");
+
 		FStringView Value; EventData.GetString("Value", Value);
-		GameplayProvider.AppendPropertyValue(ObjectId, Context.EventTime.AsSeconds(Cycle), Cycle, PropertyId, Value);
+		GameplayProvider.AppendPropertyValue(ObjectId, Context.EventTime.AsSeconds(Cycle), Cycle, ParentId, TypeId, KeyId, Value);
 		break;
 	}
 	}
diff --git a/Engine/Plugins/Animation/GameplayInsights/Source/GameplayInsights/Private/GameplayProvider.cpp b/Engine/Plugins/Animation/GameplayInsights/Source/GameplayInsights/Private/GameplayProvider.cpp
--- a/Engine/Plugins/Animation/GameplayInsights/Source/GameplayInsights/Private/GameplayProvider.cpp
+++ b/Engine/Plugins/Animation/GameplayInsights/Source/GameplayInsights/Private/GameplayProvider.cpp
@@ -416,26 +416,6 @@ void FGameplayProvider::AppendClassPropertyStringId(uint32 InStringId, const FSt
 	PropertyStrings.Add(InStringId, StoredString);
 }
 
-void FGameplayProvider::AppendClassProperty(uint64 InClassId, int32 InId, int32 InParentId, uint32 InTypeStringId, uint32 InKeyStringId)
-{
-	Session.WriteAccessCheck();
-
-	bHasAnyData = true;
-
-	if(int32* ClassInfoIndexPtr = ClassIdToIndexMap.Find(InClassId))
-	{
-		FClassInfo& ClassInfo = ClassInfos[*ClassInfoIndexPtr];
-
-		// Resize to accommodate this property if required
-		ClassInfo.Properties.SetNum(FMath::Max(ClassInfo.Properties.Num(), InId + 1));
-
-		FClassPropertyInfo& PropertyInfo = ClassInfo.Properties[InId];
-		PropertyInfo.ParentId = InParentId;
-		PropertyInfo.TypeStringId = InTypeStringId;
-		PropertyInfo.KeyStringId = InKeyStringId;
-	}
-}
-
 void FGameplayProvider::AppendPropertiesStart(uint64 InObjectId, double InTime, uint64 InEventId)
 {
 	Session.WriteAccessCheck();
@@ -495,7 +475,7 @@ void FGameplayProvider::AppendPropertiesEnd(uint64 InObjectId, double InTime)
 	}
 }
 
-void FGameplayProvider::AppendPropertyValue(uint64 InObjectId, double InTime, uint64 InEventId, int32 InPropertyId, const FStringView& InValue)
+void FGameplayProvider::AppendPropertyValue(uint64 InObjectId, double InTime, uint64 InEventId, int32 InParentId, uint32 InTypeStringId, uint32 InKeyStringId, const FStringView& InValue)
 {
 	Session.WriteAccessCheck();
 
@@ -519,9 +499,11 @@ void FGameplayProvider::AppendPropertyValue(uint64 InObjectId, double InTime, ui
 	if(Storage->OpenEventId == InEventId)
 	{
 		FObjectPropertyValue& Message = Storage->Values.AddDefaulted_GetRef();
-		Message.PropertyId = InPropertyId;
 		Message.Value = Session.StoreString(InValue);
 		Message.ValueAsFloat = FCString::Atof(Message.Value);
+		Message.ParentId = InParentId;
+		Message.TypeStringId = InTypeStringId;
+		Message.KeyStringId = InKeyStringId; 
 
 		Storage->OpenEvent.PropertyValueEndIndex = Storage->Values.Num();
 	}
diff --git a/Engine/Plugins/Animation/GameplayInsights/Source/GameplayInsights/Private/GameplayProvider.h b/Engine/Plugins/Animation/GameplayInsights/Source/GameplayInsights/Private/GameplayProvider.h
--- a/Engine/Plugins/Animation/GameplayInsights/Source/GameplayInsights/Private/GameplayProvider.h
+++ b/Engine/Plugins/Animation/GameplayInsights/Source/GameplayInsights/Private/GameplayProvider.h
@@ -59,9 +59,6 @@ public:
 	/** Add a class property string ID message */
 	void AppendClassPropertyStringId(uint32 InStringId, const FStringView& InString);
 
-	/** Add a class property message */
-	void AppendClassProperty(uint64 InClassId, int32 InId, int32 InParentId, uint32 InTypeStringId, uint32 InKeyStringId);
-
 	/** Add a properties start message */
 	void AppendPropertiesStart(uint64 InObjectId, double InTime, uint64 InEventId);
 
@@ -69,7 +66,7 @@ public:
 	void AppendPropertiesEnd(uint64 InObjectId, double InTime);
 
 	/** Add a property value message */
-	void AppendPropertyValue(uint64 InObjectId, double InTime, uint64 InEventId, int32 InPropertyId, const FStringView& InValue);
+	void AppendPropertyValue(uint64 InObjectId, double InTime, uint64 InEventId, int32 InParentId, uint32 InTypeStringId, uint32 InKeyStringId, const FStringView& InValue);
 
 	/** Check whether we have any data */
 	bool HasAnyData() const;
diff --git a/Engine/Plugins/Animation/GameplayInsights/Source/GameplayInsights/Private/ObjectPropertiesTrack.cpp b/Engine/Plugins/Animation/GameplayInsights/Source/GameplayInsights/Private/ObjectPropertiesTrack.cpp
--- a/Engine/Plugins/Animation/GameplayInsights/Source/GameplayInsights/Private/ObjectPropertiesTrack.cpp
+++ b/Engine/Plugins/Animation/GameplayInsights/Source/GameplayInsights/Private/ObjectPropertiesTrack.cpp
@@ -151,47 +151,31 @@ void FObjectPropertiesTrack::GetVariantsAtFrame(const TraceServices::FFrame& InF
 	{
 		TraceServices::FAnalysisSessionReadScope SessionReadScope(SharedData.GetAnalysisSession());
 
-		const FClassInfo& ClassInfo = GameplayProvider->GetClassInfoFromObject(GetGameplayTrack().GetObjectId());
+		TSharedRef<FVariantTreeNode> Header = OutVariants.Add_GetRef(FVariantTreeNode::MakeHeader(FText::FromString(GetName()), INDEX_NONE));
+		TArray<TSharedPtr<FVariantTreeNode>> PropertyVariants;
 
-		if(ClassInfo.Properties.Num() > 0)
+		GameplayProvider->ReadObjectPropertiesTimeline(GetGameplayTrack().GetObjectId(), [this, &InFrame, GameplayProvider, &Header, &PropertyVariants](const FGameplayProvider::ObjectPropertiesTimeline& InTimeline)
 		{
-			TSharedRef<FVariantTreeNode> Header = OutVariants.Add_GetRef(FVariantTreeNode::MakeHeader(FText::FromString(GetName()), INDEX_NONE));
-
-			// Build the class tree
-			TArray<TSharedPtr<FVariantTreeNode>> PropertyVariants;
-			PropertyVariants.SetNum(ClassInfo.Properties.Num());
-			for(int32 PropertyIndex = 0; PropertyIndex < ClassInfo.Properties.Num(); ++PropertyIndex)
+			InTimeline.EnumerateEvents(InFrame.StartTime, InFrame.EndTime, [this, GameplayProvider, &Header, &PropertyVariants](double InStartTime, double InEndTime, uint32 InDepth, const FObjectPropertiesMessage& InMessage)
 			{
-				const FClassPropertyInfo& PropertyInfo = ClassInfo.Properties[PropertyIndex];
-
-				// Add a string node with a default value
-				const TCHAR* Key = GameplayProvider->GetPropertyName(PropertyInfo.KeyStringId);
-				PropertyVariants[PropertyIndex] = FVariantTreeNode::MakeString(FText::FromString(Key), TEXT("Unknown"), PropertyIndex);
-
-				// note assumes that order is parent->child in the properties array
-				if(PropertyInfo.ParentId != INDEX_NONE)
+				GameplayProvider->EnumerateObjectPropertyValues(GetGameplayTrack().GetObjectId(), InMessage, [GameplayProvider, &Header, &PropertyVariants](const FObjectPropertyValue& InValue)
 				{
-					PropertyVariants[PropertyInfo.ParentId]->AddChild(PropertyVariants[PropertyIndex].ToSharedRef());
-				}
-				else
-				{
-					Header->AddChild(PropertyVariants[PropertyIndex].ToSharedRef());
-				}
-			}
+					const TCHAR* Key = GameplayProvider->GetPropertyName(InValue.KeyStringId);
+					PropertyVariants.Add(FVariantTreeNode::MakeString(FText::FromString(Key), InValue.Value));
 
-			// object events
-			GameplayProvider->ReadObjectPropertiesTimeline(GetGameplayTrack().GetObjectId(), [this, &InFrame, &GameplayProvider, &PropertyVariants](const FGameplayProvider::ObjectPropertiesTimeline& InTimeline)
-			{
-				InTimeline.EnumerateEvents(InFrame.StartTime, InFrame.EndTime, [this, &GameplayProvider, &PropertyVariants](double InStartTime, double InEndTime, uint32 InDepth, const FObjectPropertiesMessage& InMessage)
-				{
-					GameplayProvider->EnumerateObjectPropertyValues(GetGameplayTrack().GetObjectId(), InMessage, [&PropertyVariants](const FObjectPropertyValue& InValue)
+					// note assumes that order is parent->child in the properties array
+					if(InValue.ParentId != INDEX_NONE)
 					{
-						PropertyVariants[InValue.PropertyId]->GetValue().String.Value = InValue.Value;
-					});
-					return TraceServices::EEventEnumerate::Stop;
+						PropertyVariants[InValue.ParentId]->AddChild(PropertyVariants.Last().ToSharedRef());
+					}
+					else
+					{
+						Header->AddChild(PropertyVariants.Last().ToSharedRef());
+					}
 				});
+				return TraceServices::EEventEnumerate::Stop;
 			});
-		}
+		});
 	}
 }
 
diff --git a/Engine/Plugins/Animation/GameplayInsights/Source/GameplayInsights/Private/ObjectPropertyTrace.cpp b/Engine/Plugins/Animation/GameplayInsights/Source/GameplayInsights/Private/ObjectPropertyTrace.cpp
--- a/Engine/Plugins/Animation/GameplayInsights/Source/GameplayInsights/Private/ObjectPropertyTrace.cpp
+++ b/Engine/Plugins/Animation/GameplayInsights/Source/GameplayInsights/Private/ObjectPropertyTrace.cpp
@@ -15,14 +15,6 @@ UE_TRACE_EVENT_BEGIN(Object, ClassPropertyStringId, NoSync|Important)
 	UE_TRACE_EVENT_FIELD(UE::Trace::WideString, Value)
 UE_TRACE_EVENT_END()
 
-UE_TRACE_EVENT_BEGIN(Object, ClassProperty, NoSync|Important)
-	UE_TRACE_EVENT_FIELD(uint64, ClassId)
-	UE_TRACE_EVENT_FIELD(int32, Id)
-	UE_TRACE_EVENT_FIELD(int32, ParentId)
-	UE_TRACE_EVENT_FIELD(uint32, TypeId)
-	UE_TRACE_EVENT_FIELD(uint32, KeyId)
-UE_TRACE_EVENT_END()
-
 UE_TRACE_EVENT_BEGIN(Object, PropertiesStart)
 	UE_TRACE_EVENT_FIELD(uint64, Cycle)
 	UE_TRACE_EVENT_FIELD(uint64, ObjectId)
@@ -33,10 +25,12 @@ UE_TRACE_EVENT_BEGIN(Object, PropertiesEnd)
 	UE_TRACE_EVENT_FIELD(uint64, ObjectId)
 UE_TRACE_EVENT_END()
 
-UE_TRACE_EVENT_BEGIN(Object, PropertyValue)
+UE_TRACE_EVENT_BEGIN(Object, PropertyValue2)
 	UE_TRACE_EVENT_FIELD(uint64, Cycle)
 	UE_TRACE_EVENT_FIELD(uint64, ObjectId)
-	UE_TRACE_EVENT_FIELD(int32, PropertyId)
+	UE_TRACE_EVENT_FIELD(int32, ParentId)
+	UE_TRACE_EVENT_FIELD(uint32, TypeId)
+	UE_TRACE_EVENT_FIELD(uint32, KeyId)
 	UE_TRACE_EVENT_FIELD(UE::Trace::WideString, Value)
 UE_TRACE_EVENT_END()
 
@@ -48,8 +42,6 @@ namespace ObjectPropertyTrace
 	static uint32 CurrentClassPropertyStringId = 0;
 	static TMap<FString, uint32> StringIdMap;
 
-	static TSet<uint64> TracedClassIds;
-
 	typedef TFunctionRef<void(const FString&, const FString&, const FString&, int32, int32)> IterateFunction;
 
 	static uint32 TraceStringId(const FString& InString)
@@ -68,18 +60,6 @@ namespace ObjectPropertyTrace
 		return NewId;
 	}
 
-	static bool ShouldTraceClassProperties(uint64 InClassId)
-	{
-		if(uint64* ExistingIdPtr = TracedClassIds.Find(InClassId))
-		{
-			return false;
-		}
-
-		TracedClassIds.Add(InClassId);
-
-		return true;
-	}
-
 	static void IteratePropertiesRecursive(FProperty* InProperty, const void* InContainer, const FString& InKey, IterateFunction InFunction, int32& InId, int32 InParentId)
 	{
 		// Handle container properties
@@ -211,28 +191,19 @@ namespace ObjectPropertyTrace
 					<< PropertiesStart.ObjectId(ObjectId);
 				
 				uint64 ClassId = FObjectTrace::GetObjectId(TracedObject->GetClass());
-				const bool bTraceClassProperties = ShouldTraceClassProperties(ClassId);
 
-				IterateProperties(TracedObject->GetClass(), TracedObject, [StartCycle, ClassId, ObjectId, bTraceClassProperties](const FString& InType, const FString& InKey, const FString& InValue, int32 InId, int32 InParentId)
+				IterateProperties(TracedObject->GetClass(), TracedObject, [StartCycle, ClassId, ObjectId](const FString& InType, const FString& InKey, const FString& InValue, int32 InId, int32 InParentId)
 				{
-					if(bTraceClassProperties)
-					{
-						uint32 TypeId = TraceStringId(InType);
-						uint32 KeyId = TraceStringId(InKey);
-
-						UE_TRACE_LOG(Object, ClassProperty, ObjectProperties)
-							<< ClassProperty.ClassId(ClassId)
-							<< ClassProperty.Id(InId)
-							<< ClassProperty.ParentId(InParentId)
-							<< ClassProperty.TypeId(TypeId)
-							<< ClassProperty.KeyId(KeyId);
-					}
-
-					UE_TRACE_LOG(Object, PropertyValue, ObjectProperties)
-						<< PropertyValue.Cycle(StartCycle)
-						<< PropertyValue.ObjectId(ObjectId)
-						<< PropertyValue.PropertyId(InId)
-						<< PropertyValue.Value(*InValue);
+					uint32 TypeId = TraceStringId(InType);
+					uint32 KeyId = TraceStringId(InKey);
+
+					UE_TRACE_LOG(Object, PropertyValue2, ObjectProperties)
+						<< PropertyValue2.Cycle(StartCycle)
+						<< PropertyValue2.ObjectId(ObjectId)
+						<< PropertyValue2.Value(*InValue)
+						<< PropertyValue2.ParentId(InParentId)
+						<< PropertyValue2.TypeId(TypeId)
+						<< PropertyValue2.KeyId(KeyId);
 				});
 
 				uint64 EndCycle = FPlatformTime::Cycles64();
diff --git a/Engine/Plugins/Animation/GameplayInsights/Source/GameplayInsights/Private/SObjectPropertiesView.cpp b/Engine/Plugins/Animation/GameplayInsights/Source/GameplayInsights/Private/SObjectPropertiesView.cpp
--- a/Engine/Plugins/Animation/GameplayInsights/Source/GameplayInsights/Private/SObjectPropertiesView.cpp
+++ b/Engine/Plugins/Animation/GameplayInsights/Source/GameplayInsights/Private/SObjectPropertiesView.cpp
@@ -16,47 +16,32 @@ void SObjectPropertiesView::GetVariantsAtFrame(const TraceServices::FFrame& InFr
 	{
 		TraceServices::FAnalysisSessionReadScope SessionReadScope(*AnalysisSession);
 
-		const FClassInfo& ClassInfo = GameplayProvider->GetClassInfoFromObject(ObjectId);
+		TSharedRef<FVariantTreeNode> Header = OutVariants.Add_GetRef(FVariantTreeNode::MakeHeader(LOCTEXT("Properties", "Properties"), INDEX_NONE));
 
-		if(ClassInfo.Properties.Num() > 0)
-		{
-			TSharedRef<FVariantTreeNode> Header = OutVariants.Add_GetRef(FVariantTreeNode::MakeHeader(LOCTEXT("Properties", "Properties"), INDEX_NONE));
+		TArray<TSharedPtr<FVariantTreeNode>> PropertyVariants;
 
-			// Build the class tree
-			TArray<TSharedPtr<FVariantTreeNode>> PropertyVariants;
-			PropertyVariants.SetNum(ClassInfo.Properties.Num());
-			for(int32 PropertyIndex = 0; PropertyIndex < ClassInfo.Properties.Num(); ++PropertyIndex)
+		GameplayProvider->ReadObjectPropertiesTimeline(ObjectId, [this, &InFrame, GameplayProvider, &Header, &PropertyVariants](const FGameplayProvider::ObjectPropertiesTimeline& InTimeline)
+		{
+			InTimeline.EnumerateEvents(InFrame.StartTime, InFrame.EndTime, [this, GameplayProvider, &Header, &PropertyVariants](double InStartTime, double InEndTime, uint32 InDepth, const FObjectPropertiesMessage& InMessage)
 			{
-				const FClassPropertyInfo& PropertyInfo = ClassInfo.Properties[PropertyIndex];
-
-				// Add a string node with a default value
-				const TCHAR* Key = GameplayProvider->GetPropertyName(PropertyInfo.KeyStringId);
-				PropertyVariants[PropertyIndex] = FVariantTreeNode::MakeString(FText::FromString(Key), TEXT("Unknown"), PropertyIndex);
-
-				// note assumes that order is parent->child in the properties array
-				if(PropertyInfo.ParentId != INDEX_NONE)
+				GameplayProvider->EnumerateObjectPropertyValues(ObjectId, InMessage, [GameplayProvider, &Header, &PropertyVariants](const FObjectPropertyValue& InValue)
 				{
-					PropertyVariants[PropertyInfo.ParentId]->AddChild(PropertyVariants[PropertyIndex].ToSharedRef());
-				}
-				else
-				{
-					Header->AddChild(PropertyVariants[PropertyIndex].ToSharedRef());
-				}
-			}
+					const TCHAR* Key = GameplayProvider->GetPropertyName(InValue.KeyStringId);
+					PropertyVariants.Add(FVariantTreeNode::MakeString(FText::FromString(Key), InValue.Value));
 
-			// object events
-			GameplayProvider->ReadObjectPropertiesTimeline(ObjectId, [this, &InFrame, &GameplayProvider, &PropertyVariants](const FGameplayProvider::ObjectPropertiesTimeline& InTimeline)
-			{
-				InTimeline.EnumerateEvents(InFrame.StartTime, InFrame.EndTime, [this, &GameplayProvider, &PropertyVariants](double InStartTime, double InEndTime, uint32 InDepth, const FObjectPropertiesMessage& InMessage)
-				{
-					GameplayProvider->EnumerateObjectPropertyValues(ObjectId, InMessage, [&PropertyVariants](const FObjectPropertyValue& InValue)
+					// note assumes that order is parent->child in the properties array
+					if(InValue.ParentId != INDEX_NONE)
+					{
+						PropertyVariants[InValue.ParentId]->AddChild(PropertyVariants.Last().ToSharedRef());
+					}
+					else
 					{
-						PropertyVariants[InValue.PropertyId]->GetValue().String.Value = InValue.Value;
-					});
-					return TraceServices::EEventEnumerate::Stop;
+						Header->AddChild(PropertyVariants.Last().ToSharedRef());
+					}
 				});
+				return TraceServices::EEventEnumerate::Stop;
 			});
-		}
+		});
 	}
 }
 
diff --git a/Engine/Plugins/Animation/GameplayInsights/Source/GameplayInsights/Public/IGameplayProvider.h b/Engine/Plugins/Animation/GameplayInsights/Source/GameplayInsights/Public/IGameplayProvider.h
--- a/Engine/Plugins/Animation/GameplayInsights/Source/GameplayInsights/Public/IGameplayProvider.h
+++ b/Engine/Plugins/Animation/GameplayInsights/Source/GameplayInsights/Public/IGameplayProvider.h
@@ -6,20 +6,12 @@
 #include "TraceServices/Model/AnalysisSession.h"
 #include "TraceServices/Containers/Timelines.h"
 
-struct FClassPropertyInfo
-{
-	int32 ParentId = 0;
-	uint32 TypeStringId = 0;
-	uint32 KeyStringId = 0;
-};
-
 struct FClassInfo
 {	
 	uint64 Id = 0;
 	uint64 SuperId = 0;
 	const TCHAR* Name = nullptr;
 	const TCHAR* PathName = nullptr;
-	TArray<FClassPropertyInfo> Properties;
 };
 
 struct FObjectInfo
@@ -40,7 +32,9 @@ struct FObjectPropertiesMessage
 struct FObjectPropertyValue
 {
 	const TCHAR* Value = nullptr;
-	int32 PropertyId;
+	int32 ParentId = 0;
+	uint32 TypeStringId = 0;
+	uint32 KeyStringId = 0;
 	float ValueAsFloat = 0.0f;
 };
 
