Add support to VariantValueView for navigating to native classes in code when clicking on a class hyperlink.

diff --git a/Engine/Plugins/Animation/GameplayInsights/Source/GameplayInsights/Private/SVariantValueView.cpp b/Engine/Plugins/Animation/GameplayInsights/Source/GameplayInsights/Private/SVariantValueView.cpp
--- a/Engine/Plugins/Animation/GameplayInsights/Source/GameplayInsights/Private/SVariantValueView.cpp
+++ b/Engine/Plugins/Animation/GameplayInsights/Source/GameplayInsights/Private/SVariantValueView.cpp
@@ -24,6 +24,7 @@
 
 #if WITH_EDITOR
 #include "Subsystems/AssetEditorSubsystem.h"
+#include "SourceCodeNavigation.h"
 #include "Editor.h"
 #endif
 
@@ -191,7 +192,17 @@ static TSharedRef<SWidget> MakeVariantValueWidget(const TraceServices::IAnalysis
 				.ToolTipText(FText::Format(LOCTEXT("ClassHyperlinkTooltipFormat", "Open class '{0}'"), FText::FromString(ClassInfo.PathName)))
 				.OnNavigate_Lambda([ClassInfo]()
 				{
-					GEditor->GetEditorSubsystem<UAssetEditorSubsystem>()->OpenEditorForAsset(ClassInfo.PathName);
+					if (LoadPackage(NULL, ClassInfo.PathName, LOAD_NoRedirects)) // check if it's found as an asset
+					{
+						GEditor->GetEditorSubsystem<UAssetEditorSubsystem>()->OpenEditorForAsset(ClassInfo.PathName);
+					}
+					else // it must be a native class
+					{
+						if (UClass* FoundClass = FindObject<UClass>(ANY_PACKAGE, ClassInfo.Name))
+						{
+							FSourceCodeNavigation::NavigateToClass(FoundClass);
+						}
+					}
 				});
 #else
 			return 
