fix Blend space debugging in insights.
(mismatched trace names)

diff --git a/Engine/Plugins/Animation/GameplayInsights/Source/GameplayInsights/Private/AnimationAnalyzer.cpp b/Engine/Plugins/Animation/GameplayInsights/Source/GameplayInsights/Private/AnimationAnalyzer.cpp
--- a/Engine/Plugins/Animation/GameplayInsights/Source/GameplayInsights/Private/AnimationAnalyzer.cpp
+++ b/Engine/Plugins/Animation/GameplayInsights/Source/GameplayInsights/Private/AnimationAnalyzer.cpp
@@ -262,8 +262,8 @@ bool FAnimationAnalyzer::OnEvent(uint16 RouteId, EStyle Style, const FOnEventCon
 		uint64 AnimInstanceId = EventData.GetValue<uint64>("AnimInstanceId");
 		uint64 BlendSpaceId = EventData.GetValue<uint64>("BlendSpaceId");
 		int32 NodeId = EventData.GetValue<int32>("NodeId");
-		FVector BlendPosition(EventData.GetValue<float>("BlendPositionX"), EventData.GetValue<float>("BlendPositionY"), EventData.GetValue<float>("BlendPositionZ"));
-		FVector FilteredBlendPosition(EventData.GetValue<float>("FilteredBlendPositionX"), EventData.GetValue<float>("FilteredBlendPositionY"), EventData.GetValue<float>("FilteredBlendPositionZ"));
+		FVector BlendPosition(EventData.GetValue<float>("PositionX"), EventData.GetValue<float>("PositionY"), EventData.GetValue<float>("PositionZ"));
+		FVector FilteredBlendPosition(EventData.GetValue<float>("FilteredPositionX"), EventData.GetValue<float>("FilteredPositionY"), EventData.GetValue<float>("FilteredPositionZ"));
 		AnimationProvider.AppendBlendSpacePlayer(AnimInstanceId, Context.EventTime.AsSeconds(Cycle), NodeId, BlendSpaceId, BlendPosition, FilteredBlendPosition);
 		break;
 	}
