Rewind debugger scrubbing cleanup

diff --git a/Engine/Plugins/Animation/GameplayInsights/Source/RewindDebugger/Private/RewindDebugger.cpp b/Engine/Plugins/Animation/GameplayInsights/Source/RewindDebugger/Private/RewindDebugger.cpp
--- a/Engine/Plugins/Animation/GameplayInsights/Source/RewindDebugger/Private/RewindDebugger.cpp
+++ b/Engine/Plugins/Animation/GameplayInsights/Source/RewindDebugger/Private/RewindDebugger.cpp
@@ -9,6 +9,7 @@
 #include "Modules/ModuleManager.h"
 #include "ObjectTrace.h"
 #include "TraceServices/ITraceServicesModule.h"
+#include "TraceServices/Model/Frames.h"
 #include "LevelEditor.h"
 #include "SLevelViewport.h"
 #include "IRewindDebuggerExtension.h"
@@ -37,6 +38,7 @@ FRewindDebugger::FRewindDebugger()  :
 	bRecording(false),
 	PlaybackRate(1),
 	CurrentScrubTime(0),
+	ScrubFrameIndex(0),
 	RecordingIndex(0),
 	bTargetActorPositionValid(false)
 {
@@ -425,26 +427,46 @@ void FRewindDebugger::ScrubToEnd()
 	}
 }
 
-void FRewindDebugger::StepForward()
+void FRewindDebugger::Step(int frames)
 {
 	if (CanScrub())
 	{
 		Pause();
 
-		SetCurrentScrubTime(FMath::Min(CurrentScrubTime + 1.0f/30.0f, RecordingDuration.Get()));
-		TrackCursorDelegate.ExecuteIfBound(false);
+		if (const TraceServices::IAnalysisSession* Session = GetAnalysisSession())
+		{
+			TraceServices::FAnalysisSessionReadScope SessionReadScope(*Session);
+			UWorld* World = GetWorldToVisualize();
+
+			if (const IGameplayProvider* GameplayProvider = Session->ReadProvider<IGameplayProvider>("GameplayProvider"))
+			{
+				if (const IGameplayProvider::RecordingInfoTimeline* Recording = GameplayProvider->GetRecordingInfo(RecordingIndex))
+				{
+					uint64 EventCount = Recording->GetEventCount();
+
+					if (EventCount > 0)
+					{
+						ScrubFrameIndex = FMath::Clamp(ScrubFrameIndex + frames, 0, (int64)EventCount - 1);
+
+						const FRecordingInfoMessage& Event = Recording->GetEvent(ScrubFrameIndex);
+						CurrentScrubTime = Event.ElapsedTime;
+						TraceTime.Set(Event.ProfileTime);
+						TrackCursorDelegate.ExecuteIfBound(false);
+					}
+				}
+			}
+		}
 	}
 }
 
+void FRewindDebugger::StepForward()
+{
+	Step(1);
+}
+
 void FRewindDebugger::StepBackward()
 {
-	if (CanScrub())
-	{
-		Pause();
-		// todo: snap to actual frames from AnimationProvider
-		SetCurrentScrubTime(FMath::Max(CurrentScrubTime - 1.0f/30.0f, 0.0f));
-		TrackCursorDelegate.ExecuteIfBound(true);
-	}
+	Step(-1);
 }
 
 
@@ -502,19 +524,64 @@ void FRewindDebugger::UpdateTraceTime()
 
 				if (EventCount > 0)
 				{
-					uint64 EventIndex;
-
-					for(EventIndex = 0; EventIndex < EventCount-1; EventIndex++)
+					// check if we are outside of the recorded range, and apply the first or last frame
+					const FRecordingInfoMessage& FirstEvent = Recording->GetEvent(0);
+					const FRecordingInfoMessage& LastEvent = Recording->GetEvent(EventCount - 1);
+					if (CurrentScrubTime <= FirstEvent.ElapsedTime)
 					{
-						const FRecordingInfoMessage& Event = Recording->GetEvent(EventIndex);
-						// todo: optimize this linear search (cache current index, and pass it in as a hint)
-						//		 and do interpolation between nearest two frames
-						if (Event.ElapsedTime >= CurrentScrubTime)
+						ScrubFrameIndex = FMath::Min<uint64>(1, EventCount - 1);
+					}
+					else if (CurrentScrubTime >= LastEvent.ElapsedTime)
+					{
+						ScrubFrameIndex = EventCount - 1;
+					}
+					else
+					{
+						// find the two keys surrounding the CurrentScrubTime, and pick the nearest to update ProfileTime
+						if (Recording->GetEvent(ScrubFrameIndex).ElapsedTime > CurrentScrubTime)
 						{
-							break;
+							for (uint64 EventIndex = ScrubFrameIndex; EventIndex > 0; EventIndex--)
+							{
+								const FRecordingInfoMessage& Event = Recording->GetEvent(EventIndex);
+								const FRecordingInfoMessage& NextEvent = Recording->GetEvent(EventIndex - 1);
+								if (Event.ElapsedTime >= CurrentScrubTime && NextEvent.ElapsedTime <= CurrentScrubTime)
+								{
+									if (Event.ElapsedTime - CurrentScrubTime < CurrentScrubTime - NextEvent.ElapsedTime)
+									{
+										ScrubFrameIndex = EventIndex;
+									}
+									else
+									{
+										ScrubFrameIndex = EventIndex - 1;
+									}
+									break;
+								}
+							}
+						}
+						else
+						{
+							for (uint64 EventIndex = ScrubFrameIndex; EventIndex < EventCount - 1; EventIndex++)
+							{
+								const FRecordingInfoMessage& Event = Recording->GetEvent(EventIndex);
+								const FRecordingInfoMessage& NextEvent = Recording->GetEvent(EventIndex + 1);
+								if (Event.ElapsedTime <= CurrentScrubTime && NextEvent.ElapsedTime >= CurrentScrubTime)
+								{
+									if (CurrentScrubTime - Event.ElapsedTime < NextEvent.ElapsedTime - CurrentScrubTime)
+									{
+										ScrubFrameIndex = EventIndex;
+									}
+									else
+									{
+										ScrubFrameIndex = EventIndex + 1;
+									}
+									break;
+								}
+							}
 						}
 					}
-					TraceTime.Set(Recording->GetEvent(EventIndex).ProfileTime);
+
+					const FRecordingInfoMessage& Event = Recording->GetEvent(ScrubFrameIndex);
+					TraceTime.Set(Event.ProfileTime);
 				}
 			}
 		}
@@ -559,7 +626,7 @@ void FRewindDebugger::Tick(float DeltaTime)
 			{
 				if (RecordingDuration.Get() > 0)
 				{
-					UpdateTraceTime(); 
+					UpdateTraceTime();
 
 					if (ControlState == EControlState::Play || ControlState == EControlState::PlayReverse)
 					{
@@ -578,67 +645,64 @@ void FRewindDebugger::Tick(float DeltaTime)
 					uint64 TargetActorId = GetTargetActorId();
 
 					// update pose on all SkeletalMeshComponents (todo: move pose updating into an extension)
-					ULevel* CurLevel = World->GetCurrentLevel();
-					for( AActor* LevelActor : CurLevel->Actors )
+					const TraceServices::IFrameProvider& FrameProvider = TraceServices::ReadFrameProvider(*Session);
+					TraceServices::FFrame Frame;
+					if(FrameProvider.GetFrameFromTime(ETraceFrameType::TraceFrameType_Game, CurrentTraceTime, Frame))
 					{
-						if (LevelActor)
+						ULevel* CurLevel = World->GetCurrentLevel();
+						for (AActor* LevelActor : CurLevel->Actors)
 						{
-							bool bIsTargetActor = TargetActorId == FObjectTrace::GetObjectId(LevelActor);
-
-							if (bIsTargetActor)
+							if (LevelActor)
 							{
-								bTargetActorPositionValid = false;
-							}
-
-							TInlineComponentArray<USkeletalMeshComponent*> SkeletalMeshComponents;
-							LevelActor->GetComponents(SkeletalMeshComponents);
+								bool bIsTargetActor = TargetActorId == FObjectTrace::GetObjectId(LevelActor);
 
-							for (USkeletalMeshComponent* MeshComponent : SkeletalMeshComponents)
-							{
-								int64 ObjectId = FObjectTrace::GetObjectId(MeshComponent);
-
-								AnimationProvider->ReadSkeletalMeshPoseTimeline(ObjectId, [this, bIsTargetActor, ObjectId, CurrentTraceTime, MeshComponent, AnimationProvider](const IAnimationProvider::SkeletalMeshPoseTimeline& TimelineData, bool bHasCurves)
+								if (bIsTargetActor)
 								{
-									double PrecedingPoseTime;
-									double FollowingPoseTime;
-									const FSkeletalMeshPoseMessage* PrecedingPose;
-									const FSkeletalMeshPoseMessage* FollowingPose;
-
-									TimelineData.FindNearestEvents(CurrentTraceTime, PrecedingPose, PrecedingPoseTime, FollowingPose, FollowingPoseTime);
+									bTargetActorPositionValid = false;
+								}
 
-									if (FollowingPose || PrecedingPose)
-									{	
-										// Ideally we should iterpolate between the two poses here if both are valid (we would need to transform the component space transforms to local space though)
-										const FSkeletalMeshPoseMessage& PoseMessage = FollowingPose ? *FollowingPose : *PrecedingPose;
+								TInlineComponentArray<USkeletalMeshComponent*> SkeletalMeshComponents;
+								LevelActor->GetComponents(SkeletalMeshComponents);
 
-										FTransform ComponentWorldTransform;
-										const FSkeletalMeshInfo* SkeletalMeshInfo = AnimationProvider->FindSkeletalMeshInfo(PoseMessage.MeshId);
-										AnimationProvider->GetSkeletalMeshComponentSpacePose(PoseMessage, *SkeletalMeshInfo, ComponentWorldTransform, MeshComponent->GetEditableComponentSpaceTransforms());
-
-										if (MeshComponentsToReset.Find(ObjectId) == nullptr)
-										{
-											FMeshComponentResetData ResetData;
-											ResetData.Component = MeshComponent;
-											ResetData.RelativeTransform = MeshComponent->GetRelativeTransform();
-											MeshComponentsToReset.Add(ObjectId, ResetData);
-										}
-
-										// todo: we probably need to take into account tick order requirements for attached objects here
-										MeshComponent->SetWorldTransform(ComponentWorldTransform);
-										MeshComponent->SetForcedLOD(PoseMessage.LodIndex + 1);
+								for (USkeletalMeshComponent* MeshComponent : SkeletalMeshComponents)
+								{
+									int64 ObjectId = FObjectTrace::GetObjectId(MeshComponent);
 
-										if (bIsTargetActor && !bTargetActorPositionValid)
+									AnimationProvider->ReadSkeletalMeshPoseTimeline(ObjectId, [this, Frame, bIsTargetActor, ObjectId, MeshComponent, AnimationProvider](const IAnimationProvider::SkeletalMeshPoseTimeline& TimelineData, bool bHasCurves)
 										{
-											// until we have actor transforms traced out, the first skeletal mesh component transform will be used as as the actor position 
-											bTargetActorPositionValid = true;
-											TargetActorPosition = ComponentWorldTransform.GetTranslation();
-										}
-									}
-								});
-
-								// calling this here, even on meshes which have no recorded data makes meshes with a MasterPoseComponent attach properly during replay
-								// longer term solution: we should be recording those meshes
-								MeshComponent->ApplyEditedComponentSpaceTransforms();
+											TimelineData.EnumerateEvents(Frame.StartTime, Frame.EndTime,
+												[this, bIsTargetActor, ObjectId, MeshComponent, AnimationProvider](double InStartTime, double InEndTime, uint32 InDepth, const FSkeletalMeshPoseMessage& PoseMessage)
+												{
+													FTransform ComponentWorldTransform;
+													const FSkeletalMeshInfo* SkeletalMeshInfo = AnimationProvider->FindSkeletalMeshInfo(PoseMessage.MeshId);
+													AnimationProvider->GetSkeletalMeshComponentSpacePose(PoseMessage, *SkeletalMeshInfo, ComponentWorldTransform, MeshComponent->GetEditableComponentSpaceTransforms());
+
+													if (MeshComponentsToReset.Find(ObjectId) == nullptr)
+													{
+														FMeshComponentResetData ResetData;
+														ResetData.Component = MeshComponent;
+														ResetData.RelativeTransform = MeshComponent->GetRelativeTransform();
+														MeshComponentsToReset.Add(ObjectId, ResetData);
+													}
+
+													// todo: we probably need to take into account tick order requirements for attached objects here
+													MeshComponent->SetWorldTransform(ComponentWorldTransform);
+													MeshComponent->SetForcedLOD(PoseMessage.LodIndex + 1);
+
+													if (bIsTargetActor && !bTargetActorPositionValid)
+													{
+														// until we have actor transforms traced out, the first skeletal mesh component transform will be used as as the actor position 
+														bTargetActorPositionValid = true;
+														TargetActorPosition = ComponentWorldTransform.GetTranslation();
+													}
+													return TraceServices::EEventEnumerate::Stop;
+												});
+										});
+
+									// calling this here, even on meshes which have no recorded data makes meshes with a MasterPoseComponent attach properly during replay
+									// longer term solution: we should be recording those meshes
+									MeshComponent->ApplyEditedComponentSpaceTransforms();
+								}
 							}
 						}
 					}
diff --git a/Engine/Plugins/Animation/GameplayInsights/Source/RewindDebugger/Private/RewindDebugger.h b/Engine/Plugins/Animation/GameplayInsights/Source/RewindDebugger/Private/RewindDebugger.h
--- a/Engine/Plugins/Animation/GameplayInsights/Source/RewindDebugger/Private/RewindDebugger.h
+++ b/Engine/Plugins/Animation/GameplayInsights/Source/RewindDebugger/Private/RewindDebugger.h
@@ -70,6 +70,7 @@ public:
 
 	void ScrubToStart();
 	void ScrubToEnd();
+	void Step(int32 frames);
 	void StepForward();
 	void StepBackward();
 
@@ -131,6 +132,7 @@ private:
 
 	float PlaybackRate;
 	float CurrentScrubTime;
+	int64 ScrubFrameIndex;
 	uint16 RecordingIndex;
 
 	static FRewindDebugger* InternalInstance;
diff --git a/Engine/Plugins/Animation/GameplayInsights/Source/RewindDebugger/Private/RewindDebuggerCamera.cpp b/Engine/Plugins/Animation/GameplayInsights/Source/RewindDebugger/Private/RewindDebuggerCamera.cpp
--- a/Engine/Plugins/Animation/GameplayInsights/Source/RewindDebugger/Private/RewindDebuggerCamera.cpp
+++ b/Engine/Plugins/Animation/GameplayInsights/Source/RewindDebugger/Private/RewindDebuggerCamera.cpp
@@ -12,6 +12,7 @@
 #include "SLevelViewport.h"
 #include "ToolMenu.h"
 #include "ToolMenus.h"
+#include "TraceServices/Model/Frames.h"
 
 #define LOCTEXT_NAMESPACE "RewindDebuggerCamera"
 
@@ -120,32 +121,31 @@ void FRewindDebuggerCamera::Update(float DeltaTime, IRewindDebugger* RewindDebug
 			// always update the camera actor to the replay values even if it isn't locked
 			if (const IGameplayProvider* GameplayProvider = Session->ReadProvider<IGameplayProvider>("GameplayProvider"))
 			{
-				GameplayProvider->ReadViewTimeline([this, RewindDebugger, CurrentTraceTime](const IGameplayProvider::ViewTimeline& TimelineData)
+				GameplayProvider->ReadViewTimeline([this, RewindDebugger, CurrentTraceTime, Session](const IGameplayProvider::ViewTimeline& TimelineData)
 				{
-					double PrecedingTime;
-					double FollowingTime;
-					const FViewMessage* PrecedingView;
-					const FViewMessage* FollowingView;
-					
-					TimelineData.FindNearestEvents(CurrentTraceTime, PrecedingView, PrecedingTime, FollowingView, FollowingTime);
-
-					if (FollowingView || PrecedingView)
+					const TraceServices::IFrameProvider& FrameProvider = TraceServices::ReadFrameProvider(*Session);
+					TraceServices::FFrame Frame;
+					if(FrameProvider.GetFrameFromTime(ETraceFrameType::TraceFrameType_Game, CurrentTraceTime, Frame))
 					{
-						const FViewMessage& ViewMessage = FollowingView ? *FollowingView : *PrecedingView;
-
-						if (!CameraActor.IsValid())
-						{
-							FActorSpawnParameters SpawnParameters;
-							SpawnParameters.ObjectFlags |= RF_Transient;
-							CameraActor = RewindDebugger->GetWorldToVisualize()->SpawnActor<ACameraActor>(ViewMessage.Position, ViewMessage.Rotation, SpawnParameters);
-							UCameraComponent* Camera = CameraActor->GetCameraComponent();
-							CameraActor->SetActorLabel("RewindDebuggerCamera"); 
-						}
-
-						UCameraComponent* Camera = CameraActor->GetCameraComponent();
-						Camera->SetWorldLocationAndRotation(ViewMessage.Position, ViewMessage.Rotation);
-						Camera->SetFieldOfView(ViewMessage.Fov);
-						Camera->SetAspectRatio(ViewMessage.AspectRatio);
+						TimelineData.EnumerateEvents(Frame.StartTime, Frame.EndTime,
+							[this, RewindDebugger](double InStartTime, double InEndTime, uint32 InDepth, const FViewMessage& ViewMessage)
+							{
+								if (!CameraActor.IsValid())
+								{
+									FActorSpawnParameters SpawnParameters;
+									SpawnParameters.ObjectFlags |= RF_Transient;
+									CameraActor = RewindDebugger->GetWorldToVisualize()->SpawnActor<ACameraActor>(ViewMessage.Position, ViewMessage.Rotation, SpawnParameters);
+									UCameraComponent* Camera = CameraActor->GetCameraComponent();
+									CameraActor->SetActorLabel("RewindDebuggerCamera"); 
+								}
+
+								UCameraComponent* Camera = CameraActor->GetCameraComponent();
+								Camera->SetWorldLocationAndRotation(ViewMessage.Position, ViewMessage.Rotation);
+								Camera->SetFieldOfView(ViewMessage.Fov);
+								Camera->SetAspectRatio(ViewMessage.AspectRatio);
+
+								return TraceServices::EEventEnumerate::Stop;
+							});
 					}
 				});
 			}
diff --git a/Engine/Plugins/Experimental/Animation/PoseSearch/Source/Editor/Private/PoseSearchDebugger.cpp b/Engine/Plugins/Experimental/Animation/PoseSearch/Source/Editor/Private/PoseSearchDebugger.cpp
--- a/Engine/Plugins/Experimental/Animation/PoseSearch/Source/Editor/Private/PoseSearchDebugger.cpp
+++ b/Engine/Plugins/Experimental/Animation/PoseSearch/Source/Editor/Private/PoseSearchDebugger.cpp
@@ -6,6 +6,7 @@
 #include "RewindDebuggerInterface/Public/IRewindDebugger.h"
 #include "ObjectTrace.h"
 #include "TraceServices/Model/AnalysisSession.h"
+#include "TraceServices/Model/Frames.h"
 #include "Trace/PoseSearchTraceProvider.h"
 #include "Widgets/Layout/SScrollBox.h"
 #include "Widgets/Layout/SWidgetSwitcher.h"
@@ -1266,30 +1267,22 @@ void FDebugger::OnSelectionChanged(uint64 AnimInstanceId, int32 NodeId)
 	}
 
 	const double TraceTime = RewindDebugger->CurrentTraceTime();
-	TraceProvider->ReadMotionMatchingStateTimeline(AnimInstanceId, NodeId, [this, TraceTime](const FTraceProvider::FMotionMatchingStateTimeline& TimelineData)
+	TraceProvider->ReadMotionMatchingStateTimeline(AnimInstanceId, NodeId, [this, Session, TraceTime](const FTraceProvider::FMotionMatchingStateTimeline& TimelineData)
 	{
-		double PrecedingValueTime;
-		double FollowingValueTime;
-		const FTraceMotionMatchingStateMessage* PrecedingValue;
-		const FTraceMotionMatchingStateMessage* FollowingValue;
-
-		// Get event data at the trace time
-		TimelineData.FindNearestEvents(
-			TraceTime,
-			PrecedingValue,
-			PrecedingValueTime,
-			FollowingValue,
-			FollowingValueTime
-		);
-
-		// Grab following event if it exists, otherwise get the preceding (mimicking skeletal mesh updates in the Rewind Debugger)
-		if (PrecedingValue || FollowingValue)
+		const TraceServices::IFrameProvider& FrameProvider = TraceServices::ReadFrameProvider(*Session);
+		TraceServices::FFrame Frame;
+		if(FrameProvider.GetFrameFromTime(ETraceFrameType::TraceFrameType_Game, TraceTime, Frame))
 		{
-			MotionMatchingState = FollowingValue ? FollowingValue : PrecedingValue;
+			TimelineData.EnumerateEvents(Frame.StartTime, Frame.EndTime,
+				[this](double InStartTime, double InEndTime, uint32 InDepth, const FTraceMotionMatchingStateMessage& Message)
+				{
+					MotionMatchingState = &Message;
+
+					return TraceServices::EEventEnumerate::Stop;
+				});
 		}
 	});
 
-
 	UpdateReflection();
 }
 
diff --git a/Engine/Source/Developer/TraceServices/Public/Model/IntervalTimeline.h b/Engine/Source/Developer/TraceServices/Public/Model/IntervalTimeline.h
--- a/Engine/Source/Developer/TraceServices/Public/Model/IntervalTimeline.h
+++ b/Engine/Source/Developer/TraceServices/Public/Model/IntervalTimeline.h
@@ -76,71 +76,6 @@ public:
 			}
 		}
 	}
-
-	virtual void FindNearestEvents(double InTime, const EventType*& OutPrecedingEvent, double& OutPrecedingEventTime, const EventType*& OutFollowingEvent, double& OutFollowingEventTime) const override
-	{
-		OutPrecedingEvent =  nullptr;
-		OutFollowingEvent =  nullptr;
-
-		if (Events.Num() == 0)
-		{
-			return;
-		}
-
-		auto EventIterator = Events.GetIteratorFromPage(0);
-		const FEventInternal* Event = EventIterator.GetCurrentItem();
-		const FEventPage* Page = EventIterator.GetCurrentPage();
-
-		while(Event != nullptr && InTime > Page->EndTime)
-		{
-			EventIterator.NextPage();
-			Page = EventIterator.GetCurrentPage();
-			Event = EventIterator.GetCurrentItem();
-		}
-
-		if (Event == nullptr)
-		{
-			// InTime is past the end, so just output the last event and no following event
-			const FEventInternal& LastEvent = Events.Last();
-			OutPrecedingEvent = &LastEvent.Event;
-			OutPrecedingEventTime = LastEvent.StartTime;
-			return;
-		}
-		if (Event->StartTime > InTime)
-		{
-			// first event on page is past our time, so it will be the following event, preceding event is the last event from the previous page
-			OutFollowingEvent = &Event->Event;
-			OutFollowingEventTime = Event->StartTime;
-
-			EventIterator.PrevPage();
-			if (Event == EventIterator.GetCurrentItem())
-			{
-				OutPrecedingEvent = &Event->Event;
-				OutPrecedingEventTime = Event->StartTime;
-			}
-
-			return;
-		}
-
-		// normal case, just find the event before and after InTime on this page
-		while(Event != nullptr)
-		{
-			const FEventInternal* NextEvent = EventIterator.NextItem();
-			if (NextEvent)
-			{
-				if (NextEvent->StartTime > InTime)
-				{
-					OutPrecedingEvent = &Event->Event;
-					OutPrecedingEventTime = Event->StartTime;
-					OutFollowingEvent = &NextEvent->Event;
-					OutFollowingEventTime = Event->StartTime;
-					return;
-				}
-
-			}
-			Event = NextEvent;
-		}
-	}
 	
 	uint64 AppendBeginEvent(double StartTime, const EventType& Event)
 	{
diff --git a/Engine/Source/Developer/TraceServices/Public/Model/MonotonicTimeline.h b/Engine/Source/Developer/TraceServices/Public/Model/MonotonicTimeline.h
--- a/Engine/Source/Developer/TraceServices/Public/Model/MonotonicTimeline.h
+++ b/Engine/Source/Developer/TraceServices/Public/Model/MonotonicTimeline.h
@@ -160,13 +160,6 @@ public:
 		return DetailLevels[0].Events[InIndex]; 
 	}
 
-	virtual void FindNearestEvents(double InTime, const EventType*& OutPrecedingEvent, double& OutPrecedingEventTime, const EventType*& OutFollowingEvent, double& OutFollowingEventTime) const override
-	{
-		// not implemented
-		OutPrecedingEvent =  nullptr;
-		OutFollowingEvent =  nullptr;
-	}
-
 	virtual double GetStartTime() const override
 	{
 		return DetailLevels[0].ScopeEntries.Num() > 0 ? FMath::Abs(DetailLevels[0].ScopeEntries[0].Time) : 0.0;
diff --git a/Engine/Source/Developer/TraceServices/Public/Model/PointTimeline.h b/Engine/Source/Developer/TraceServices/Public/Model/PointTimeline.h
--- a/Engine/Source/Developer/TraceServices/Public/Model/PointTimeline.h
+++ b/Engine/Source/Developer/TraceServices/Public/Model/PointTimeline.h
@@ -144,71 +144,6 @@ public:
 		return false;
 	}
 
-	virtual void FindNearestEvents(double InTime, const EventType*& OutPrecedingEvent, double& OutPrecedingEventTime, const EventType*& OutFollowingEvent, double& OutFollowingEventTime) const override
-	{
-		OutPrecedingEvent =  nullptr;
-		OutFollowingEvent =  nullptr;
-
-		if (Events.Num() == 0)
-		{
-			return;
-		}
-
-		auto EventIterator = Events.GetIteratorFromPage(0);
-		const FEventInternal* Event = EventIterator.GetCurrentItem();
-		const FEventPage* Page = EventIterator.GetCurrentPage();
-
-		while(Event != nullptr && InTime > Page->EndTime)
-		{
-			EventIterator.NextPage();
-			Page = EventIterator.GetCurrentPage();
-			Event = EventIterator.GetCurrentItem();
-		}
-
-		if (Event == nullptr)
-		{
-			// InTime is past the end, so just output the last event and no following event
-			const FEventInternal& LastEvent = Events.Last();
-			OutPrecedingEvent = &LastEvent.Event;
-			OutPrecedingEventTime = LastEvent.Time;
-			return;
-		}
-		if (Event->Time > InTime)
-		{
-			// first event on page is past our time, so it will be the following event, preceding event is the last event from the previous page
-			OutFollowingEvent = &Event->Event;
-			OutFollowingEventTime = Event->Time;
-
-			EventIterator.PrevPage();
-			if (Event == EventIterator.GetCurrentItem())
-			{
-				OutPrecedingEvent = &Event->Event;
-				OutPrecedingEventTime = Event->Time;
-			}
-
-			return;
-		}
-
-		// normal case, just find the event before and after InTime on this page
-		while(Event != nullptr)
-		{
-			const FEventInternal* NextEvent = EventIterator.NextItem();
-			if (NextEvent)
-			{
-				if (NextEvent->Time > InTime)
-				{
-					OutPrecedingEvent = &Event->Event;
-					OutPrecedingEventTime = Event->Time;
-					OutFollowingEvent = &NextEvent->Event;
-					OutFollowingEventTime = Event->Time;
-					return;
-				}
-
-			}
-			Event = NextEvent;
-		}
-	}
-
 private:
 	struct FEventInternal
 	{
diff --git a/Engine/Source/Developer/TraceServices/Public/TraceServices/Containers/Timelines.h b/Engine/Source/Developer/TraceServices/Public/TraceServices/Containers/Timelines.h
--- a/Engine/Source/Developer/TraceServices/Public/TraceServices/Containers/Timelines.h
+++ b/Engine/Source/Developer/TraceServices/Public/TraceServices/Containers/Timelines.h
@@ -59,7 +59,6 @@ public:
 	virtual void EnumerateEvents(double IntervalStart, double IntervalEnd, EventCallback Callback) const = 0;
 	virtual void EnumerateEvents(double IntervalStart, double IntervalEnd, EventRangeCallback Callback) const = 0;
 	virtual int32 GetDepthAt(double Time) const { return 0; };
-	virtual void FindNearestEvents(double InTime, const EventType*& OutPrecedingEvent, double& OutPrecedingEventTime, const EventType*& OutFollowingEvent, double& OutFollowingEventTime) const = 0;
 	
 	/**
 	 * Finds event information for the event closest to InTime from the interval [InTime - DeltaTime, InTime + DeltaTime]
