Animation blueprint debugging support in rewind debugger

diff --git a/Engine/Plugins/Animation/GameplayInsights/Source/RewindDebugger/Private/AnimInstanceHelpers.cpp b/Engine/Plugins/Animation/GameplayInsights/Source/RewindDebugger/Private/AnimInstanceHelpers.cpp
new file mode 100644
--- /dev/null
+++ b/Engine/Plugins/Animation/GameplayInsights/Source/RewindDebugger/Private/AnimInstanceHelpers.cpp
@@ -0,0 +1,96 @@
+﻿// Copyright Epic Games, Inc. All Rights Reserved.
+
+#include "AnimInstanceHelpers.h"
+#include "Modules/ModuleManager.h"
+#include "ObjectTrace.h"
+#include "IGameplayProvider.h"
+#include "Animation/AnimBlueprint.h"
+#include "Animation/AnimBlueprintGeneratedClass.h"
+#include "Editor.h"
+#include "EdGraph/EdGraph.h"
+#include "Subsystems/AssetEditorSubsystem.h"
+#include "IAnimationBlueprintEditor.h"
+#include "Insights/IUnrealInsightsModule.h"
+#include "ToolMenus.h"
+
+#define LOCTEXT_NAMESPACE "SRewindDebuggerAnimBPTools"
+
+static bool OpenBlueprintAndAttachDebugger(const TraceServices::IAnalysisSession* Session, uint64 ObjectId)
+{
+	TraceServices::FAnalysisSessionReadScope SessionReadScope(*Session);
+
+	if (const IGameplayProvider* GameplayProvider = Session->ReadProvider<IGameplayProvider>("GameplayProvider"))
+	{
+		if (const FObjectInfo* ObjectInfo = GameplayProvider->FindObjectInfo(ObjectId))
+		{
+			if (const FClassInfo* ClassInfo = GameplayProvider->FindClassInfo(ObjectInfo->ClassId))
+			{
+				if (UAnimBlueprintGeneratedClass* InstanceClass = TSoftObjectPtr<UAnimBlueprintGeneratedClass>(FSoftObjectPath(ClassInfo->PathName)).LoadSynchronous())
+				{
+					if(UAnimBlueprint* AnimBlueprint = Cast<UAnimBlueprint>(InstanceClass->ClassGeneratedBy))
+					{
+						GEditor->GetEditorSubsystem<UAssetEditorSubsystem>()->OpenEditorForAsset(AnimBlueprint);
+
+						if (UObject* SelectedInstance = FObjectTrace::GetObjectFromId(ObjectId))
+						{
+							AnimBlueprint->SetObjectBeingDebugged(SelectedInstance);
+						}
+
+						if(IAnimationBlueprintEditor* AnimBlueprintEditor = static_cast<IAnimationBlueprintEditor*>(GEditor->GetEditorSubsystem<UAssetEditorSubsystem>()->FindEditorForAsset(AnimBlueprint, true)))
+						{
+							// navigate the opened editor to the AnimGraph
+							if (AnimBlueprint->FunctionGraphs.Num() > 0)
+							{
+								AnimBlueprintEditor->JumpToHyperlink(AnimBlueprint->FunctionGraphs[0]);
+							}
+						}
+
+						return true;
+					}
+				}
+			}
+		}
+	}
+	return false;
+}
+
+bool FAnimInstanceDoubleClickHandler::HandleDoubleClick(IRewindDebugger* RewindDebugger)
+{
+	TSharedPtr<FDebugObjectInfo> SelectedObject = RewindDebugger->GetSelectedComponent();
+	if (SelectedObject.IsValid())
+	{
+		return OpenBlueprintAndAttachDebugger(RewindDebugger->GetAnalysisSession(), SelectedObject->ObjectId);
+	}
+	return false;
+}
+
+FName FAnimInstanceDoubleClickHandler::GetTargetTypeName() const
+{
+	static const FName ObjectTypeName = "AnimInstance";
+	return ObjectTypeName;
+}
+
+void FAnimInstanceMenu::Register()
+{
+	UToolMenu* Menu = UToolMenus::Get()->RegisterMenu("RewindDebugger.ComponentContextMenu");
+	FToolMenuSection& Section = Menu->FindOrAddSection("Blueprint");
+	FToolMenuEntry& Entry = Section.AddDynamicEntry("DebugAnimInstanceEntry", FNewToolMenuSectionDelegate::CreateLambda([](FToolMenuSection& InSection)
+	{
+		UComponentContextMenuContext* Context = InSection.FindContext<UComponentContextMenuContext>();
+		if (Context && Context->SelectedObject.IsValid() && Context->TypeHierarchy.Contains("AnimInstance"))
+		{
+			InSection.AddMenuEntry(NAME_None,
+						LOCTEXT("Open AnimBP", "Open/Debug AnimGraph"),
+							LOCTEXT("Open AnimBP ToolTip", "Open this Animation Blueprint and attach the debugger to this instance"),
+							FSlateIcon(),
+							FExecuteAction::CreateLambda([ObjectId = Context->SelectedObject->ObjectId]()
+							{
+								IUnrealInsightsModule& UnrealInsightsModule = FModuleManager::LoadModuleChecked<IUnrealInsightsModule>("TraceInsights");
+								TSharedPtr<const TraceServices::IAnalysisSession> Session = UnrealInsightsModule.GetAnalysisSession();
+								OpenBlueprintAndAttachDebugger(Session.Get(), ObjectId);
+							}));
+		}
+	}));
+}
+
+#undef LOCTEXT_NAMESPACE
diff --git a/Engine/Plugins/Animation/GameplayInsights/Source/RewindDebugger/Private/AnimInstanceHelpers.h b/Engine/Plugins/Animation/GameplayInsights/Source/RewindDebugger/Private/AnimInstanceHelpers.h
new file mode 100644
--- /dev/null
+++ b/Engine/Plugins/Animation/GameplayInsights/Source/RewindDebugger/Private/AnimInstanceHelpers.h
@@ -0,0 +1,16 @@
+﻿// Copyright Epic Games, Inc. All Rights Reserved.
+#pragma once
+
+#include "IRewindDebuggerDoubleClickHandler.h"
+
+class FAnimInstanceDoubleClickHandler : public IRewindDebuggerDoubleClickHandler
+{
+	virtual FName GetTargetTypeName() const override;
+	bool HandleDoubleClick(IRewindDebugger* RewindDebugger) override;
+};
+
+class FAnimInstanceMenu
+{
+public:
+	static void Register();
+};
\ No newline at end of file
diff --git a/Engine/Plugins/Animation/GameplayInsights/Source/RewindDebugger/Private/PropertyTraceMenu.cpp b/Engine/Plugins/Animation/GameplayInsights/Source/RewindDebugger/Private/PropertyTraceMenu.cpp
new file mode 100644
--- /dev/null
+++ b/Engine/Plugins/Animation/GameplayInsights/Source/RewindDebugger/Private/PropertyTraceMenu.cpp
@@ -0,0 +1,60 @@
+﻿// Copyright Epic Games, Inc. All Rights Reserved.
+
+#include "PropertyTraceMenu.h"
+#include "Modules/ModuleManager.h"
+#include "IGameplayInsightsModule.h"
+#include "ObjectTrace.h"
+#include "ToolMenus.h"
+
+#define LOCTEXT_NAMESPACE "RewindDebuggerComponentHelpers"
+
+void FPropertyTraceMenu::Register()
+{
+	UToolMenu* Menu = UToolMenus::Get()->RegisterMenu("RewindDebugger.ComponentContextMenu");
+	FToolMenuSection& Section = Menu->FindOrAddSection("Object");
+	FToolMenuEntry& Entry = Section.AddDynamicEntry("AssetManagerEditorViewCommands", FNewToolMenuSectionDelegate::CreateLambda([](FToolMenuSection& InSection)
+	{
+		UComponentContextMenuContext* Context = InSection.FindContext<UComponentContextMenuContext>();
+		if (Context && Context->SelectedObject.IsValid())
+		{
+			InSection.AddMenuEntry(
+				NAME_None,
+				LOCTEXT("Trace Object Properties", "Trace Object Properties"),
+				LOCTEXT("Trace Object Properties Tooltip", "Record this object's properties so they will show in the Rewind Debugger when scrubbing."),
+				FSlateIcon(), 
+				FUIAction(FExecuteAction::CreateLambda([SelectedObject=Context->SelectedObject]()
+					{
+						IGameplayInsightsModule& GameplayInsightsModule = FModuleManager::GetModuleChecked<IGameplayInsightsModule>("GameplayInsights");
+						if (SelectedObject.IsValid())
+						{
+							if (UObject* Object = FObjectTrace::GetObjectFromId(SelectedObject->ObjectId))
+							{
+								GameplayInsightsModule.EnableObjectPropertyTrace(Object, !GameplayInsightsModule.IsObjectPropertyTraceEnabled(Object));
+							}
+						}
+					}),
+					FCanExecuteAction(),
+					FGetActionCheckState::CreateLambda([SelectedObject=Context->SelectedObject]()
+					{
+						IGameplayInsightsModule& GameplayInsightsModule = FModuleManager::GetModuleChecked<IGameplayInsightsModule>("GameplayInsights");
+						bool bEnabled = false;
+	 
+						if (SelectedObject.IsValid())
+						{
+							if (UObject* Object = FObjectTrace::GetObjectFromId(SelectedObject->ObjectId))
+							{
+								bEnabled = GameplayInsightsModule.IsObjectPropertyTraceEnabled(Object);
+							}
+						}
+						 return bEnabled ? ECheckBoxState::Checked : ECheckBoxState::Unchecked;
+					}
+					),
+					FIsActionButtonVisible()
+				),
+				EUserInterfaceActionType::Check
+			);
+		}
+	}));
+}
+
+#undef LOCTEXT_NAMESPACE
diff --git a/Engine/Plugins/Animation/GameplayInsights/Source/RewindDebugger/Private/PropertyTraceMenu.h b/Engine/Plugins/Animation/GameplayInsights/Source/RewindDebugger/Private/PropertyTraceMenu.h
new file mode 100644
--- /dev/null
+++ b/Engine/Plugins/Animation/GameplayInsights/Source/RewindDebugger/Private/PropertyTraceMenu.h
@@ -0,0 +1,10 @@
+﻿// Copyright Epic Games, Inc. All Rights Reserved.
+#pragma once
+
+#include "IRewindDebuggerDoubleClickHandler.h"
+
+class FPropertyTraceMenu
+{
+public:
+	static void Register();
+};
\ No newline at end of file
diff --git a/Engine/Plugins/Animation/GameplayInsights/Source/RewindDebugger/Private/RewindDebugger.cpp b/Engine/Plugins/Animation/GameplayInsights/Source/RewindDebugger/Private/RewindDebugger.cpp
--- a/Engine/Plugins/Animation/GameplayInsights/Source/RewindDebugger/Private/RewindDebugger.cpp
+++ b/Engine/Plugins/Animation/GameplayInsights/Source/RewindDebugger/Private/RewindDebugger.cpp
@@ -8,12 +8,13 @@
 #include "Insights/IUnrealInsightsModule.h"
 #include "Modules/ModuleManager.h"
 #include "ObjectTrace.h"
-#include "TraceServices/ITraceServicesModule.h"
 #include "TraceServices/Model/Frames.h"
-#include "LevelEditor.h"
 #include "SLevelViewport.h"
 #include "IRewindDebuggerExtension.h"
-
+#include "IRewindDebuggerDoubleClickHandler.h"
+#include "Animation/AnimBlueprint.h"
+#include "Animation/AnimBlueprintGeneratedClass.h"
+#include "ToolMenus.h"
 
 FRewindDebugger* FRewindDebugger::InternalInstance = nullptr;
 
@@ -78,7 +79,6 @@ FRewindDebugger::~FRewindDebugger()
 	FEditorDelegates::SingleStepPIE.RemoveAll(this);
 
 	FTSTicker::GetCoreTicker().RemoveTicker(TickerHandle);
-
 }
 
 void FRewindDebugger::Initialize() 
@@ -607,8 +607,11 @@ void FRewindDebugger::Tick(float DeltaTime)
 			// if you select a debug target before you start recording, update component list when it becomes valid
 			RefreshDebugComponents();
 		}
+		
+		const IAnimationProvider* AnimationProvider = Session->ReadProvider<IAnimationProvider>("AnimationProvider");
+		const IGameplayProvider* GameplayProvider = Session->ReadProvider<IGameplayProvider>("GameplayProvider");
 
-		if (const IAnimationProvider* AnimationProvider = Session->ReadProvider<IAnimationProvider>("AnimationProvider"))
+		if (AnimationProvider && GameplayProvider)
 		{
 			TraceServices::FAnalysisSessionReadScope SessionReadScope(*Session);
 			UWorld* World = GetWorldToVisualize();
@@ -702,6 +705,125 @@ void FRewindDebugger::Tick(float DeltaTime)
 									// calling this here, even on meshes which have no recorded data makes meshes with a MasterPoseComponent attach properly during replay
 									// longer term solution: we should be recording those meshes
 									MeshComponent->ApplyEditedComponentSpaceTransforms();
+
+									// update debug info for attached Animation Blueprint editors
+									if (UAnimBlueprintGeneratedClass* InstanceClass = Cast<UAnimBlueprintGeneratedClass>(MeshComponent->GetAnimClass()))
+									{
+										if(UAnimBlueprint* AnimBlueprint = Cast<UAnimBlueprint>(InstanceClass->ClassGeneratedBy))
+										{
+											UAnimInstance* AnimInstance = MeshComponent->GetAnimInstance();
+											if(AnimInstance && AnimBlueprint->IsObjectBeingDebugged(AnimInstance))
+											{
+												uint64 Id = FObjectTrace::GetObjectId(AnimInstance);
+												const int32 NodeCount = InstanceClass->GetAnimNodeProperties().Num();
+							
+												AnimationProvider->ReadAnimGraphTimeline(Id, [this, Id, InstanceClass, AnimationProvider, GameplayProvider, Frame, NodeCount](const IAnimationProvider::AnimGraphTimeline& InGraphTimeline)
+												{
+													FAnimBlueprintDebugData& DebugData = InstanceClass->GetAnimBlueprintDebugData();
+													DebugData.ResetNodeVisitSites();
+							
+													InGraphTimeline.EnumerateEvents(Frame.StartTime, Frame.EndTime, [this, Id, AnimationProvider, GameplayProvider, &DebugData, NodeCount](double InGraphStartTime, double InGraphEndTime, uint32 InDepth, const FAnimGraphMessage& InMessage)
+													{
+														// Basic verification - check node count is the same
+														// @TODO: could add some form of node hash/CRC to the class to improve this
+														if(InMessage.NodeCount == NodeCount)
+														{
+															// Check for an update phase (which contains weights)
+															if(InMessage.Phase == EAnimGraphPhase::Update)
+															{
+																AnimationProvider->ReadAnimNodesTimeline(Id, [InGraphStartTime, InGraphEndTime, &DebugData](const IAnimationProvider::AnimNodesTimeline& InNodesTimeline)
+																{
+																	InNodesTimeline.EnumerateEvents(InGraphStartTime, InGraphEndTime, [&DebugData](double InStartTime, double InEndTime, uint32 InDepth, const FAnimNodeMessage& InMessage)
+																	{
+																		DebugData.RecordNodeVisit(InMessage.NodeId, InMessage.PreviousNodeId, InMessage.Weight);
+																		return TraceServices::EEventEnumerate::Continue;
+																	});
+																});
+							
+																AnimationProvider->ReadStateMachinesTimeline(Id, [InGraphStartTime, InGraphEndTime, &DebugData](const IAnimationProvider::StateMachinesTimeline& InStateMachinesTimeline)
+																{
+																	InStateMachinesTimeline.EnumerateEvents(InGraphStartTime, InGraphEndTime, [&DebugData](double InStartTime, double InEndTime, uint32 InDepth, const FAnimStateMachineMessage& InMessage)
+																	{
+																		DebugData.RecordStateData(InMessage.StateMachineIndex, InMessage.StateIndex, InMessage.StateWeight, InMessage.ElapsedTime);
+																		return TraceServices::EEventEnumerate::Continue;
+																	});
+																});
+							
+																AnimationProvider->ReadAnimSequencePlayersTimeline(Id, [InGraphStartTime, InGraphEndTime, GameplayProvider, &DebugData](const IAnimationProvider::AnimSequencePlayersTimeline& InSequencePlayersTimeline)
+																{
+																	InSequencePlayersTimeline.EnumerateEvents(InGraphStartTime, InGraphEndTime, [&DebugData](double InStartTime, double InEndTime, uint32 InDepth, const FAnimSequencePlayerMessage& InMessage)
+																	{
+																		DebugData.RecordSequencePlayer(InMessage.NodeId, InMessage.Position, InMessage.Length, InMessage.FrameCounter);
+																		return TraceServices::EEventEnumerate::Continue;
+																	});
+																});
+							
+																AnimationProvider->ReadAnimBlendSpacePlayersTimeline(Id, [InGraphStartTime, InGraphEndTime, GameplayProvider, &DebugData](const IAnimationProvider::BlendSpacePlayersTimeline& InBlendSpacePlayersTimeline)
+																{
+																	InBlendSpacePlayersTimeline.EnumerateEvents(InGraphStartTime, InGraphEndTime, [GameplayProvider, &DebugData](double InStartTime, double InEndTime, uint32 InDepth, const FBlendSpacePlayerMessage& InMessage)
+																	{
+																		UBlendSpace* BlendSpace = nullptr;
+																		const FObjectInfo* BlendSpaceInfo = GameplayProvider->FindObjectInfo(InMessage.BlendSpaceId);
+																		if(BlendSpaceInfo)
+																		{
+																			BlendSpace = TSoftObjectPtr<UBlendSpace>(FSoftObjectPath(BlendSpaceInfo->PathName)).LoadSynchronous();
+																		}
+							
+																		DebugData.RecordBlendSpacePlayer(InMessage.NodeId, BlendSpace, FVector(InMessage.PositionX, InMessage.PositionY, InMessage.PositionZ), FVector(InMessage.FilteredPositionX, InMessage.FilteredPositionY, InMessage.FilteredPositionZ));
+																		return TraceServices::EEventEnumerate::Continue;
+																	});
+																});
+							
+																AnimationProvider->ReadAnimSyncTimeline(Id, [InGraphStartTime, InGraphEndTime, AnimationProvider, &DebugData](const IAnimationProvider::AnimSyncTimeline& InAnimSyncTimeline)
+																{
+																	InAnimSyncTimeline.EnumerateEvents(InGraphStartTime, InGraphEndTime, [AnimationProvider, &DebugData](double InStartTime, double InEndTime, uint32 InDepth, const FAnimSyncMessage& InMessage)
+																	{
+																		const TCHAR* GroupName = AnimationProvider->GetName(InMessage.GroupNameId);
+																		if(GroupName)
+																		{
+																			DebugData.RecordNodeSync(InMessage.SourceNodeId, FName(GroupName));
+																		}
+																
+																		return TraceServices::EEventEnumerate::Continue;
+																	});
+																});
+															}
+							
+															// Some traces come from both update and evaluate phases
+															if(InMessage.Phase == EAnimGraphPhase::Update || InMessage.Phase == EAnimGraphPhase::Evaluate)
+															{
+																AnimationProvider->ReadAnimAttributesTimeline(Id, [InGraphStartTime, InGraphEndTime, AnimationProvider, &DebugData](const IAnimationProvider::AnimAttributeTimeline& InAnimAttributeTimeline)
+																{
+																	InAnimAttributeTimeline.EnumerateEvents(InGraphStartTime, InGraphEndTime, [AnimationProvider, &DebugData](double InStartTime, double InEndTime, uint32 InDepth, const FAnimAttributeMessage& InMessage)
+																	{
+																		const TCHAR* AttributeName = AnimationProvider->GetName(InMessage.AttributeNameId);
+																		if(AttributeName)
+																		{
+																			DebugData.RecordNodeAttribute(InMessage.TargetNodeId, InMessage.SourceNodeId, FName(AttributeName));
+																		}
+																
+																		return TraceServices::EEventEnumerate::Continue;
+																	});
+																});
+															}
+							
+															// Anim node values can come from all phases
+															AnimationProvider->ReadAnimNodeValuesTimeline(Id, [InGraphStartTime, InGraphEndTime, AnimationProvider, &DebugData](const IAnimationProvider::AnimNodeValuesTimeline& InNodeValuesTimeline)
+															{
+																InNodeValuesTimeline.EnumerateEvents(InGraphStartTime, InGraphEndTime, [AnimationProvider, &DebugData](double InStartTime, double InEndTime, uint32 InDepth, const FAnimNodeValueMessage& InMessage)
+																{
+																	FText Text = AnimationProvider->FormatNodeKeyValue(InMessage);
+																	DebugData.RecordNodeValue(InMessage.NodeId, Text.ToString());
+																	return TraceServices::EEventEnumerate::Continue;
+																});
+															});
+														}
+														return TraceServices::EEventEnumerate::Continue;
+													});
+												});
+											}
+										}	
+									}
 								}
 							}
 						}
@@ -718,3 +840,81 @@ void FRewindDebugger::Tick(float DeltaTime)
 		);
 	}
 }
+
+void FRewindDebugger::ComponentSelectionChanged(TSharedPtr<FDebugObjectInfo> SelectedObject)
+{
+	SelectedComponent = SelectedObject;
+}
+
+void FRewindDebugger::ComponentDoubleClicked(TSharedPtr<FDebugObjectInfo> SelectedObject)
+{
+	if (!SelectedObject.IsValid())
+	{
+		return;
+	}
+	
+	SelectedComponent = SelectedObject;
+	
+	IModularFeatures& ModularFeatures = IModularFeatures::Get();
+	static const FName HandlerFeatureName = IRewindDebuggerDoubleClickHandler::ModularFeatureName;
+		
+	if (const TraceServices::IAnalysisSession* Session = GetAnalysisSession())
+	{
+		TraceServices::FAnalysisSessionReadScope SessionReadScope(*Session);
+	
+		const IGameplayProvider* GameplayProvider = Session->ReadProvider<IGameplayProvider>("GameplayProvider");
+		const FObjectInfo& ObjectInfo = GameplayProvider->GetObjectInfo(SelectedObject->ObjectId);
+		uint64 ClassId = ObjectInfo.ClassId;
+		bool bHandled = false;
+		
+		const int32 NumExtensions = ModularFeatures.GetModularFeatureImplementationCount(HandlerFeatureName);
+
+		// iterate up the class hierarchy, looking for a registered double click handler, until we find the one that succeeeds that is most specific to the type of this object
+		while (ClassId != 0 && !bHandled)
+		{
+			const FClassInfo& ClassInfo = GameplayProvider->GetClassInfo(ClassId);
+		
+			for (int32 ExtensionIndex = 0; ExtensionIndex < NumExtensions; ++ExtensionIndex)
+			{
+				IRewindDebuggerDoubleClickHandler* Handler = static_cast<IRewindDebuggerDoubleClickHandler*>(ModularFeatures.GetModularFeatureImplementation(HandlerFeatureName, ExtensionIndex));
+				if (Handler->GetTargetTypeName() == ClassInfo.Name)
+				{
+					if (Handler->HandleDoubleClick(this))
+					{
+						bHandled = true;
+						break;
+					}
+				}
+			}
+		
+			ClassId = ClassInfo.SuperId;
+		}
+	}
+}
+
+TSharedPtr<SWidget> FRewindDebugger::BuildComponentContextMenu()
+{
+	UComponentContextMenuContext* MenuContext = NewObject<UComponentContextMenuContext>();
+	MenuContext->SelectedObject = SelectedComponent;
+
+	if (SelectedComponent.IsValid())
+	{
+		// build a list of class hierarchy names to make it easier for extensions to enable menu entries by type
+		if (const TraceServices::IAnalysisSession* Session = GetAnalysisSession())
+		{
+			TraceServices::FAnalysisSessionReadScope SessionReadScope(*Session);
+	
+			const IGameplayProvider* GameplayProvider = Session->ReadProvider<IGameplayProvider>("GameplayProvider");
+			const FObjectInfo& ObjectInfo = GameplayProvider->GetObjectInfo(SelectedComponent->ObjectId);
+			uint64 ClassId = ObjectInfo.ClassId;
+			while (ClassId != 0)
+			{
+				const FClassInfo& ClassInfo = GameplayProvider->GetClassInfo(ClassId);
+				MenuContext->TypeHierarchy.Add(ClassInfo.Name);
+				ClassId = ClassInfo.SuperId;
+			}
+		}
+	}
+
+	return UToolMenus::Get()->GenerateWidget("RewindDebugger.ComponentContextMenu", FToolMenuContext(MenuContext));
+ }
diff --git a/Engine/Plugins/Animation/GameplayInsights/Source/RewindDebugger/Private/RewindDebugger.h b/Engine/Plugins/Animation/GameplayInsights/Source/RewindDebugger/Private/RewindDebugger.h
--- a/Engine/Plugins/Animation/GameplayInsights/Source/RewindDebugger/Private/RewindDebugger.h
+++ b/Engine/Plugins/Animation/GameplayInsights/Source/RewindDebugger/Private/RewindDebugger.h
@@ -7,6 +7,7 @@
 #include "Containers/Ticker.h"
 #include "UObject/WeakObjectPtr.h"
 
+class FMenuBuilder;
 class USkeletalMeshComponent;
 
 namespace TraceServices
@@ -14,6 +15,8 @@ namespace TraceServices
 	class IAnalysisSession;
 }
 
+class SWidget;
+
 // Singleton class that handles the logic for the Rewind Debugger
 // handles:
 //  Playback/Scrubbing state
@@ -35,6 +38,7 @@ public:
 	virtual bool IsRecording() const override { return bRecording; }
 	virtual bool IsPIESimulating() const override { return bPIESimulating; }
 	virtual double GetRecordingDuration() const override { return RecordingDuration.Get(); }
+	virtual TSharedPtr<FDebugObjectInfo> GetSelectedComponent() const override { return SelectedComponent; }
 
 	// create singleton instance
 	static void Initialize();
@@ -89,6 +93,10 @@ public:
 
 	DECLARE_DELEGATE(FOnComponentListChanged)
 	void OnComponentListChanged(const FOnComponentListChanged& ComponentListChangedCallback);
+	
+	void ComponentDoubleClicked(TSharedPtr<FDebugObjectInfo> SelectedObject);
+	void ComponentSelectionChanged(TSharedPtr<FDebugObjectInfo> SelectedObject);
+	TSharedPtr<SWidget> BuildComponentContextMenu();
 
 	DECLARE_DELEGATE_OneParam( FOnTrackCursor, bool)
 	void OnTrackCursor(const FOnTrackCursor& TrackCursorCallback);
@@ -139,6 +147,7 @@ private:
 	static FRewindDebugger* InternalInstance;
 
 	TArray<TSharedPtr<FDebugObjectInfo>> DebugComponents;
+	TSharedPtr<FDebugObjectInfo> SelectedComponent;
 
 	struct FMeshComponentResetData
 	{
diff --git a/Engine/Plugins/Animation/GameplayInsights/Source/RewindDebugger/Private/RewindDebuggerModule.cpp b/Engine/Plugins/Animation/GameplayInsights/Source/RewindDebugger/Private/RewindDebuggerModule.cpp
--- a/Engine/Plugins/Animation/GameplayInsights/Source/RewindDebugger/Private/RewindDebuggerModule.cpp
+++ b/Engine/Plugins/Animation/GameplayInsights/Source/RewindDebugger/Private/RewindDebuggerModule.cpp
@@ -1,24 +1,22 @@
 // Copyright Epic Games, Inc. All Rights Reserved.
 
 #include "RewindDebuggerModule.h"
-#include "Containers/Ticker.h"
-#include "Editor.h"
-#include "Engine/Engine.h"
+
+
 #include "Engine/Selection.h"
 #include "Features/IModularFeatures.h"
 #include "Framework/Docking/LayoutExtender.h"
 #include "IAnimationBlueprintEditorModule.h"
-#include "LevelEditorMenuContext.h"
 #include "Modules/ModuleManager.h"
 #include "RewindDebugger.h"
 #include "RewindDebuggerStyle.h"
 #include "RewindDebuggerCommands.h"
 #include "SRewindDebugger.h"
-#include "SSCSEditorMenuContext.h"
-#include "ToolMenus.h"
 #include "Widgets/Docking/SDockTab.h"
 #include "WorkspaceMenuStructure.h"
 #include "WorkspaceMenuStructureModule.h"
+#include "AnimInstanceHelpers.h"
+#include "PropertyTraceMenu.h"
 
 #define LOCTEXT_NAMESPACE "RewindDebuggerModule"
 
@@ -102,7 +100,10 @@ TSharedRef<SDockTab> FRewindDebuggerModule::SpawnRewindDebuggerTab(const FSpawnT
 								.RecordingDuration(DebuggerInstance->GetRecordingDurationProperty())
 								.DebugComponents(&DebuggerInstance->GetDebugComponents())
 								.TraceTime(DebuggerInstance->GetTraceTimeProperty())
-								.OnScrubPositionChanged(SRewindDebugger::FOnScrubPositionChanged::CreateRaw(DebuggerInstance,&FRewindDebugger::ScrubToTime))
+								.OnScrubPositionChanged_Raw(DebuggerInstance,&FRewindDebugger::ScrubToTime)
+								.OnComponentDoubleClicked_Raw(DebuggerInstance, &FRewindDebugger::ComponentDoubleClicked)
+								.OnComponentSelectionChanged_Raw(DebuggerInstance, &FRewindDebugger::ComponentSelectionChanged)
+								.BuildComponentContextMenu_Raw(DebuggerInstance, &FRewindDebugger::BuildComponentContextMenu)
 								.ScrubTime_Lambda([]() { return FRewindDebugger::Instance()->GetScrubTime(); });
 
 	DebuggerInstance->OnTrackCursor(FRewindDebugger::FOnTrackCursor::CreateSP(RewindDebuggerWidget.Get(), &SRewindDebugger::TrackCursor));
@@ -113,6 +114,8 @@ TSharedRef<SDockTab> FRewindDebuggerModule::SpawnRewindDebuggerTab(const FSpawnT
 	return MajorTab;
 }
 
+static FAnimInstanceDoubleClickHandler AnimInstanceDoubleClickHandler;
+
 void FRewindDebuggerModule::StartupModule()
 {
 	FRewindDebuggerStyle::Initialize();
@@ -127,11 +130,16 @@ void FRewindDebuggerModule::StartupModule()
 
 
 	IModularFeatures::Get().RegisterModularFeature(IRewindDebuggerExtension::ModularFeatureName, &RewindDebuggerCameraExtension);
+	IModularFeatures::Get().RegisterModularFeature(IRewindDebuggerDoubleClickHandler::ModularFeatureName, &AnimInstanceDoubleClickHandler);
+
+	FPropertyTraceMenu::Register();
+	FAnimInstanceMenu::Register();
 }
 
 void FRewindDebuggerModule::ShutdownModule()
 {
 	IModularFeatures::Get().UnregisterModularFeature(IRewindDebuggerExtension::ModularFeatureName, &RewindDebuggerCameraExtension);
+	IModularFeatures::Get().UnregisterModularFeature(IRewindDebuggerDoubleClickHandler::ModularFeatureName, &AnimInstanceDoubleClickHandler);
 
 	FRewindDebuggerCommands::Unregister();
 	FRewindDebuggerStyle::Shutdown();
diff --git a/Engine/Plugins/Animation/GameplayInsights/Source/RewindDebugger/Private/SRewindDebugger.cpp b/Engine/Plugins/Animation/GameplayInsights/Source/RewindDebugger/Private/SRewindDebugger.cpp
--- a/Engine/Plugins/Animation/GameplayInsights/Source/RewindDebugger/Private/SRewindDebugger.cpp
+++ b/Engine/Plugins/Animation/GameplayInsights/Source/RewindDebugger/Private/SRewindDebugger.cpp
@@ -2,16 +2,12 @@
 
 #include "SRewindDebugger.h"
 #include "ActorPickerMode.h"
-#include "Async/Future.h"
-#include "Components/SkeletalMeshComponent.h"
 #include "Editor.h"
 #include "Editor/EditorEngine.h"
 #include "Framework/Application/SlateApplication.h"
 #include "Framework/Docking/LayoutService.h"
 #include "Framework/MultiBox/MultiBoxBuilder.h"
 #include "Insights/IUnrealInsightsModule.h"
-#include "IGameplayInsightsModule.h"
-#include "ISettingsModule.h"
 #include "Modules/ModuleManager.h"
 #include "ObjectTrace.h"
 #include "RewindDebuggerStyle.h"
@@ -27,11 +23,11 @@
 #include "Widgets/Images/SImage.h"
 #include "Widgets/Input/SComboButton.h"
 #include "Widgets/Input/SButton.h"
-#include "Widgets/Input/SSearchBox.h"
 #include "Widgets/Layout/SExpandableArea.h"
 #include "Widgets/Layout/SScrollBox.h"
 #include "IRewindDebuggerViewCreator.h"
 #include "RewindDebuggerViewCreators.h"
+#include "IRewindDebuggerDoubleClickHandler.h"
 
 #define LOCTEXT_NAMESPACE "SRewindDebugger"
 
@@ -190,6 +186,8 @@ TSharedRef<SWidget> SRewindDebugger::MakeSelectActorMenu()
 void SRewindDebugger::Construct(const FArguments& InArgs, TSharedRef<FUICommandList> CommandList, const TSharedRef<SDockTab>& ConstructUnderMajorTab, const TSharedPtr<SWindow>& ConstructUnderWindow)
 {
 	OnScrubPositionChanged = InArgs._OnScrubPositionChanged;
+	OnComponentSelectionChanged = InArgs._OnComponentSelectionChanged;
+	BuildComponentContextMenu = InArgs._BuildComponentContextMenu;
 	ScrubTimeAttribute = InArgs._ScrubTime;
 	DebugComponents = InArgs._DebugComponents;
 	TraceTime.Initialize(InArgs._TraceTime);
@@ -239,7 +237,9 @@ void SRewindDebugger::Construct(const FArguments& InArgs, TSharedRef<FUICommandL
 
 	ComponentTreeView =	SNew(SRewindDebuggerComponentTree)
 				.DebugComponents(InArgs._DebugComponents)
-				.OnSelectionChanged(this, &SRewindDebugger::OnComponentSelectionChanged);
+				.OnMouseButtonDoubleClick(InArgs._OnComponentDoubleClicked)
+				.OnContextMenuOpening(this, &SRewindDebugger::OnContextMenuOpening)
+				.OnSelectionChanged(this, &SRewindDebugger::ComponentSelectionChanged);
 
 
 	TraceTime.OnPropertyChanged = TraceTime.OnPropertyChanged.CreateRaw(this, &SRewindDebugger::TraceTimeChanged);
@@ -263,7 +263,6 @@ void SRewindDebugger::Construct(const FArguments& InArgs, TSharedRef<FUICommandL
 
 	TSharedPtr<FTabManager::FStack> MainTabStack = FTabManager::NewStack();
 
-	IGameplayInsightsModule& GameplayInsightsModule = FModuleManager::LoadModuleChecked<IGameplayInsightsModule>("GameplayInsights");
 	FRewindDebuggerViewCreators::EnumerateCreators([this, MainTabStack](const IRewindDebuggerViewCreator* Creator)
 		{
 			FName TabName = Creator->GetName();
@@ -499,7 +498,6 @@ void SRewindDebugger::OnPinnedTabClosed(TSharedRef<SDockTab> Tab)
 	FSlateIcon TabIcon;
 	FText TabLabel;
 
-	IGameplayInsightsModule& GameplayInsightsModule = FModuleManager::LoadModuleChecked<IGameplayInsightsModule>("GameplayInsights");
 	if (const IRewindDebuggerViewCreator* Creator = FRewindDebuggerViewCreators::GetCreator(TabName))
 	{
 		TabIcon = Creator->GetIcon();
@@ -570,7 +568,6 @@ void SRewindDebugger::CreateDebugViews()
 	if (SelectedComponent.IsValid())
 	{
 		IUnrealInsightsModule &UnrealInsightsModule = FModuleManager::LoadModuleChecked<IUnrealInsightsModule>("TraceInsights");
-		IGameplayInsightsModule& GameplayInsightsModule = FModuleManager::LoadModuleChecked<IGameplayInsightsModule>("GameplayInsights");
 		TSharedPtr<const TraceServices::IAnalysisSession> Session = UnrealInsightsModule.GetAnalysisSession();
 
 		FRewindDebuggerViewCreators::CreateDebugViews(SelectedComponent->ObjectId, TraceTime.Get(), *Session, DebugViews);
@@ -611,14 +608,21 @@ void SRewindDebugger::CreateDebugTabs()
 	}
 }
 
-void SRewindDebugger::OnComponentSelectionChanged(TSharedPtr<FDebugObjectInfo> SelectedItem, ESelectInfo::Type SelectInfo)
+void SRewindDebugger::ComponentSelectionChanged(TSharedPtr<FDebugObjectInfo> SelectedItem, ESelectInfo::Type SelectInfo)
 {
 	SelectedComponent = SelectedItem;
 
+	OnComponentSelectionChanged.ExecuteIfBound(SelectedItem);
+
 	CreateDebugViews();
 	CreateDebugTabs();
 }
 
+TSharedPtr<SWidget> SRewindDebugger::OnContextMenuOpening()
+{
+	return BuildComponentContextMenu.Execute();
+}
+
 FReply SRewindDebugger::OnSelectActorClicked()
 {
 		FActorPickerModeModule& ActorPickerMode = FModuleManager::Get().GetModuleChecked<FActorPickerModeModule>("ActorPickerMode");
diff --git a/Engine/Plugins/Animation/GameplayInsights/Source/RewindDebugger/Private/SRewindDebugger.h b/Engine/Plugins/Animation/GameplayInsights/Source/RewindDebugger/Private/SRewindDebugger.h
--- a/Engine/Plugins/Animation/GameplayInsights/Source/RewindDebugger/Private/SRewindDebugger.h
+++ b/Engine/Plugins/Animation/GameplayInsights/Source/RewindDebugger/Private/SRewindDebugger.h
@@ -22,6 +22,9 @@ class SRewindDebugger : public SCompoundWidget
 public:
 	DECLARE_DELEGATE_TwoParams( FOnScrubPositionChanged, float, bool )
 	DECLARE_DELEGATE_OneParam( FOnDebugTargetChanged, TSharedPtr<FString> )
+	DECLARE_DELEGATE_OneParam( FOnComponentDoubleClicked, TSharedPtr<FDebugObjectInfo> )
+	DECLARE_DELEGATE_OneParam( FOnComponentSelectionChanged, TSharedPtr<FDebugObjectInfo> )
+	DECLARE_DELEGATE_RetVal( TSharedPtr<SWidget>, FBuildComponentContextMenu )
 
 	SLATE_BEGIN_ARGS(SRewindDebugger) { }
 		SLATE_ARGUMENT( TArray< TSharedPtr< FDebugObjectInfo > >*, DebugComponents );
@@ -29,7 +32,10 @@ public:
 		SLATE_ARGUMENT(TBindablePropertyInitializer<double>, TraceTime);
 		SLATE_ARGUMENT(TBindablePropertyInitializer<float>, RecordingDuration);
 		SLATE_ATTRIBUTE(float, ScrubTime);
-		SLATE_EVENT(FOnScrubPositionChanged, OnScrubPositionChanged);
+		SLATE_EVENT( FOnScrubPositionChanged, OnScrubPositionChanged);
+		SLATE_EVENT( FBuildComponentContextMenu, BuildComponentContextMenu );
+		SLATE_EVENT( FOnComponentDoubleClicked, OnComponentDoubleClicked );
+		SLATE_EVENT( FOnComponentSelectionChanged, OnComponentSelectionChanged );
 	SLATE_END_ARGS()
 	
 public:
@@ -58,6 +64,8 @@ private:
 	TAttribute<float> ScrubTimeAttribute;
 	TAttribute<bool> TrackScrubbingAttribute;
 	FOnScrubPositionChanged OnScrubPositionChanged;
+	FBuildComponentContextMenu BuildComponentContextMenu;
+	FOnComponentSelectionChanged OnComponentSelectionChanged;
 
 	// debug actor selector
 	TSharedRef<SWidget> MakeSelectActorMenu();
@@ -77,8 +85,10 @@ private:
 	// component tree view
 	TArray<TSharedPtr<FDebugObjectInfo>>* DebugComponents;
 	TSharedPtr<FDebugObjectInfo> SelectedComponent;
-    void OnComponentSelectionChanged(TSharedPtr<FDebugObjectInfo> SelectedItem, ESelectInfo::Type SelectInfo);
+    void ComponentSelectionChanged(TSharedPtr<FDebugObjectInfo> SelectedItem, ESelectInfo::Type SelectInfo);
+	
 	TSharedPtr<SRewindDebuggerComponentTree> ComponentTreeView;
+	TSharedPtr<SWidget> OnContextMenuOpening();
 
 	// Debug View Tabs
 	TSharedRef<SDockTab> SpawnTab(const FSpawnTabArgs& Args, FName ViewName);
diff --git a/Engine/Plugins/Animation/GameplayInsights/Source/RewindDebugger/Private/SRewindDebuggerComponentTree.cpp b/Engine/Plugins/Animation/GameplayInsights/Source/RewindDebugger/Private/SRewindDebuggerComponentTree.cpp
--- a/Engine/Plugins/Animation/GameplayInsights/Source/RewindDebugger/Private/SRewindDebuggerComponentTree.cpp
+++ b/Engine/Plugins/Animation/GameplayInsights/Source/RewindDebugger/Private/SRewindDebuggerComponentTree.cpp
@@ -1,18 +1,9 @@
 // Copyright Epic Games, Inc. All Rights Reserved.
 
 #include "SRewindDebuggerComponentTree.h"
-#include "Async/Future.h"
-#include "Editor.h"
-#include "Editor/EditorEngine.h"
-#include "Framework/Application/SlateApplication.h"
-#include "Framework/MultiBox/MultiBoxBuilder.h"
-#include "IGameplayInsightsModule.h"
-#include "Modules/ModuleManager.h"
 #include "ObjectTrace.h"
 #include "Styling/SlateIconFinder.h"
 #include "Widgets/Images/SImage.h"
-#include "Widgets/SOverlay.h"
-#include "Widgets/Views/SListView.h"
 
 #define LOCTEXT_NAMESPACE "SAnimationInsights"
 
@@ -24,7 +15,6 @@ SRewindDebuggerComponentTree::SRewindDebuggerComponentTree()
 
 SRewindDebuggerComponentTree::~SRewindDebuggerComponentTree() 
 {
-
 }
 
 TSharedRef<ITableRow> ComponentTreeViewGenerateRow(TSharedPtr<FDebugObjectInfo> InItem, const TSharedRef<STableViewBase>& OwnerTable)
@@ -74,55 +64,6 @@ void ComponentTreeViewExpansionChanged(TSharedPtr<FDebugObjectInfo> InItem, bool
 	InItem->bExpanded = bShouldBeExpanded;
 }
 
-TSharedPtr<SWidget> SRewindDebuggerComponentTree::ComponentTreeOnContextMenuOpening()
-{
-	FMenuBuilder MenuBuilder(/*bInShouldCloseWindowAfterMenuSelection=*/true, nullptr, nullptr, /*bCloseSelfOnly=*/true);
-	IGameplayInsightsModule& GameplayInsightsModule = FModuleManager::LoadModuleChecked<IGameplayInsightsModule>("GameplayInsights");
-
-	MenuBuilder.BeginSection(NAME_None, LOCTEXT("Property Tracing", "Property Tracing"));
-	{
-		MenuBuilder.AddMenuEntry(
-			LOCTEXT("Trace Object Properties", "Trace Object Properties"),
-			LOCTEXT("Trace Object Properties Tooltip", "Record this object's properties so they will show in the Rewind Debugger when scrubbing."),
-			FSlateIcon(), 
-			FUIAction(FExecuteAction::CreateLambda([this, &GameplayInsightsModule]()
-				{
-					TArray<TSharedPtr<FDebugObjectInfo>> SelectedObjects = ComponentTreeView->GetSelectedItems();
-
-					if (SelectedObjects.Num() > 0)
-					{
-						if (UObject* Object = FObjectTrace::GetObjectFromId(SelectedObjects[0]->ObjectId))
-						{
-							GameplayInsightsModule.EnableObjectPropertyTrace(Object, !GameplayInsightsModule.IsObjectPropertyTraceEnabled(Object));
-						}
-					}
-				}),
-				FCanExecuteAction(),
-				FGetActionCheckState::CreateLambda([this, &GameplayInsightsModule]()
-				{
-					bool bEnabled = false;
-					TArray<TSharedPtr<FDebugObjectInfo>> SelectedObjects = ComponentTreeView->GetSelectedItems();
-
-					if (SelectedObjects.Num() > 0)
-					{
-						if (UObject* Object = FObjectTrace::GetObjectFromId(SelectedObjects[0]->ObjectId))
-						{
-							bEnabled = GameplayInsightsModule.IsObjectPropertyTraceEnabled(Object);
-						}
-					}
-                    return bEnabled ? ECheckBoxState::Checked : ECheckBoxState::Unchecked;
-				}
-				),
-				FIsActionButtonVisible()
-			),
-			NAME_None,
-			EUserInterfaceActionType::Check
-		);
-	}
-	MenuBuilder.EndSection();
-
-	return MenuBuilder.MakeWidget();
-}
 
 
 void SRewindDebuggerComponentTree::Construct(const FArguments& InArgs)
@@ -137,7 +78,8 @@ void SRewindDebuggerComponentTree::Construct(const FArguments& InArgs)
 									.OnExpansionChanged_Static(&ComponentTreeViewExpansionChanged)
 									.SelectionMode(ESelectionMode::Single)
 									.OnSelectionChanged(InArgs._OnSelectionChanged)
-									.OnContextMenuOpening(this, &SRewindDebuggerComponentTree::ComponentTreeOnContextMenuOpening);
+									.OnMouseButtonDoubleClick(InArgs._OnMouseButtonDoubleClick)
+									.OnContextMenuOpening(InArgs._OnContextMenuOpening);
 
 	ChildSlot
 	[
diff --git a/Engine/Plugins/Animation/GameplayInsights/Source/RewindDebugger/Private/SRewindDebuggerComponentTree.h b/Engine/Plugins/Animation/GameplayInsights/Source/RewindDebugger/Private/SRewindDebuggerComponentTree.h
--- a/Engine/Plugins/Animation/GameplayInsights/Source/RewindDebugger/Private/SRewindDebuggerComponentTree.h
+++ b/Engine/Plugins/Animation/GameplayInsights/Source/RewindDebugger/Private/SRewindDebuggerComponentTree.h
@@ -11,14 +11,15 @@
 class SRewindDebuggerComponentTree : public SCompoundWidget
 {
 	using FOnSelectionChanged = typename STreeView<TSharedPtr<FDebugObjectInfo>>::FOnSelectionChanged;
-
+	using FOnMouseButtonDoubleClick = typename STreeView<TSharedPtr<FDebugObjectInfo>>::FOnMouseButtonDoubleClick;
+	
 public:
 	SLATE_BEGIN_ARGS(SRewindDebuggerComponentTree) { }
 		SLATE_ARGUMENT( TArray< TSharedPtr< FDebugObjectInfo > >*, DebugComponents );
 		SLATE_EVENT( FOnSelectionChanged, OnSelectionChanged )
+		SLATE_EVENT( FOnMouseButtonDoubleClick, OnMouseButtonDoubleClick )
+		SLATE_EVENT( FOnContextMenuOpening, OnContextMenuOpening )
 	SLATE_END_ARGS()
-	
-public:
 
 	/**
 	* Default constructor.
@@ -35,8 +36,6 @@ public:
 
 	void Refresh();
 
-	TSharedPtr<SWidget> ComponentTreeOnContextMenuOpening();
-
 private:
 	TArray<TSharedPtr<FDebugObjectInfo>>* DebugComponents;
     void OnComponentSelectionChanged(TSharedPtr<FDebugObjectInfo> SelectedItem, ESelectInfo::Type SelectInfo);
diff --git a/Engine/Source/Editor/RewindDebuggerInterface/Private/RewindDebuggerInterface.cpp b/Engine/Source/Editor/RewindDebuggerInterface/Private/RewindDebuggerInterface.cpp
--- a/Engine/Source/Editor/RewindDebuggerInterface/Private/RewindDebuggerInterface.cpp
+++ b/Engine/Source/Editor/RewindDebuggerInterface/Private/RewindDebuggerInterface.cpp
@@ -2,10 +2,12 @@
 
 #include "IRewindDebuggerExtension.h"
 #include "IRewindDebuggerViewCreator.h"
+#include "IRewindDebuggerDoubleClickHandler.h"
 #include "IRewindDebugger.h"
 
 const FName IRewindDebuggerExtension::ModularFeatureName = "RewindDebuggerExtension";
 const FName IRewindDebuggerViewCreator::ModularFeatureName = "RewindDebuggerViewCreator";
+const FName IRewindDebuggerDoubleClickHandler::ModularFeatureName = "RewindDebuggerDoubleClickHandler";
 
 IRewindDebugger::~IRewindDebugger()
 {
diff --git a/Engine/Source/Editor/RewindDebuggerInterface/Public/IRewindDebugger.h b/Engine/Source/Editor/RewindDebuggerInterface/Public/IRewindDebugger.h
--- a/Engine/Source/Editor/RewindDebuggerInterface/Public/IRewindDebugger.h
+++ b/Engine/Source/Editor/RewindDebuggerInterface/Public/IRewindDebugger.h
@@ -2,6 +2,7 @@
 #pragma once
 
 #include "CoreMinimal.h"
+#include "IRewindDebugger.generated.h"
 
 // IRewindDebugger
 //
@@ -25,6 +26,15 @@ struct FDebugObjectInfo
 	TArray<TSharedPtr<FDebugObjectInfo>> Children;
 };
 
+UCLASS()
+class REWINDDEBUGGERINTERFACE_API UComponentContextMenuContext : public UObject
+{
+	GENERATED_BODY()
+public:
+	TSharedPtr<FDebugObjectInfo> SelectedObject;
+	TArray<FName> TypeHierarchy;
+};
+
 class REWINDDEBUGGERINTERFACE_API IRewindDebugger
 {
 public:
@@ -43,6 +53,9 @@ public:
 
 	// get a list of all components of the selected target actor (with the actor as the first element in the list)
 	virtual TArray<TSharedPtr<FDebugObjectInfo>>& GetDebugComponents() = 0;
+	
+	// returns the currently selected debug component
+	virtual TSharedPtr<FDebugObjectInfo> GetSelectedComponent() const = 0;
 
 	// get posiotion of the selected target actor (returns true if position is valid)
 	virtual bool GetTargetActorPosition(FVector& OutPosition) const = 0;
diff --git a/Engine/Source/Editor/RewindDebuggerInterface/Public/IRewindDebuggerDoubleClickHandler.h b/Engine/Source/Editor/RewindDebuggerInterface/Public/IRewindDebuggerDoubleClickHandler.h
new file mode 100644
--- /dev/null
+++ b/Engine/Source/Editor/RewindDebuggerInterface/Public/IRewindDebuggerDoubleClickHandler.h
@@ -0,0 +1,20 @@
+// Copyright Epic Games, Inc. All Rights Reserved.
+
+#pragma once
+
+#include "CoreMinimal.h"
+#include "IRewindDebugger.h"
+#include "Features/IModularFeatures.h"
+
+// Interface class for adding custom double click handling to components
+class REWINDDEBUGGERINTERFACE_API IRewindDebuggerDoubleClickHandler : public IModularFeature
+{
+public:
+	static const FName ModularFeatureName;
+	
+	// returns the name of a type of UObject for which this menu extender will be called
+	virtual FName GetTargetTypeName() const = 0;
+
+	// called when the selected component has been double clicked.  return true if something was opened
+	virtual bool HandleDoubleClick(IRewindDebugger* IRewindDebugger) = 0;
+};
diff --git a/Engine/Source/Editor/RewindDebuggerInterface/RewindDebuggerInterface.Build.cs b/Engine/Source/Editor/RewindDebuggerInterface/RewindDebuggerInterface.Build.cs
--- a/Engine/Source/Editor/RewindDebuggerInterface/RewindDebuggerInterface.Build.cs
+++ b/Engine/Source/Editor/RewindDebuggerInterface/RewindDebuggerInterface.Build.cs
@@ -9,6 +9,7 @@ public class RewindDebuggerInterface : ModuleRules
 		PrivateDependencyModuleNames.AddRange(
 			new string[] {
 				"Core",
+				"CoreUObject",
 				"SlateCore",
 			}
 		);
