Add subobjectEditorMenuContext and ExtensionContext in favor of SCS Editor extension context.

#rnx

diff --git a/Engine/Plugins/Animation/GameplayInsights/Source/GameplayInsights/Private/GameplayInsightsModule.cpp b/Engine/Plugins/Animation/GameplayInsights/Source/GameplayInsights/Private/GameplayInsightsModule.cpp
--- a/Engine/Plugins/Animation/GameplayInsights/Source/GameplayInsights/Private/GameplayInsightsModule.cpp
+++ b/Engine/Plugins/Animation/GameplayInsights/Source/GameplayInsights/Private/GameplayInsightsModule.cpp
@@ -28,7 +28,7 @@
 #include "ToolMenus.h"
 #include "LevelEditorMenuContext.h"
 #include "Engine/Selection.h"
-#include "SSCSEditorMenuContext.h"
+#include "SubobjectEditorMenuContext.h"
 #include "GameplayInsightsStyle.h"
 #include "SSubobjectInstanceEditor.h"
 
@@ -381,7 +381,7 @@ void FGameplayInsightsModule::RegisterMenus()
 		{
 			if (FObjectPropertyTrace::IsEnabled())
 			{
-				USSCSEditorMenuContext* ContextObject = InContext.FindContext<USSCSEditorMenuContext>();
+				USubobjectEditorMenuContext* ContextObject = InContext.FindContext<USubobjectEditorMenuContext>();
 				TSharedPtr<SSubobjectEditor> SubobjectEditor = ContextObject ? ContextObject->SubobjectEditor.Pin() : nullptr;
 
 				if (SubobjectEditor.IsValid() && StaticCastSharedPtr<SSubobjectInstanceEditor>(SubobjectEditor))
@@ -407,7 +407,7 @@ void FGameplayInsightsModule::RegisterMenus()
 		{
 			if (FObjectPropertyTrace::IsEnabled())
 			{
-				USSCSEditorMenuContext* ContextObject = InContext.FindContext<USSCSEditorMenuContext>();
+				USubobjectEditorMenuContext* ContextObject = InContext.FindContext<USubobjectEditorMenuContext>();
 				if (ContextObject && ContextObject->SubobjectEditor.IsValid() && StaticCastSharedPtr<SSubobjectInstanceEditor>(ContextObject->SubobjectEditor.Pin()))
 				{
 					TSharedPtr<SSubobjectEditor> SubobjectEditor = ContextObject->SubobjectEditor.Pin();
@@ -422,7 +422,7 @@ void FGameplayInsightsModule::RegisterMenus()
 		});
 		Action.GetActionCheckState = FToolMenuGetActionCheckState::CreateLambda([&GetCheckState](const FToolMenuContext& InContext)
 		{
-			USSCSEditorMenuContext* ContextObject = InContext.FindContext<USSCSEditorMenuContext>();
+			USubobjectEditorMenuContext* ContextObject = InContext.FindContext<USubobjectEditorMenuContext>();
 			if (ContextObject && StaticCastSharedPtr<SSubobjectInstanceEditor>(ContextObject->SubobjectEditor.Pin()))
 			{
 				TSharedPtr<SSubobjectEditor> SubobjectEditor = ContextObject->SubobjectEditor.Pin();
@@ -439,7 +439,7 @@ void FGameplayInsightsModule::RegisterMenus()
 		{
 			if (FObjectPropertyTrace::IsEnabled())
 			{
-				USSCSEditorMenuContext* ContextObject = InContext.FindContext<USSCSEditorMenuContext>();
+				USubobjectEditorMenuContext* ContextObject = InContext.FindContext<USubobjectEditorMenuContext>();
 				if (ContextObject && ContextObject->SubobjectEditor.IsValid() && StaticCastSharedPtr<SSubobjectInstanceEditor>(ContextObject->SubobjectEditor.Pin()))
 				{
 					return true;
diff --git a/Engine/Plugins/Runtime/nDisplay/Source/DisplayClusterConfigurator/Private/DisplayClusterConfiguratorBlueprintEditor.cpp b/Engine/Plugins/Runtime/nDisplay/Source/DisplayClusterConfigurator/Private/DisplayClusterConfiguratorBlueprintEditor.cpp
--- a/Engine/Plugins/Runtime/nDisplay/Source/DisplayClusterConfigurator/Private/DisplayClusterConfiguratorBlueprintEditor.cpp
+++ b/Engine/Plugins/Runtime/nDisplay/Source/DisplayClusterConfigurator/Private/DisplayClusterConfiguratorBlueprintEditor.cpp
@@ -44,7 +44,8 @@
 #include "EditorViewportTabContent.h"
 #include "ISCSEditorUICustomization.h"
 #include "SBlueprintEditorToolbar.h"
-#include "SCSEditorExtensionContext.h"
+#include "SubobjectEditorExtensionContext.h"
+#include "SubobjectEditorExtensionContext.h"
 #include "SKismetInspector.h"
 #include "SSCSEditor.h"
 #include "ToolMenu.h"
@@ -1271,9 +1272,9 @@ TSharedRef<SWidget> FDisplayClusterConfiguratorBlueprintEditor::CreateSCSEditorE
 	auto PerformComboAddClass = [=](TSubclassOf<UActorComponent> ComponentClass, EComponentCreateAction::Type ComponentCreateAction,
 		UObject* AssetOverride) -> UActorComponent*
 	{
-		if (USCSEditorExtensionContext* SCSEditorExtensionContext = Cast<USCSEditorExtensionContext>(ExtensionContext.Get()))
+		if (USubobjectEditorExtensionContext* SubobjectEditorExtensionContext = Cast<USubobjectEditorExtensionContext>(ExtensionContext.Get()))
 		{
-			if (const TSharedPtr<SSubobjectEditor> SubobjectEditor = SCSEditorExtensionContext->GetSubobjectEditor().Pin())
+			if (const TSharedPtr<SSubobjectEditor> SubobjectEditor = SubobjectEditorExtensionContext->GetSubobjectEditor().Pin())
 			{
 				UClass* NewClass = ComponentClass;
 
@@ -1355,11 +1356,11 @@ TSharedRef<SWidget> FDisplayClusterConfiguratorBlueprintEditor::CreateSCSEditorE
 			.OnComponentClassSelected(FComponentClassSelected::CreateLambda(PerformComboAddClass))
 			.Visibility_Lambda([this, ExtensionContext]
 			{
-				if(USCSEditorExtensionContext* SCSEditorExtensionContext = Cast<USCSEditorExtensionContext>(ExtensionContext.Get()))
+				if(USubobjectEditorExtensionContext* SubobjectEditorExtensionContext = Cast<USubobjectEditorExtensionContext>(ExtensionContext.Get()))
 				{
-					if (SCSEditorExtensionContext->GetSubobjectEditor().IsValid())
+					if (SubobjectEditorExtensionContext->GetSubobjectEditor().IsValid())
 					{
-						if (const ADisplayClusterRootActor* RootActor = Cast<ADisplayClusterRootActor>(SCSEditorExtensionContext->GetSubobjectEditor().Pin()->GetObjectContext()))
+						if (const ADisplayClusterRootActor* RootActor = Cast<ADisplayClusterRootActor>(SubobjectEditorExtensionContext->GetSubobjectEditor().Pin()->GetObjectContext()))
 						{
 							// Check that this context is for our root actor and that we are editing just the CDO.
 							// Without this we could show up for unrelated blueprints or in the level editor when selecting our actor.
diff --git a/Engine/Source/Editor/Kismet/Private/SSCSEditor.cpp b/Engine/Source/Editor/Kismet/Private/SSCSEditor.cpp
--- a/Engine/Source/Editor/Kismet/Private/SSCSEditor.cpp
+++ b/Engine/Source/Editor/Kismet/Private/SSCSEditor.cpp
@@ -97,6 +97,8 @@ static const FName SCS_ColumnName_Mobility( "Mobility" );
 
 static const FName SCS_ContextMenuName( "Kismet.SCSEditorContextMenu" );
 
+PRAGMA_DISABLE_DEPRECATION_WARNINGS
+
 //////////////////////////////////////////////////////////////////////////
 // SSCSEditorDragDropTree
 void SSCSEditorDragDropTree::Construct( const FArguments& InArgs )
@@ -7320,4 +7322,6 @@ TSubclassOf<UActorComponent> SSCSEditor::GetComponentTypeFilterToApply() const
 	return ComponentType;
 }
 
-#undef LOCTEXT_NAMESPACE
+PRAGMA_ENABLE_DEPRECATION_WARNINGS
+
+#undef LOCTEXT_NAMESPACE
\ No newline at end of file
diff --git a/Engine/Source/Editor/Kismet/Private/SubobjectEditorExtensionContext.cpp b/Engine/Source/Editor/Kismet/Private/SubobjectEditorExtensionContext.cpp
new file mode 100644
--- /dev/null
+++ b/Engine/Source/Editor/Kismet/Private/SubobjectEditorExtensionContext.cpp
@@ -0,0 +1,3 @@
+// Copyright Epic Games, Inc. All Rights Reserved.
+
+#include "SubobjectEditorExtensionContext.h"
diff --git a/Engine/Source/Editor/Kismet/Private/SubobjectEditorMenuContext.cpp b/Engine/Source/Editor/Kismet/Private/SubobjectEditorMenuContext.cpp
new file mode 100644
--- /dev/null
+++ b/Engine/Source/Editor/Kismet/Private/SubobjectEditorMenuContext.cpp
@@ -0,0 +1,3 @@
+// Copyright Epic Games, Inc. All Rights Reserved.
+
+#include "SubobjectEditorMenuContext.h"
diff --git a/Engine/Source/Editor/Kismet/Public/SCSEditorExtensionContext.h b/Engine/Source/Editor/Kismet/Public/SCSEditorExtensionContext.h
--- a/Engine/Source/Editor/Kismet/Public/SCSEditorExtensionContext.h
+++ b/Engine/Source/Editor/Kismet/Public/SCSEditorExtensionContext.h
@@ -10,7 +10,8 @@ class SSubobjectEditor;
 class SSubobjectInstanceEditor;
 class SSubobjectBlueprintEditor;
 
-// #TODO_BH Rename this file to Subobject Extension context
+class UE_DEPRECATED(5.0, "USCSEditorExtensionContext has been deprecated, use USubobjectEditorExtensionContext instead.") USCSEditorExtensionContext;
+
 UCLASS()
 class KISMET_API USCSEditorExtensionContext : public UObject
 {
diff --git a/Engine/Source/Editor/Kismet/Public/SSCSEditorMenuContext.h b/Engine/Source/Editor/Kismet/Public/SSCSEditorMenuContext.h
--- a/Engine/Source/Editor/Kismet/Public/SSCSEditorMenuContext.h
+++ b/Engine/Source/Editor/Kismet/Public/SSCSEditorMenuContext.h
@@ -9,7 +9,7 @@
 class SSCSEditor;
 class SSubobjectEditor;
 
-// #TODO_BH Rename this for subobject editor menu context
+class UE_DEPRECATED(5.0, "USSCSEditorMenuContext has been deprecated, use USubobjectEditorMenuContext instead.") USSCSEditorMenuContext;
 UCLASS()
 class KISMET_API USSCSEditorMenuContext : public UObject
 {
diff --git a/Engine/Source/Editor/Kismet/Public/SubobjectEditorExtensionContext.h b/Engine/Source/Editor/Kismet/Public/SubobjectEditorExtensionContext.h
new file mode 100644
--- /dev/null
+++ b/Engine/Source/Editor/Kismet/Public/SubobjectEditorExtensionContext.h
@@ -0,0 +1,27 @@
+// Copyright Epic Games, Inc. All Rights Reserved.
+
+#pragma once
+
+#include "CoreMinimal.h"
+#include "SubobjectEditorExtensionContext.generated.h"
+
+class SSubobjectEditor;
+class SSubobjectInstanceEditor;
+class SSubobjectBlueprintEditor;
+
+UCLASS()
+class KISMET_API USubobjectEditorExtensionContext : public UObject
+{
+	GENERATED_BODY()
+
+public:
+	
+	const TWeakPtr<SSubobjectEditor>& GetSubobjectEditor() const { return SubobjectEditor; }
+
+private:
+	friend SSubobjectEditor;
+	friend SSubobjectInstanceEditor;
+	friend SSubobjectBlueprintEditor;
+
+	TWeakPtr<SSubobjectEditor> SubobjectEditor;
+};
diff --git a/Engine/Source/Editor/Kismet/Public/SubobjectEditorMenuContext.h b/Engine/Source/Editor/Kismet/Public/SubobjectEditorMenuContext.h
new file mode 100644
--- /dev/null
+++ b/Engine/Source/Editor/Kismet/Public/SubobjectEditorMenuContext.h
@@ -0,0 +1,20 @@
+// Copyright Epic Games, Inc. All Rights Reserved.
+
+#pragma once
+
+#include "CoreMinimal.h"
+
+#include "SubobjectEditorMenuContext.generated.h"
+
+class SSubobjectEditor;
+
+UCLASS()
+class KISMET_API USubobjectEditorMenuContext : public UObject
+{
+	GENERATED_BODY()
+public:
+
+	TWeakPtr<SSubobjectEditor> SubobjectEditor;
+	
+	bool bOnlyShowPasteOption;
+};
diff --git a/Engine/Source/Editor/SubobjectEditor/Private/SSubobjectBlueprintEditor.cpp b/Engine/Source/Editor/SubobjectEditor/Private/SSubobjectBlueprintEditor.cpp
--- a/Engine/Source/Editor/SubobjectEditor/Private/SSubobjectBlueprintEditor.cpp
+++ b/Engine/Source/Editor/SubobjectEditor/Private/SSubobjectBlueprintEditor.cpp
@@ -7,7 +7,7 @@
 #include "GraphEditorActions.h"
 #include "Editor/EditorEngine.h"
 #include "ISCSEditorUICustomization.h"	// #TODO_BH Rename this to subobject
-#include "SCSEditorExtensionContext.h"	// #TODO_BH Rename this to subobject (this one will break VK, for sure)
+#include "SubobjectEditorExtensionContext.h"
 #include "ToolMenus.h"
 #include "PropertyPath.h"
 
@@ -68,7 +68,7 @@ void SSubobjectBlueprintEditor::Construct(const FArguments& InArgs)
 
 	FMenuBuilder EditBlueprintMenuBuilder = CreateMenuBuilder();
 
-	USCSEditorExtensionContext* ExtensionContext = NewObject<USCSEditorExtensionContext>();
+	USubobjectEditorExtensionContext* ExtensionContext = NewObject<USubobjectEditorExtensionContext>();
 	ExtensionContext->SubobjectEditor = SharedThis(this);
 	ExtensionContext->AddToRoot();
 
diff --git a/Engine/Source/Editor/SubobjectEditor/Private/SSubobjectEditor.cpp b/Engine/Source/Editor/SubobjectEditor/Private/SSubobjectEditor.cpp
--- a/Engine/Source/Editor/SubobjectEditor/Private/SSubobjectEditor.cpp
+++ b/Engine/Source/Editor/SubobjectEditor/Private/SSubobjectEditor.cpp
@@ -15,7 +15,7 @@
 #include "GraphEditorActions.h"
 #include "Editor/EditorEngine.h"
 #include "ISCSEditorUICustomization.h"	// #TODO_BH Rename this to subobject
-#include "SCSEditorExtensionContext.h"	// #TODO_BH Rename this to subobject (this one will break VK, for sure)
+#include "SubobjectEditorMenuContext.h"
 #include "SSCSEditorMenuContext.h"		// #TODO_BH Rename this to subobject
 #include "Toolkits/ToolkitManager.h"
 #include "BlueprintEditor.h"
@@ -131,7 +131,7 @@ namespace SubobjectEditorUtils
 		FToolMenuSection& Section = Menu->AddSection("ChildActor", LOCTEXT("ChildActorHeading", "Child Actor"));
 		{
 			TWeakPtr<SSubobjectEditor> WeakEditorPtr;
-			if (USSCSEditorMenuContext* MenuContext = Menu->FindContext<USSCSEditorMenuContext>())
+			if (USubobjectEditorMenuContext* MenuContext = Menu->FindContext<USubobjectEditorMenuContext>())
 			{
 				WeakEditorPtr = MenuContext->SubobjectEditor;
 			}
@@ -2960,7 +2960,7 @@ void SSubobjectEditor::RegisterContextMenu()
 		UToolMenu* Menu = ToolMenus->RegisterMenu(SubobjectTree_ContextMenuName);
 		Menu->AddDynamicSection("SCSEditorDynamic", FNewToolMenuDelegate::CreateLambda([](UToolMenu* InMenu)
 		{
-			USSCSEditorMenuContext* ContextObject = InMenu->FindContext<USSCSEditorMenuContext>();
+			USubobjectEditorMenuContext* ContextObject = InMenu->FindContext<USubobjectEditorMenuContext>();
 			if (ContextObject && ContextObject->SubobjectEditor.IsValid())
 			{
 				ContextObject->SubobjectEditor.Pin()->PopulateContextMenu(InMenu);
@@ -2976,7 +2976,7 @@ TSharedPtr<SWidget> SSubobjectEditor::CreateContextMenu()
 	if (SelectedItems.Num() > 0 || CanPasteNodes())
 	{
 		RegisterContextMenu();
-		USSCSEditorMenuContext* ContextObject = NewObject<USSCSEditorMenuContext>();
+		USubobjectEditorMenuContext* ContextObject = NewObject<USubobjectEditorMenuContext>();
 		ContextObject->SubobjectEditor = SharedThis(this);
 		FToolMenuContext ToolMenuContext(CommandList, TSharedPtr<FExtender>(), ContextObject);
 		return UToolMenus::Get()->GenerateWidget(SubobjectTree_ContextMenuName, ToolMenuContext);
diff --git a/Engine/Source/Editor/SubobjectEditor/Private/SSubobjectInstanceEditor.cpp b/Engine/Source/Editor/SubobjectEditor/Private/SSubobjectInstanceEditor.cpp
--- a/Engine/Source/Editor/SubobjectEditor/Private/SSubobjectInstanceEditor.cpp
+++ b/Engine/Source/Editor/SubobjectEditor/Private/SSubobjectInstanceEditor.cpp
@@ -7,7 +7,7 @@
 #include "GraphEditorActions.h"
 #include "Editor/EditorEngine.h"
 #include "ISCSEditorUICustomization.h"	// #TODO_BH Rename this to subobject
-#include "SCSEditorExtensionContext.h"	// #TODO_BH Rename this to subobject
+#include "SubobjectEditorExtensionContext.h"
 
 #include "Styling/SlateIconFinder.h"
 #include "SlateOptMacros.h"
@@ -60,7 +60,7 @@ void SSubobjectInstanceEditor::Construct(const FArguments& InArgs)
 	FMenuBuilder EditBlueprintMenuBuilder = CreateMenuBuilder();
 
 	// Extension context for new boi
-	USCSEditorExtensionContext* ExtensionContext = NewObject<USCSEditorExtensionContext>();
+	USubobjectEditorExtensionContext* ExtensionContext = NewObject<USubobjectEditorExtensionContext>();
 	ExtensionContext->SubobjectEditor = SharedThis(this);
 	ExtensionContext->AddToRoot();
 
diff --git a/Engine/Source/Editor/UnrealEd/Private/Kismet2/ChildActorComponentEditorUtils.cpp b/Engine/Source/Editor/UnrealEd/Private/Kismet2/ChildActorComponentEditorUtils.cpp
--- a/Engine/Source/Editor/UnrealEd/Private/Kismet2/ChildActorComponentEditorUtils.cpp
+++ b/Engine/Source/Editor/UnrealEd/Private/Kismet2/ChildActorComponentEditorUtils.cpp
@@ -5,6 +5,7 @@
 #include "ToolMenus.h"
 #include "SSCSEditor.h"
 #include "SSCSEditorMenuContext.h"
+#include "SubobjectEditorMenuContext.h"
 #include "SSubobjectEditor.h"
 
 #define LOCTEXT_NAMESPACE "ChildActorComponentEditorUtils"
@@ -187,7 +188,7 @@ void FChildActorComponentEditorUtils::FillComponentContextMenuOptions(UToolMenu*
 	FToolMenuSection& Section = Menu->AddSection("ChildActorComponent", LOCTEXT("ChildActorComponentHeading", "Child Actor Component"));
 	{
 		TWeakPtr<SSubobjectEditor> WeakEditorPtr;
-		if (USSCSEditorMenuContext* MenuContext = Menu->FindContext<USSCSEditorMenuContext>())
+		if (USubobjectEditorMenuContext* MenuContext = Menu->FindContext<USubobjectEditorMenuContext>())
 		{
 			WeakEditorPtr = MenuContext->SubobjectEditor;
 		}
@@ -224,7 +225,7 @@ void FChildActorComponentEditorUtils::FillChildActorContextMenuOptions(UToolMenu
 	FToolMenuSection& Section = Menu->AddSection("ChildActor", LOCTEXT("ChildActorHeading", "Child Actor"));
 	{
 		TWeakPtr<SSubobjectEditor> WeakEditorPtr;
-		if (USSCSEditorMenuContext* MenuContext = Menu->FindContext<USSCSEditorMenuContext>())
+		if (USubobjectEditorMenuContext* MenuContext = Menu->FindContext<USubobjectEditorMenuContext>())
 		{
 			WeakEditorPtr = MenuContext->SubobjectEditor;
 		}
