Fix incorrect ToolMenu implementation.  Was calling RegisterMenu multiple times, rather than Register menu once, followed by FindMenu.

diff --git a/Engine/Plugins/Animation/GameplayInsights/Source/RewindDebugger/Private/AnimInstanceHelpers.cpp b/Engine/Plugins/Animation/GameplayInsights/Source/RewindDebugger/Private/AnimInstanceHelpers.cpp
--- a/Engine/Plugins/Animation/GameplayInsights/Source/RewindDebugger/Private/AnimInstanceHelpers.cpp
+++ b/Engine/Plugins/Animation/GameplayInsights/Source/RewindDebugger/Private/AnimInstanceHelpers.cpp
@@ -72,7 +72,7 @@ FName FAnimInstanceDoubleClickHandler::GetTargetTypeName() const
 
 void FAnimInstanceMenu::Register()
 {
-	UToolMenu* Menu = UToolMenus::Get()->RegisterMenu("RewindDebugger.ComponentContextMenu");
+	UToolMenu* Menu = UToolMenus::Get()->FindMenu("RewindDebugger.ComponentContextMenu");
 	FToolMenuSection& Section = Menu->FindOrAddSection("Blueprint");
 	FToolMenuEntry& Entry = Section.AddDynamicEntry("DebugAnimInstanceEntry", FNewToolMenuSectionDelegate::CreateLambda([](FToolMenuSection& InSection)
 	{
diff --git a/Engine/Plugins/Animation/GameplayInsights/Source/RewindDebugger/Private/PropertyTraceMenu.cpp b/Engine/Plugins/Animation/GameplayInsights/Source/RewindDebugger/Private/PropertyTraceMenu.cpp
--- a/Engine/Plugins/Animation/GameplayInsights/Source/RewindDebugger/Private/PropertyTraceMenu.cpp
+++ b/Engine/Plugins/Animation/GameplayInsights/Source/RewindDebugger/Private/PropertyTraceMenu.cpp
@@ -10,7 +10,7 @@
 
 void FPropertyTraceMenu::Register()
 {
-	UToolMenu* Menu = UToolMenus::Get()->RegisterMenu("RewindDebugger.ComponentContextMenu");
+	UToolMenu* Menu = UToolMenus::Get()->FindMenu("RewindDebugger.ComponentContextMenu");
 	FToolMenuSection& Section = Menu->FindOrAddSection("Object");
 	FToolMenuEntry& Entry = Section.AddDynamicEntry("AssetManagerEditorViewCommands", FNewToolMenuSectionDelegate::CreateLambda([](FToolMenuSection& InSection)
 	{
diff --git a/Engine/Plugins/Animation/GameplayInsights/Source/RewindDebugger/Private/RewindDebuggerCamera.cpp b/Engine/Plugins/Animation/GameplayInsights/Source/RewindDebugger/Private/RewindDebuggerCamera.cpp
--- a/Engine/Plugins/Animation/GameplayInsights/Source/RewindDebugger/Private/RewindDebuggerCamera.cpp
+++ b/Engine/Plugins/Animation/GameplayInsights/Source/RewindDebugger/Private/RewindDebuggerCamera.cpp
@@ -19,7 +19,11 @@
 FRewindDebuggerCamera::FRewindDebuggerCamera()
 	: LastPositionValid(false)
 {
-	UToolMenu* Menu = UToolMenus::Get()->RegisterMenu("RewindDebugger.MainMenu");
+}
+
+void FRewindDebuggerCamera::Initialize()
+{
+	UToolMenu* Menu = UToolMenus::Get()->FindMenu("RewindDebugger.MainMenu");
 
 	Menu->AddSection("Camera Mode", LOCTEXT("Camera Mode","Camera Mode"));
 	Menu->AddMenuEntry("Camera Mode",
diff --git a/Engine/Plugins/Animation/GameplayInsights/Source/RewindDebugger/Private/RewindDebuggerCamera.h b/Engine/Plugins/Animation/GameplayInsights/Source/RewindDebugger/Private/RewindDebuggerCamera.h
--- a/Engine/Plugins/Animation/GameplayInsights/Source/RewindDebugger/Private/RewindDebuggerCamera.h
+++ b/Engine/Plugins/Animation/GameplayInsights/Source/RewindDebugger/Private/RewindDebuggerCamera.h
@@ -24,6 +24,7 @@ public:
 
 	FRewindDebuggerCamera();
 	virtual ~FRewindDebuggerCamera() {};
+	void Initialize();
 
 	virtual void Update(float DeltaTime, IRewindDebugger* RewindDebugger) override;
 
diff --git a/Engine/Plugins/Animation/GameplayInsights/Source/RewindDebugger/Private/RewindDebuggerModule.cpp b/Engine/Plugins/Animation/GameplayInsights/Source/RewindDebugger/Private/RewindDebuggerModule.cpp
--- a/Engine/Plugins/Animation/GameplayInsights/Source/RewindDebugger/Private/RewindDebuggerModule.cpp
+++ b/Engine/Plugins/Animation/GameplayInsights/Source/RewindDebugger/Private/RewindDebuggerModule.cpp
@@ -17,6 +17,7 @@
 #include "WorkspaceMenuStructureModule.h"
 #include "AnimInstanceHelpers.h"
 #include "PropertyTraceMenu.h"
+#include "ToolMenus.h"
 
 #define LOCTEXT_NAMESPACE "RewindDebuggerModule"
 
@@ -28,7 +29,7 @@ TSharedRef<SDockTab> FRewindDebuggerModule::SpawnRewindDebuggerTab(const FSpawnT
 	{
 		FRewindDebugger::Initialize();
 	}
-
+	
 	const TSharedRef<SDockTab> MajorTab = SNew(SDockTab)
 		.TabRole(ETabRole::NomadTab);
 
@@ -118,6 +119,9 @@ static FAnimInstanceDoubleClickHandler AnimInstanceDoubleClickHandler;
 
 void FRewindDebuggerModule::StartupModule()
 {
+	UToolMenus::Get()->RegisterMenu("RewindDebugger.MainMenu");
+	UToolMenus::Get()->RegisterMenu("RewindDebugger.ComponentContextMenu");
+
 	FRewindDebuggerStyle::Initialize();
 	FRewindDebuggerCommands::Register();
 
@@ -129,6 +133,7 @@ void FRewindDebuggerModule::StartupModule()
 		.SetTooltipText(LOCTEXT("TooltipText", "Opens Rewind Debugger."));
 
 
+	RewindDebuggerCameraExtension.Initialize();
 	IModularFeatures::Get().RegisterModularFeature(IRewindDebuggerExtension::ModularFeatureName, &RewindDebuggerCameraExtension);
 	IModularFeatures::Get().RegisterModularFeature(IRewindDebuggerDoubleClickHandler::ModularFeatureName, &AnimInstanceDoubleClickHandler);
 
diff --git a/Engine/Plugins/Animation/GameplayInsights/Source/RewindDebugger/Private/SRewindDebugger.cpp b/Engine/Plugins/Animation/GameplayInsights/Source/RewindDebugger/Private/SRewindDebugger.cpp
--- a/Engine/Plugins/Animation/GameplayInsights/Source/RewindDebugger/Private/SRewindDebugger.cpp
+++ b/Engine/Plugins/Animation/GameplayInsights/Source/RewindDebugger/Private/SRewindDebugger.cpp
@@ -290,7 +290,7 @@ void SRewindDebugger::Construct(const FArguments& InArgs, TSharedRef<FUICommandL
 	// load saved layout if it exists
 	Layout = FLayoutSaveRestore::LoadFromConfig(GEditorLayoutIni, Layout);
 
-	UToolMenu* Menu = UToolMenus::Get()->RegisterMenu("RewindDebugger.MainMenu");
+	UToolMenu* Menu = UToolMenus::Get()->FindMenu("RewindDebugger.MainMenu");
 
 	FToolMenuSection& Section = Menu->AddSection("ViewsSection", LOCTEXT("Views", "Views"));
 
