Track Controller/Pawn Possesion in insights trace data, and use that data to add the Controller for an Actor to the Rewind Debugger component view.

diff --git a/Engine/Plugins/Animation/GameplayInsights/Source/GameplayInsights/Private/GameplayAnalyzer.cpp b/Engine/Plugins/Animation/GameplayInsights/Source/GameplayInsights/Private/GameplayAnalyzer.cpp
--- a/Engine/Plugins/Animation/GameplayInsights/Source/GameplayInsights/Private/GameplayAnalyzer.cpp
+++ b/Engine/Plugins/Animation/GameplayInsights/Source/GameplayInsights/Private/GameplayAnalyzer.cpp
@@ -19,6 +19,7 @@ void FGameplayAnalyzer::OnAnalysisBegin(const FOnAnalysisContext& Context)
 	Builder.RouteEvent(RouteId_Class, "Object", "Class");
 	Builder.RouteEvent(RouteId_Object, "Object", "Object");
 	Builder.RouteEvent(RouteId_ObjectEvent, "Object", "ObjectEvent");
+	Builder.RouteEvent(RouteId_PawnPossess, "Object", "PawnPossess");
 	Builder.RouteEvent(RouteId_World, "Object", "World");
 	Builder.RouteEvent(RouteId_View, "Object", "View");
 	Builder.RouteEvent(RouteId_ClassPropertyStringId, "Object", "ClassPropertyStringId");
@@ -102,6 +103,14 @@ bool FGameplayAnalyzer::OnEvent(uint16 RouteId, EStyle Style, const FOnEventCont
 		GameplayProvider.AppendObjectEvent(Id, Context.EventTime.AsSeconds(Cycle), *Event);
 		break;
 	}
+	case RouteId_PawnPossess:
+	{
+		uint64 Cycle = EventData.GetValue<uint64>("Cycle");
+		uint64 ControllerId = EventData.GetValue<uint64>("ControllerId");
+		uint64 PawnId = EventData.GetValue<uint64>("PawnId");
+		GameplayProvider.AppendPawnPossess(ControllerId, PawnId, Context.EventTime.AsSeconds(Cycle));
+		break;
+	}
 	case RouteId_View:
 	{
 		uint64 Cycle = EventData.GetValue<uint64>("Cycle");
diff --git a/Engine/Plugins/Animation/GameplayInsights/Source/GameplayInsights/Private/GameplayAnalyzer.h b/Engine/Plugins/Animation/GameplayInsights/Source/GameplayInsights/Private/GameplayAnalyzer.h
--- a/Engine/Plugins/Animation/GameplayInsights/Source/GameplayInsights/Private/GameplayAnalyzer.h
+++ b/Engine/Plugins/Animation/GameplayInsights/Source/GameplayInsights/Private/GameplayAnalyzer.h
@@ -23,6 +23,7 @@ private:
 		RouteId_Class,
 		RouteId_Object,
 		RouteId_ObjectEvent,
+		RouteId_PawnPossess,
 		RouteId_View,
 		RouteId_ClassPropertyStringId,
 		RouteId_ClassProperty,
diff --git a/Engine/Plugins/Animation/GameplayInsights/Source/GameplayInsights/Private/GameplayProvider.cpp b/Engine/Plugins/Animation/GameplayInsights/Source/GameplayInsights/Private/GameplayProvider.cpp
--- a/Engine/Plugins/Animation/GameplayInsights/Source/GameplayInsights/Private/GameplayProvider.cpp
+++ b/Engine/Plugins/Animation/GameplayInsights/Source/GameplayInsights/Private/GameplayProvider.cpp
@@ -2,6 +2,8 @@
 
 #include "GameplayProvider.h"
 
+#include "Model/AsyncEnumerateTask.h"
+
 FName FGameplayProvider::ProviderName("GameplayProvider");
 
 #define LOCTEXT_NAMESPACE "GameplayProvider"
@@ -9,6 +11,7 @@ FName FGameplayProvider::ProviderName("GameplayProvider");
 FGameplayProvider::FGameplayProvider(TraceServices::IAnalysisSession& InSession)
 	: Session(InSession)
 	, EndPlayEvent(nullptr)
+	, PawnPossession(InSession.GetLinearAllocator())
 	, bHasAnyData(false)
 	, bHasObjectProperties(false)
 {
@@ -316,6 +319,42 @@ void FGameplayProvider::AppendObjectEvent(uint64 InObjectId, double InTime, cons
 	Session.UpdateDurationSeconds(InTime);
 }
 
+void FGameplayProvider::AppendPawnPossess(uint64 InControllerId, uint64 InPawnId, double InTime)
+{
+	Session.WriteAccessCheck();
+	bHasAnyData = true;
+
+	// End any active controller attachment interval for this controller
+	if (const uint64 *FoundIndex = ActivePawnPossession.Find(InControllerId))
+	{
+		PawnPossession.EndEvent(*FoundIndex, InTime);
+		ActivePawnPossession.Remove(InControllerId);
+	}
+	
+	if (InPawnId)
+	{
+		FPawnPossessMessage Message;
+		Message.ControllerId = InControllerId;
+		Message.PawnId = InPawnId;
+		ActivePawnPossession.Add(InControllerId, PawnPossession.AppendBeginEvent(InTime, Message));
+	}
+}
+
+uint64 FGameplayProvider::FindPossessingController(uint64 PawnId, double Time) const
+{
+	uint64 ControllerId = 0;
+	PawnPossession.EnumerateEvents(Time,Time,[&ControllerId, PawnId](double StartTime,double EndTime, const FPawnPossessMessage Message)
+	{
+		if (Message.PawnId == PawnId)
+		{
+			ControllerId = Message.ControllerId;
+			return TraceServices::EEventEnumerate::Stop;
+		}
+		return TraceServices::EEventEnumerate::Continue;
+	});
+	return ControllerId;
+}
+
 void FGameplayProvider::ReadViewTimeline(TFunctionRef<void(const IGameplayProvider::ViewTimeline&)> Callback) const
 {
 	Session.ReadAccessCheck();
diff --git a/Engine/Plugins/Animation/GameplayInsights/Source/GameplayInsights/Private/GameplayProvider.h b/Engine/Plugins/Animation/GameplayInsights/Source/GameplayInsights/Private/GameplayProvider.h
--- a/Engine/Plugins/Animation/GameplayInsights/Source/GameplayInsights/Private/GameplayProvider.h
+++ b/Engine/Plugins/Animation/GameplayInsights/Source/GameplayInsights/Private/GameplayProvider.h
@@ -46,6 +46,9 @@ public:
 
 	/** Add an object event message */
 	void AppendObjectEvent(uint64 InObjectId, double InTime, const TCHAR* InEvent);
+	
+	/** Add a Controller Attach message */
+	void AppendPawnPossess(uint64 InControllerId, uint64 InPawnId, double InTime);
 
 	/** Add a view message */
 	void AppendView(uint64 InObjectId, double InTime, const FVector& InPosition, const FRotator& InRotation, float InFov, float InAspectRatio);
@@ -74,6 +77,9 @@ public:
 	/** Check whether we have any object property data */
 	bool HasObjectProperties() const;
 
+	/** Search PawnPossession timeline to find the Controller Object Id for a Pawn (or 0 if none)*/
+	virtual uint64 FindPossessingController(uint64 Pawn, double Time) const override;
+
 private:
 	TraceServices::IAnalysisSession& Session;
 
@@ -99,6 +105,12 @@ private:
 	TMap<uint64, uint32> ObjectIdToEventTimelines;
 	TMap<uint64, uint32> ObjectIdToPropertiesStorage;
 
+	struct FPawnPossessMessage
+	{
+		uint64 ControllerId = 0;
+		uint64 PawnId = 0;
+	};
+
 	struct FObjectPropertiesStorage
 	{
 		double OpenStartTime;
@@ -128,6 +140,12 @@ private:
 	/** Map of RecordingInfo Timelines by RecordingId - Each timeline is a mapping from Gameplay Elapsed time to Profiler Elapsed Time for one recording session */
 	TMap<uint32, TSharedRef<TraceServices::TPointTimeline<FRecordingInfoMessage>>> Recordings;
 
+	// Timeline containing intervals where a controller is attached to a pawn
+	TraceServices::TIntervalTimeline<FPawnPossessMessage> PawnPossession;
+
+	// map from controller id to index in the PawnPossession timeline, for lookup when ending events
+	TMap<uint64, uint64> ActivePawnPossession;
+
 	/** Whether we have any data */
 	bool bHasAnyData;
 
diff --git a/Engine/Plugins/Animation/GameplayInsights/Source/GameplayInsights/Public/IGameplayProvider.h b/Engine/Plugins/Animation/GameplayInsights/Source/GameplayInsights/Public/IGameplayProvider.h
--- a/Engine/Plugins/Animation/GameplayInsights/Source/GameplayInsights/Public/IGameplayProvider.h
+++ b/Engine/Plugins/Animation/GameplayInsights/Source/GameplayInsights/Public/IGameplayProvider.h
@@ -124,5 +124,6 @@ public:
 	virtual const TCHAR* GetPropertyName(uint32 InPropertyStringId) const = 0;
 	virtual const RecordingInfoTimeline* GetRecordingInfo(uint32 RecordingId) const = 0; 
 	virtual void ReadViewTimeline(TFunctionRef<void(const ViewTimeline&)> Callback) const = 0;
+	virtual uint64 FindPossessingController(uint64 Pawn, double Time) const = 0;
 };
 
diff --git a/Engine/Plugins/Animation/GameplayInsights/Source/RewindDebugger/Private/RewindDebugger.cpp b/Engine/Plugins/Animation/GameplayInsights/Source/RewindDebugger/Private/RewindDebugger.cpp
--- a/Engine/Plugins/Animation/GameplayInsights/Source/RewindDebugger/Private/RewindDebugger.cpp
+++ b/Engine/Plugins/Animation/GameplayInsights/Source/RewindDebugger/Private/RewindDebugger.cpp
@@ -179,7 +179,7 @@ void FRewindDebugger::OnPIEStopped(bool bSimulating)
 	SetCurrentScrubTime(0);
 }
 
-bool FRewindDebugger::UpdateComponentList(uint64 ParentId, TArray<TSharedPtr<FDebugObjectInfo>>& ComponentList)
+bool FRewindDebugger::UpdateComponentList(uint64 ParentId, TArray<TSharedPtr<FDebugObjectInfo>>& ComponentList, bool bAddController)
 {
 	TSharedPtr<const TraceServices::IAnalysisSession> Session = UnrealInsightsModule->GetAnalysisSession();
 	TraceServices::FAnalysisSessionReadScope SessionReadScope(*Session.Get());
@@ -193,25 +193,46 @@ bool FRewindDebugger::UpdateComponentList(uint64 ParentId, TArray<TSharedPtr<FDe
 	{
 		// todo: filter components based on creation/destruction frame, and the current scrubbing time (so dynamically created and destroyed components won't add up)
 
-		const FObjectInfo* ObjectInfo = &InObjectInfo;
-		if (ObjectInfo->OuterId == ParentId)
+		if (InObjectInfo.OuterId == ParentId)
 		{
-			int32 FoundIndex = ComponentList.FindLastByPredicate([ObjectInfo](const TSharedPtr<FDebugObjectInfo>& Info) { return Info->ObjectName == ObjectInfo->Name; });
+			const int32 FoundIndex = ComponentList.FindLastByPredicate([&InObjectInfo](const TSharedPtr<FDebugObjectInfo>& Info) { return Info->ObjectName == InObjectInfo.Name; });
 
 			if (FoundIndex >= 0 && ComponentList[FoundIndex]->ObjectName == InObjectInfo.Name)
 			{
-				ComponentList[FoundIndex]->ObjectId = ObjectInfo->Id; // there is an issue with skeletal mesh components changing ids
-				bChanged = bChanged || UpdateComponentList(ObjectInfo->Id, ComponentList[FoundIndex]->Children);
+				ComponentList[FoundIndex]->ObjectId = InObjectInfo.Id; // there is an issue with skeletal mesh components changing ids
+				bChanged = bChanged || UpdateComponentList(InObjectInfo.Id, ComponentList[FoundIndex]->Children);
 			}
 			else
 			{
 				bChanged = true;
-				ComponentList.Add(MakeShared<FDebugObjectInfo>(ObjectInfo->Id,InObjectInfo.Name));
-				UpdateComponentList(ObjectInfo->Id, ComponentList.Last()->Children);
+				ComponentList.Add(MakeShared<FDebugObjectInfo>(InObjectInfo.Id,InObjectInfo.Name));
+				UpdateComponentList(InObjectInfo.Id, ComponentList.Last()->Children);
 			}
 		}
 	});
 
+	// add controller and it's component hierarchy if one is attached
+	if (bAddController)
+	{
+		if (uint64 ControllerId = GameplayProvider->FindPossessingController(ParentId, CurrentTraceTime()))
+		{
+			const FObjectInfo& ObjectInfo = GameplayProvider->GetObjectInfo(ControllerId);
+			const int32 FoundIndex = ComponentList.FindLastByPredicate([ObjectInfo](const TSharedPtr<FDebugObjectInfo>& Info) { return Info->ObjectName == ObjectInfo.Name; });
+
+			if (FoundIndex >= 0 && ComponentList[FoundIndex]->ObjectName == ObjectInfo.Name)
+			{
+				ComponentList[FoundIndex]->ObjectId = ObjectInfo.Id;
+				bChanged = bChanged || UpdateComponentList(ObjectInfo.Id, ComponentList[FoundIndex]->Children);
+			}
+			else
+			{
+				bChanged = true;
+				ComponentList.Add(MakeShared<FDebugObjectInfo>(ObjectInfo.Id,ObjectInfo.Name));
+				UpdateComponentList(ObjectInfo.Id, ComponentList.Last()->Children);
+			}
+		}
+	}
+	
 	return bChanged;
 }
 
@@ -282,7 +303,7 @@ void FRewindDebugger::RefreshDebugComponents()
 
 			if (TargetActorId != 0 && DebugComponents.Num() > 0)
 			{
-				bChanged = bChanged || UpdateComponentList(TargetActorId, DebugComponents[0]->Children);
+				bChanged = bChanged || UpdateComponentList(TargetActorId, DebugComponents[0]->Children, true);
 			}
 
 			if (bChanged)
diff --git a/Engine/Plugins/Animation/GameplayInsights/Source/RewindDebugger/Private/RewindDebugger.h b/Engine/Plugins/Animation/GameplayInsights/Source/RewindDebugger/Private/RewindDebugger.h
--- a/Engine/Plugins/Animation/GameplayInsights/Source/RewindDebugger/Private/RewindDebugger.h
+++ b/Engine/Plugins/Animation/GameplayInsights/Source/RewindDebugger/Private/RewindDebugger.h
@@ -113,7 +113,7 @@ private:
 	void OnPIESingleStepped(bool bSimulating);
 
 
-	bool UpdateComponentList(uint64 ParentId, TArray<TSharedPtr<FDebugObjectInfo>>& NewComponentList);
+	bool UpdateComponentList(uint64 ParentId, TArray<TSharedPtr<FDebugObjectInfo>>& NewComponentList, bool bAddController = false);
 	void SetCurrentScrubTime(float Time);
 	void UpdateTraceTime();
 
diff --git a/Engine/Source/Runtime/Engine/Private/Controller.cpp b/Engine/Source/Runtime/Engine/Private/Controller.cpp
--- a/Engine/Source/Runtime/Engine/Private/Controller.cpp
+++ b/Engine/Source/Runtime/Engine/Private/Controller.cpp
@@ -23,6 +23,7 @@
 #include "Engine/Canvas.h"
 
 #include "GameFramework/PlayerState.h"
+#include "ObjectTrace.h"
 
 DEFINE_LOG_CATEGORY(LogController);
 DEFINE_LOG_CATEGORY(LogPath);
@@ -334,6 +335,8 @@ void AController::Possess(APawn* InPawn)
 		OnNewPawn.Broadcast(NewPawn);
 		OnPossessedPawnChanged.Broadcast(bNotificationRequired ? nullptr : CurrentPawn, NewPawn);
 	}
+	
+	TRACE_PAWN_POSSESS(this, InPawn); 
 }
 
 void AController::OnPossess(APawn* InPawn)
@@ -386,6 +389,8 @@ void AController::UnPossess()
 		OnNewPawn.Broadcast(NewPawn);
 		OnPossessedPawnChanged.Broadcast(CurrentPawn, NewPawn);
 	}
+	
+	TRACE_PAWN_POSSESS(this, (APawn*)nullptr) 
 }
 
 void AController::OnUnPossess()
diff --git a/Engine/Source/Runtime/Engine/Private/ObjectTrace.cpp b/Engine/Source/Runtime/Engine/Private/ObjectTrace.cpp
--- a/Engine/Source/Runtime/Engine/Private/ObjectTrace.cpp
+++ b/Engine/Source/Runtime/Engine/Private/ObjectTrace.cpp
@@ -44,6 +44,12 @@ UE_TRACE_EVENT_BEGIN(Object, ObjectEvent)
 	UE_TRACE_EVENT_FIELD(UE::Trace::WideString, Event)
 UE_TRACE_EVENT_END()
 
+UE_TRACE_EVENT_BEGIN(Object, PawnPossess)
+	UE_TRACE_EVENT_FIELD(uint64, Cycle)
+	UE_TRACE_EVENT_FIELD(uint64, ControllerId)
+	UE_TRACE_EVENT_FIELD(uint64, PawnId)
+UE_TRACE_EVENT_END()
+
 UE_TRACE_EVENT_BEGIN(Object, World)
 	UE_TRACE_EVENT_FIELD(uint64, Id)
 	UE_TRACE_EVENT_FIELD(int32, PIEInstanceId)
@@ -402,6 +408,26 @@ void FObjectTrace::OutputObjectEvent(const UObject* InObject, const TCHAR* InEve
 		<< ObjectEvent.Event(InEvent);
 }
 
+void FObjectTrace::OutputPawnPossess(const UObject* InController, const UObject* InPawn)
+{
+	bool bChannelEnabled = UE_TRACE_CHANNELEXPR_IS_ENABLED(ObjectChannel);
+	if (!bChannelEnabled || InController == nullptr)
+	{
+		return;
+	}
+
+	TRACE_OBJECT(InController);
+	if (InPawn)
+	{
+		TRACE_OBJECT(InPawn);
+	}
+	
+	UE_TRACE_LOG(Object, PawnPossess, ObjectChannel)
+		<< PawnPossess.Cycle(FPlatformTime::Cycles64())
+		<< PawnPossess.ControllerId(GetObjectId(InController))
+		<< PawnPossess.PawnId(GetObjectId(InPawn));
+}
+
 void FObjectTrace::OutputWorld(const UWorld* InWorld)
 {
 	bool bChannelEnabled = UE_TRACE_CHANNELEXPR_IS_ENABLED(ObjectChannel);
diff --git a/Engine/Source/Runtime/Engine/Public/ObjectTrace.h b/Engine/Source/Runtime/Engine/Public/ObjectTrace.h
--- a/Engine/Source/Runtime/Engine/Public/ObjectTrace.h
+++ b/Engine/Source/Runtime/Engine/Public/ObjectTrace.h
@@ -60,6 +60,9 @@ struct FObjectTrace
 
 	/** Helper function to output an object event */
 	ENGINE_API static void OutputObjectEvent(const UObject* InObject, const TCHAR* InEvent);
+	
+	/** Helper function to output controller attach event */
+	ENGINE_API static void OutputPawnPossess(const UObject* InController, const UObject* InPawn);
 
 	/** Helper function to get an object ID from a UObject */
 	ENGINE_API static uint64 GetObjectId(const UObject* InObject);
@@ -107,10 +110,16 @@ struct FObjectTrace
 #define TRACE_OBJECT_EVENT(Object, Event) \
 	if(CAN_TRACE_OBJECT(Object)) { UNCONDITIONAL_TRACE_OBJECT_EVENT(Object, Event); }
 
+#define TRACE_PAWN_POSSESS(Controller, Pawn)\
+	if(CAN_TRACE_OBJECT(Controller) && (Pawn==nullptr || (CAN_TRACE_OBJECT(Pawn)))) FObjectTrace::OutputPawnPossess(Controller, Pawn);
+
 #else
 
 #define TRACE_OBJECT_EVENT(Object, Event) \
 	UNCONDITIONAL_TRACE_OBJECT_EVENT(Object, Event);
+	
+#define TRACE_PAWN_POSSESS(Controller, Pawn)\
+	FObjectTrace::OutputPawnPossess(Controller, Pawn);
 
 #endif
 	
@@ -120,12 +129,14 @@ struct FObjectTrace
 #define TRACE_WORLD(World) \
 	FObjectTrace::OutputWorld(World);
 
+
 #else
 
 #define TRACE_CLASS(Class)
 #define TRACE_OBJECT(Object)
 #define TRACE_OBJECT_EVENT(Object, Event)
 #define TRACE_WORLD(World)
+#define TRACE_PAWN_POSSESS(Controller, Pawn)
 #define TRACE_VIEW(Player, View)
 
 #endif
\ No newline at end of file
