Add object existence interval tracking to Gameplay insights.
Remove most of the ObjectEvents being tracked, which will prevent allocating an event timeline for every object and wasting memory.

diff --git a/Engine/Plugins/Animation/ControlRig/Source/ControlRig/Private/ControlRig.cpp b/Engine/Plugins/Animation/ControlRig/Source/ControlRig/Private/ControlRig.cpp
--- a/Engine/Plugins/Animation/ControlRig/Source/ControlRig/Private/ControlRig.cpp
+++ b/Engine/Plugins/Animation/ControlRig/Source/ControlRig/Private/ControlRig.cpp
@@ -122,7 +122,7 @@ void UControlRig::BeginDestroy()
 	}
 #endif
 
-	TRACE_OBJECT_EVENT(this, DestroyControlRig);
+	TRACE_OBJECT_LIFETIME_END(this);
 }
 
 UWorld* UControlRig::GetWorld() const
@@ -153,7 +153,7 @@ UWorld* UControlRig::GetWorld() const
 
 void UControlRig::Initialize(bool bInitRigUnits)
 {
-	TRACE_OBJECT_EVENT(this, InitializeControlRig);
+	TRACE_OBJECT_LIFETIME_BEGIN(this);
 
 	DECLARE_SCOPE_HIERARCHICAL_COUNTER_FUNC()
 	QUICK_SCOPE_CYCLE_COUNTER(STAT_ControlRig_Initialize);
diff --git a/Engine/Plugins/Animation/GameplayInsights/Source/GameplayInsights/Private/GameplayAnalyzer.cpp b/Engine/Plugins/Animation/GameplayInsights/Source/GameplayInsights/Private/GameplayAnalyzer.cpp
--- a/Engine/Plugins/Animation/GameplayInsights/Source/GameplayInsights/Private/GameplayAnalyzer.cpp
+++ b/Engine/Plugins/Animation/GameplayInsights/Source/GameplayInsights/Private/GameplayAnalyzer.cpp
@@ -19,6 +19,8 @@ void FGameplayAnalyzer::OnAnalysisBegin(const FOnAnalysisContext& Context)
 	Builder.RouteEvent(RouteId_Class, "Object", "Class");
 	Builder.RouteEvent(RouteId_Object, "Object", "Object");
 	Builder.RouteEvent(RouteId_ObjectEvent, "Object", "ObjectEvent");
+	Builder.RouteEvent(RoutId_ObjectLifetimeBegin, "Object", "ObjectLifetimeBegin");
+	Builder.RouteEvent(RoutId_ObjectLifetimeEnd, "Object", "ObjectLifetimeEnd");
 	Builder.RouteEvent(RouteId_PawnPossess, "Object", "PawnPossess");
 	Builder.RouteEvent(RouteId_World, "Object", "World");
 	Builder.RouteEvent(RouteId_View, "Object", "View");
@@ -103,6 +105,20 @@ bool FGameplayAnalyzer::OnEvent(uint16 RouteId, EStyle Style, const FOnEventCont
 		GameplayProvider.AppendObjectEvent(Id, Context.EventTime.AsSeconds(Cycle), *Event);
 		break;
 	}
+	case RoutId_ObjectLifetimeBegin:
+	{
+		uint64 Cycle = EventData.GetValue<uint64>("Cycle");
+		uint64 Id = EventData.GetValue<uint64>("Id");
+		GameplayProvider.AppendObjectLifetimeBegin(Id, Context.EventTime.AsSeconds(Cycle));
+		break;
+	}
+	case RoutId_ObjectLifetimeEnd:
+	{
+		uint64 Cycle = EventData.GetValue<uint64>("Cycle");
+		uint64 Id = EventData.GetValue<uint64>("Id");
+		GameplayProvider.AppendObjectLifetimeEnd(Id, Context.EventTime.AsSeconds(Cycle));
+		break;
+	}
 	case RouteId_PawnPossess:
 	{
 		uint64 Cycle = EventData.GetValue<uint64>("Cycle");
diff --git a/Engine/Plugins/Animation/GameplayInsights/Source/GameplayInsights/Private/GameplayAnalyzer.h b/Engine/Plugins/Animation/GameplayInsights/Source/GameplayInsights/Private/GameplayAnalyzer.h
--- a/Engine/Plugins/Animation/GameplayInsights/Source/GameplayInsights/Private/GameplayAnalyzer.h
+++ b/Engine/Plugins/Animation/GameplayInsights/Source/GameplayInsights/Private/GameplayAnalyzer.h
@@ -23,6 +23,8 @@ private:
 		RouteId_Class,
 		RouteId_Object,
 		RouteId_ObjectEvent,
+		RoutId_ObjectLifetimeBegin,
+		RoutId_ObjectLifetimeEnd,
 		RouteId_PawnPossess,
 		RouteId_View,
 		RouteId_ClassPropertyStringId,
diff --git a/Engine/Plugins/Animation/GameplayInsights/Source/GameplayInsights/Private/GameplayProvider.cpp b/Engine/Plugins/Animation/GameplayInsights/Source/GameplayInsights/Private/GameplayProvider.cpp
--- a/Engine/Plugins/Animation/GameplayInsights/Source/GameplayInsights/Private/GameplayProvider.cpp
+++ b/Engine/Plugins/Animation/GameplayInsights/Source/GameplayInsights/Private/GameplayProvider.cpp
@@ -12,6 +12,7 @@ FGameplayProvider::FGameplayProvider(TraceServices::IAnalysisSession& InSession)
 	: Session(InSession)
 	, EndPlayEvent(nullptr)
 	, PawnPossession(InSession.GetLinearAllocator())
+	, ObjectLifetimes(InSession.GetLinearAllocator())
 	, bHasAnyData(false)
 	, bHasObjectProperties(false)
 {
@@ -92,6 +93,19 @@ void FGameplayProvider::EnumerateObjects(TFunctionRef<void(const FObjectInfo&)>
 	}
 }
 
+void FGameplayProvider::EnumerateObjects(double StartTime, double EndTime, TFunctionRef<void(const FObjectInfo&)> Callback) const
+{
+	Session.ReadAccessCheck();
+
+	ObjectLifetimes.EnumerateEvents(StartTime, EndTime,
+		[this, Callback](double InStartTime, double InEndTime, uint32 InDepth, const FObjectExistsMessage& ExistsMessage)
+		{
+			checkSlow(ObjectIdToIndexMap.Contains(ExistsMessage.ObjectId));
+			Callback(ObjectInfos[ObjectIdToIndexMap[ExistsMessage.ObjectId]]);
+			return TraceServices::EEventEnumerate::Continue;
+		});
+}
+
 const FClassInfo* FGameplayProvider::FindClassInfo(uint64 InClassId) const
 {
 	Session.ReadAccessCheck();
@@ -277,6 +291,36 @@ void FGameplayProvider::AppendObject(uint64 InObjectId, uint64 InOuterId, uint64
 	}
 }
 
+void FGameplayProvider::AppendObjectLifetimeBegin(uint64 InObjectId, double InTime)
+{
+	Session.WriteAccessCheck();
+	bHasAnyData = true;
+
+	if (InObjectId)
+	{
+		FObjectExistsMessage Message;
+		Message.ObjectId = InObjectId;
+		ActiveObjectLifetimes.Add(InObjectId, ObjectLifetimes.AppendBeginEvent(InTime, Message));
+	}
+}
+
+void FGameplayProvider::AppendObjectLifetimeEnd(uint64 InObjectId, double InTime)
+{
+	Session.WriteAccessCheck();
+	bHasAnyData = true;
+
+	if (const uint64 *FoundIndex = ActiveObjectLifetimes.Find(InObjectId))
+	{
+		ObjectLifetimes.EndEvent(*FoundIndex, InTime);
+		ActiveObjectLifetimes.Remove(InObjectId);
+	}
+
+	if(int32* ObjectInfoIndex = ObjectIdToIndexMap.Find(InObjectId))
+	{
+		OnObjectEndPlayDelegate.Broadcast(InObjectId, InTime, ObjectInfos[*ObjectInfoIndex]);
+	}
+}
+
 void FGameplayProvider::AppendObjectEvent(uint64 InObjectId, double InTime, const TCHAR* InEventName)
 {
 	Session.WriteAccessCheck();
diff --git a/Engine/Plugins/Animation/GameplayInsights/Source/GameplayInsights/Private/GameplayProvider.h b/Engine/Plugins/Animation/GameplayInsights/Source/GameplayInsights/Private/GameplayProvider.h
--- a/Engine/Plugins/Animation/GameplayInsights/Source/GameplayInsights/Private/GameplayProvider.h
+++ b/Engine/Plugins/Animation/GameplayInsights/Source/GameplayInsights/Private/GameplayProvider.h
@@ -24,6 +24,7 @@ public:
 	virtual bool ReadObjectPropertiesTimeline(uint64 InObjectId, TFunctionRef<void(const ObjectPropertiesTimeline&)> Callback) const override;
 	virtual void EnumerateObjectPropertyValues(uint64 InObjectId, const FObjectPropertiesMessage& InMessage, TFunctionRef<void(const FObjectPropertyValue&)> Callback) const override;
 	virtual void EnumerateObjects(TFunctionRef<void(const FObjectInfo&)> Callback) const override;
+	virtual void EnumerateObjects(double StartTime, double EndTime, TFunctionRef<void(const FObjectInfo&)> Callback) const override;
 	virtual const FClassInfo* FindClassInfo(uint64 InClassId) const override;
 	virtual const FClassInfo* FindClassInfo(const TCHAR* InClassPath) const override;
 	virtual const FObjectInfo* FindObjectInfo(uint64 InObjectId) const override;
@@ -44,6 +45,12 @@ public:
 	/** Add an object message */
 	void AppendObject(uint64 InObjectId, uint64 InOuterId, uint64 InClassId, const TCHAR* InObjectName, const TCHAR* InObjectPathName);
 
+	/** Add an object create message */
+	void AppendObjectLifetimeBegin(uint64 InObjectId, double InTime);
+	
+	/** Add an object destroy message */
+	void AppendObjectLifetimeEnd(uint64 InObjectId, double InTime);	
+
 	/** Add an object event message */
 	void AppendObjectEvent(uint64 InObjectId, double InTime, const TCHAR* InEvent);
 	
@@ -110,6 +117,11 @@ private:
 		uint64 ControllerId = 0;
 		uint64 PawnId = 0;
 	};
+	
+	struct FObjectExistsMessage
+	{
+		uint64 ObjectId = 0;
+	};
 
 	struct FObjectPropertiesStorage
 	{
@@ -142,10 +154,17 @@ private:
 
 	// Timeline containing intervals where a controller is attached to a pawn
 	TraceServices::TIntervalTimeline<FPawnPossessMessage> PawnPossession;
+	
+	// Timeline containing intervals where an object exists
+	TraceServices::TIntervalTimeline<FObjectExistsMessage> ObjectLifetimes;
+
 
 	// map from controller id to index in the PawnPossession timeline, for lookup when ending events
 	TMap<uint64, uint64> ActivePawnPossession;
 
+	// map from object id to index in ObjectLifetimes, for lookup when objects are destroyed
+	TMap<uint64, uint64> ActiveObjectLifetimes;
+
 	/** Whether we have any data */
 	bool bHasAnyData;
 
diff --git a/Engine/Plugins/Animation/GameplayInsights/Source/GameplayInsights/Public/IGameplayProvider.h b/Engine/Plugins/Animation/GameplayInsights/Source/GameplayInsights/Public/IGameplayProvider.h
--- a/Engine/Plugins/Animation/GameplayInsights/Source/GameplayInsights/Public/IGameplayProvider.h
+++ b/Engine/Plugins/Animation/GameplayInsights/Source/GameplayInsights/Public/IGameplayProvider.h
@@ -111,6 +111,7 @@ public:
 	virtual bool ReadObjectPropertiesTimeline(uint64 InObjectId, TFunctionRef<void(const ObjectPropertiesTimeline&)> Callback) const = 0;
 	virtual void EnumerateObjectPropertyValues(uint64 InObjectId, const FObjectPropertiesMessage& InMessage, TFunctionRef<void(const FObjectPropertyValue&)> Callback) const = 0;
 	virtual void EnumerateObjects(TFunctionRef<void(const FObjectInfo&)> Callback) const = 0;
+	virtual void EnumerateObjects(double StartTime, double EndTime, TFunctionRef<void(const FObjectInfo&)> Callback) const = 0;
 	virtual const FClassInfo* FindClassInfo(uint64 InClassId) const = 0;
 	virtual const FClassInfo* FindClassInfo(const TCHAR* InClassPath) const = 0;
 	virtual const FObjectInfo* FindObjectInfo(uint64 InObjectId) const = 0;
diff --git a/Engine/Plugins/Animation/GameplayInsights/Source/RewindDebugger/Private/RewindDebugger.cpp b/Engine/Plugins/Animation/GameplayInsights/Source/RewindDebugger/Private/RewindDebugger.cpp
--- a/Engine/Plugins/Animation/GameplayInsights/Source/RewindDebugger/Private/RewindDebugger.cpp
+++ b/Engine/Plugins/Animation/GameplayInsights/Source/RewindDebugger/Private/RewindDebugger.cpp
@@ -189,16 +189,19 @@ bool FRewindDebugger::UpdateComponentList(uint64 ParentId, TArray<TSharedPtr<FDe
 
 	bool bChanged = false;
 
-	GameplayProvider->EnumerateObjects([this, &bChanged, ParentId, GameplayProvider, &ComponentList](const FObjectInfo& InObjectInfo)
-	{
-		// todo: filter components based on creation/destruction frame, and the current scrubbing time (so dynamically created and destroyed components won't add up)
+	double Time = CurrentTraceTime();
+
+	TArray<uint64, TInlineAllocator<32>> FoundObjects;
 
+	GameplayProvider->EnumerateObjects(Time, Time, [this, &bChanged, ParentId, GameplayProvider, &ComponentList, &FoundObjects](const FObjectInfo& InObjectInfo)
+	{
 		if (InObjectInfo.OuterId == ParentId)
 		{
 			const int32 FoundIndex = ComponentList.FindLastByPredicate([&InObjectInfo](const TSharedPtr<FDebugObjectInfo>& Info) { return Info->ObjectName == InObjectInfo.Name; });
 
 			if (FoundIndex >= 0 && ComponentList[FoundIndex]->ObjectName == InObjectInfo.Name)
 			{
+				// we need to make sure we reusing existing FDebugObjectInfo instances for the treeview selection to be stable
 				ComponentList[FoundIndex]->ObjectId = InObjectInfo.Id; // there is an issue with skeletal mesh components changing ids
 				bChanged = bChanged || UpdateComponentList(InObjectInfo.Id, ComponentList[FoundIndex]->Children);
 			}
@@ -208,6 +211,8 @@ bool FRewindDebugger::UpdateComponentList(uint64 ParentId, TArray<TSharedPtr<FDe
 				ComponentList.Add(MakeShared<FDebugObjectInfo>(InObjectInfo.Id,InObjectInfo.Name));
 				UpdateComponentList(InObjectInfo.Id, ComponentList.Last()->Children);
 			}
+			
+			FoundObjects.Add(InObjectInfo.Id);
 		}
 	});
 
@@ -230,6 +235,18 @@ bool FRewindDebugger::UpdateComponentList(uint64 ParentId, TArray<TSharedPtr<FDe
 				ComponentList.Add(MakeShared<FDebugObjectInfo>(ObjectInfo.Id,ObjectInfo.Name));
 				UpdateComponentList(ObjectInfo.Id, ComponentList.Last()->Children);
 			}
+			
+			FoundObjects.Add(ControllerId);
+		}
+	}
+	
+	// remove any components previously in the list that were not found in this time range.
+	for(int Index=ComponentList.Num()-1; Index>=0; Index--)
+	{
+		if (!FoundObjects.Contains(ComponentList[Index]->ObjectId))
+		{
+			bChanged = true;
+			ComponentList.RemoveAt(Index);
 		}
 	}
 	
@@ -255,7 +272,8 @@ uint64 FRewindDebugger::GetTargetActorId() const
 	{
 		if (const IGameplayProvider* GameplayProvider = Session->ReadProvider<IGameplayProvider>("GameplayProvider"))
 		{
-			GameplayProvider->EnumerateObjects([this,&TargetActorId](const FObjectInfo& InObjectInfo)
+			double Time = CurrentTraceTime();
+			GameplayProvider->EnumerateObjects(Time, Time, [this,&TargetActorId](const FObjectInfo& InObjectInfo)
 			{
 				if (DebugTargetActor.Get() == InObjectInfo.Name)
 				{
diff --git a/Engine/Source/Runtime/Engine/Private/Actor.cpp b/Engine/Source/Runtime/Engine/Private/Actor.cpp
--- a/Engine/Source/Runtime/Engine/Private/Actor.cpp
+++ b/Engine/Source/Runtime/Engine/Private/Actor.cpp
@@ -2499,7 +2499,7 @@ void AActor::EndPlay(const EEndPlayReason::Type EndPlayReason)
 {
 	if (ActorHasBegunPlay == EActorBeginPlayState::HasBegunPlay)
 	{
-		TRACE_OBJECT_EVENT(this, EndPlay);
+		TRACE_OBJECT_LIFETIME_END(this);
 
 		ActorHasBegunPlay = EActorBeginPlayState::HasNotBegunPlay;
 
@@ -3782,7 +3782,7 @@ void AActor::DispatchBeginPlay(bool bFromLevelStreaming)
 
 void AActor::BeginPlay()
 {
-	TRACE_OBJECT_EVENT(this, BeginPlay);
+	TRACE_OBJECT_LIFETIME_BEGIN(this);
 
 	ensureMsgf(ActorHasBegunPlay == EActorBeginPlayState::BeginningPlay, TEXT("BeginPlay was called on actor %s which was in state %d"), *GetPathName(), (int32)ActorHasBegunPlay);
 	SetLifeSpan( InitialLifeSpan );
diff --git a/Engine/Source/Runtime/Engine/Private/Animation/AnimInstance.cpp b/Engine/Source/Runtime/Engine/Private/Animation/AnimInstance.cpp
--- a/Engine/Source/Runtime/Engine/Private/Animation/AnimInstance.cpp
+++ b/Engine/Source/Runtime/Engine/Private/Animation/AnimInstance.cpp
@@ -196,13 +196,13 @@ UWorld* UAnimInstance::GetWorld() const
 
 void UAnimInstance::InitializeAnimation(bool bInDeferRootNodeInitialization)
 {
-	TRACE_OBJECT_EVENT(this, InitializeAnimation);
-
 	FScopeCycleCounterUObject ContextScope(this);
 	SCOPE_CYCLE_COUNTER(STAT_AnimInitTime);
 	LLM_SCOPE(ELLMTag::Animation);
 
 	UninitializeAnimation();
+	
+	TRACE_OBJECT_LIFETIME_BEGIN(this);
 
 	// make sure your skeleton is initialized
 	// you can overwrite different skeleton
@@ -253,8 +253,6 @@ void UAnimInstance::InitializeAnimation(bool bInDeferRootNodeInitialization)
 
 void UAnimInstance::UninitializeAnimation()
 {
-	TRACE_OBJECT_EVENT(this, UninitializeAnimation);
-
 	NativeUninitializeAnimation();
 
 	GetProxyOnGameThread<FAnimInstanceProxy>().Uninitialize(this);
@@ -312,6 +310,8 @@ void UAnimInstance::UninitializeAnimation()
 	NotifyQueue.Reset(SkelMeshComp);
 
 	SlotGroupInertializationRequestMap.Reset();
+	
+	TRACE_OBJECT_LIFETIME_END(this);
 }
 
 #if WITH_EDITORONLY_DATA
diff --git a/Engine/Source/Runtime/Engine/Private/Components/ActorComponent.cpp b/Engine/Source/Runtime/Engine/Private/Components/ActorComponent.cpp
--- a/Engine/Source/Runtime/Engine/Private/Components/ActorComponent.cpp
+++ b/Engine/Source/Runtime/Engine/Private/Components/ActorComponent.cpp
@@ -1022,7 +1022,7 @@ void UActorComponent::UninitializeComponent()
 
 void UActorComponent::BeginPlay()
 {
-	TRACE_OBJECT_EVENT(this, BeginPlay);
+	TRACE_OBJECT_LIFETIME_BEGIN(this);
 
 	check(bRegistered);
 	check(!bHasBegunPlay);
@@ -1038,7 +1038,7 @@ void UActorComponent::BeginPlay()
 
 void UActorComponent::EndPlay(const EEndPlayReason::Type EndPlayReason)
 {
-	TRACE_OBJECT_EVENT(this, EndPlay);
+	TRACE_OBJECT_LIFETIME_END(this);
 
 	check(bHasBegunPlay);
 
diff --git a/Engine/Source/Runtime/Engine/Private/ObjectTrace.cpp b/Engine/Source/Runtime/Engine/Private/ObjectTrace.cpp
--- a/Engine/Source/Runtime/Engine/Private/ObjectTrace.cpp
+++ b/Engine/Source/Runtime/Engine/Private/ObjectTrace.cpp
@@ -38,6 +38,16 @@ UE_TRACE_EVENT_BEGIN(Object, Object)
 	UE_TRACE_EVENT_FIELD(UE::Trace::WideString, Path)
 UE_TRACE_EVENT_END()
 
+UE_TRACE_EVENT_BEGIN(Object, ObjectLifetimeBegin)
+	UE_TRACE_EVENT_FIELD(uint64, Cycle)
+	UE_TRACE_EVENT_FIELD(uint64, Id)
+UE_TRACE_EVENT_END()
+
+UE_TRACE_EVENT_BEGIN(Object, ObjectLifetimeEnd)
+	UE_TRACE_EVENT_FIELD(uint64, Cycle)
+	UE_TRACE_EVENT_FIELD(uint64, Id)
+UE_TRACE_EVENT_END()
+
 UE_TRACE_EVENT_BEGIN(Object, ObjectEvent)
 	UE_TRACE_EVENT_FIELD(uint64, Cycle)
 	UE_TRACE_EVENT_FIELD(uint64, Id)
@@ -420,6 +430,56 @@ void FObjectTrace::OutputObjectEvent(const UObject* InObject, const TCHAR* InEve
 		<< ObjectEvent.Event(InEvent);
 }
 
+void FObjectTrace::OutputObjectLifetimeBegin(const UObject* InObject)
+{
+	const bool bChannelEnabled = UE_TRACE_CHANNELEXPR_IS_ENABLED(ObjectChannel);
+	if (!bChannelEnabled || InObject == nullptr)
+	{
+		return;
+	}
+
+	if(InObject->HasAnyFlags(RF_ClassDefaultObject))
+	{
+		return;
+	}
+
+	if (CANNOT_TRACE_OBJECT(InObject->GetWorld()))
+	{
+		return;
+	}
+
+	TRACE_OBJECT(InObject);
+
+	UE_TRACE_LOG(Object, ObjectLifetimeBegin, ObjectChannel)
+		<< ObjectLifetimeBegin.Cycle(FPlatformTime::Cycles64())
+		<< ObjectLifetimeBegin.Id(GetObjectId(InObject));
+}
+
+void FObjectTrace::OutputObjectLifetimeEnd(const UObject* InObject)
+{
+	const bool bChannelEnabled = UE_TRACE_CHANNELEXPR_IS_ENABLED(ObjectChannel);
+	if (!bChannelEnabled || InObject == nullptr)
+	{
+		return;
+	}
+
+	if(InObject->HasAnyFlags(RF_ClassDefaultObject))
+	{
+		return;
+	}
+
+	if (CANNOT_TRACE_OBJECT(InObject->GetWorld()))
+	{
+		return;
+	}
+
+	TRACE_OBJECT(InObject);
+
+	UE_TRACE_LOG(Object, ObjectLifetimeEnd, ObjectChannel)
+		<< ObjectLifetimeEnd.Cycle(FPlatformTime::Cycles64())
+		<< ObjectLifetimeEnd.Id(GetObjectId(InObject));
+}
+
 void FObjectTrace::OutputPawnPossess(const UObject* InController, const UObject* InPawn)
 {
 	bool bChannelEnabled = UE_TRACE_CHANNELEXPR_IS_ENABLED(ObjectChannel);
diff --git a/Engine/Source/Runtime/Engine/Public/ObjectTrace.h b/Engine/Source/Runtime/Engine/Public/ObjectTrace.h
--- a/Engine/Source/Runtime/Engine/Public/ObjectTrace.h
+++ b/Engine/Source/Runtime/Engine/Public/ObjectTrace.h
@@ -54,6 +54,12 @@ struct FObjectTrace
 
 	/** Helper function to output an object */
 	ENGINE_API static void OutputObject(const UObject* InObject);
+	
+	/** Helper function to output object create event */
+	ENGINE_API static void OutputObjectLifetimeBegin(const UObject* InObject);
+	
+	/** Helper function to output object create event */
+	ENGINE_API static void OutputObjectLifetimeEnd(const UObject* InObject);
 
 	/** Helper function to output camera information for a player */
 	ENGINE_API static void OutputView(const UObject* LocalPlayer, const FSceneView* View);
@@ -109,6 +115,12 @@ struct FObjectTrace
 
 #define TRACE_OBJECT_EVENT(Object, Event) \
 	if(CAN_TRACE_OBJECT(Object)) { UNCONDITIONAL_TRACE_OBJECT_EVENT(Object, Event); }
+	
+#define TRACE_OBJECT_LIFETIME_BEGIN(Object) \
+	if(CAN_TRACE_OBJECT(Object)) { FObjectTrace::OutputObjectLifetimeBegin(Object); }
+
+#define TRACE_OBJECT_LIFETIME_END(Object) \
+	if(CAN_TRACE_OBJECT(Object)) { FObjectTrace::OutputObjectLifetimeEnd(Object); }
 
 #define TRACE_PAWN_POSSESS(Controller, Pawn)\
 	if(CAN_TRACE_OBJECT(Controller) && (Pawn==nullptr || (CAN_TRACE_OBJECT(Pawn)))) FObjectTrace::OutputPawnPossess(Controller, Pawn);
@@ -117,6 +129,12 @@ struct FObjectTrace
 
 #define TRACE_OBJECT_EVENT(Object, Event) \
 	UNCONDITIONAL_TRACE_OBJECT_EVENT(Object, Event);
+
+#define TRACE_OBJECT_LIFETIME_BEGIN(Object) \
+	FObjectTrace::OutputObjectLifetimeBegin(Object);
+	
+#define TRACE_OBJECT_LIFETIME_END(Object) \
+	FObjectTrace::OutputObjectLifetimeEnd(Object);
 	
 #define TRACE_PAWN_POSSESS(Controller, Pawn)\
 	FObjectTrace::OutputPawnPossess(Controller, Pawn);
@@ -138,5 +156,7 @@ struct FObjectTrace
 #define TRACE_WORLD(World)
 #define TRACE_PAWN_POSSESS(Controller, Pawn)
 #define TRACE_VIEW(Player, View)
+#define TRACE_OBJECT_LIFETIME_BEGIN(Object)
+#define TRACE_OBJECT_LIFETIME_END(Object)
 
 #endif
\ No newline at end of file
