fix log spam when using rewind debugger on physics meshes

diff --git a/Engine/Plugins/Animation/GameplayInsights/Source/RewindDebugger/Private/RewindDebugger.cpp b/Engine/Plugins/Animation/GameplayInsights/Source/RewindDebugger/Private/RewindDebugger.cpp
--- a/Engine/Plugins/Animation/GameplayInsights/Source/RewindDebugger/Private/RewindDebugger.cpp
+++ b/Engine/Plugins/Animation/GameplayInsights/Source/RewindDebugger/Private/RewindDebugger.cpp
@@ -749,7 +749,7 @@ void FRewindDebugger::Tick(float DeltaTime)
 												}
 
 												// todo: we need to take into account tick order requirements for attached objects here
-												MeshComponent->SetWorldTransform(ComponentWorldTransform);
+												MeshComponent->SetWorldTransform(ComponentWorldTransform, false, nullptr, ETeleportType::TeleportPhysics);
 												MeshComponent->SetForcedLOD(PoseMessage.LodIndex + 1);
 												return TraceServices::EEventEnumerate::Stop;
 											});
