disable rewind debugger visual logger integration for release-5.0

diff --git a/Engine/Plugins/Animation/GameplayInsights/Source/RewindDebuggerVLog/Private/RewindDebuggerVLogModule.cpp b/Engine/Plugins/Animation/GameplayInsights/Source/RewindDebuggerVLog/Private/RewindDebuggerVLogModule.cpp
--- a/Engine/Plugins/Animation/GameplayInsights/Source/RewindDebuggerVLog/Private/RewindDebuggerVLogModule.cpp
+++ b/Engine/Plugins/Animation/GameplayInsights/Source/RewindDebuggerVLog/Private/RewindDebuggerVLogModule.cpp
@@ -6,16 +6,24 @@
 
 #define LOCTEXT_NAMESPACE "RewindDebuggerVLogModule"
 
+#ifndef ENABLE_REWINDDEBUGGER_VLOG_INTEGRATION
+#define ENABLE_REWINDDEBUGGER_VLOG_INTEGRATION 0
+#endif
+
 void FRewindDebuggerVLogModule::StartupModule()
 {
+#if ENABLE_REWINDDEBUGGER_VLOG_INTEGRATION
 	IModularFeatures::Get().RegisterModularFeature(IRewindDebuggerExtension::ModularFeatureName, &RewindDebuggerVLogExtension);
 	IModularFeatures::Get().RegisterModularFeature(TraceServices::ModuleFeatureName, &VLogTraceModule);
+#endif
 }
 
 void FRewindDebuggerVLogModule::ShutdownModule()
 {
+#if ENABLE_REWINDDEBUGGER_VLOG_INTEGRATION
 	IModularFeatures::Get().UnregisterModularFeature(IRewindDebuggerExtension::ModularFeatureName, &RewindDebuggerVLogExtension);
 	IModularFeatures::Get().UnregisterModularFeature(TraceServices::ModuleFeatureName, &VLogTraceModule);
+#endif
 }
 
 IMPLEMENT_MODULE(FRewindDebuggerVLogModule, RewindDebuggerVLog);
