Fixes for rewind debugger tab layouts crashing and other bad behavior.

diff --git a/Engine/Plugins/Animation/GameplayInsights/Source/RewindDebugger/Private/RewindDebuggerModule.cpp b/Engine/Plugins/Animation/GameplayInsights/Source/RewindDebugger/Private/RewindDebuggerModule.cpp
--- a/Engine/Plugins/Animation/GameplayInsights/Source/RewindDebugger/Private/RewindDebuggerModule.cpp
+++ b/Engine/Plugins/Animation/GameplayInsights/Source/RewindDebugger/Private/RewindDebuggerModule.cpp
@@ -31,7 +31,13 @@ TSharedRef<SDockTab> FRewindDebuggerModule::SpawnRewindDebuggerTab(const FSpawnT
 	}
 	
 	const TSharedRef<SDockTab> MajorTab = SNew(SDockTab)
-		.TabRole(ETabRole::NomadTab);
+		.TabRole(ETabRole::NomadTab)
+		.OnTabClosed_Lambda([this](TSharedRef<SDockTab>)
+		{
+			// clear reference to widget so it will be destroyed
+			RewindDebuggerWidget->CloseAllTabs();
+			RewindDebuggerWidget = nullptr;
+		});
 
 	TSharedPtr<SWidget> TabContent;
 
diff --git a/Engine/Plugins/Animation/GameplayInsights/Source/RewindDebugger/Private/SRewindDebugger.cpp b/Engine/Plugins/Animation/GameplayInsights/Source/RewindDebugger/Private/SRewindDebugger.cpp
--- a/Engine/Plugins/Animation/GameplayInsights/Source/RewindDebugger/Private/SRewindDebugger.cpp
+++ b/Engine/Plugins/Animation/GameplayInsights/Source/RewindDebugger/Private/SRewindDebugger.cpp
@@ -28,6 +28,7 @@
 #include "IRewindDebuggerViewCreator.h"
 #include "RewindDebuggerViewCreators.h"
 #include "IRewindDebuggerDoubleClickHandler.h"
+#include "Interfaces/IMainFrameModule.h"
 
 #define LOCTEXT_NAMESPACE "SRewindDebugger"
 
@@ -183,8 +184,26 @@ TSharedRef<SWidget> SRewindDebugger::MakeSelectActorMenu()
 	return MenuBuilder.MakeWidget();
 }
 
+void SRewindDebugger::CloseAllTabs()
+{
+	 for(FName& TabName : TabNames)
+	 {
+		CloseTab(TabName);
+	 }
+}
+
+void SRewindDebugger::MainFrameCreationFinished(TSharedPtr<SWindow> InRootWindow, bool bIsNewProjectWindow)
+{
+	IMainFrameModule::Get().OnMainFrameCreationFinished().RemoveAll(this);
+	
+	 // close all tabs that may be open from restoring the saved layout config
+	CloseAllTabs();
+}
+
 void SRewindDebugger::Construct(const FArguments& InArgs, TSharedRef<FUICommandList> CommandList, const TSharedRef<SDockTab>& ConstructUnderMajorTab, const TSharedPtr<SWindow>& ConstructUnderWindow)
 {
+	bInitializing = true;
+	
 	OnScrubPositionChanged = InArgs._OnScrubPositionChanged;
 	OnComponentSelectionChanged = InArgs._OnComponentSelectionChanged;
 	BuildComponentContextMenu = InArgs._BuildComponentContextMenu;
@@ -403,11 +422,18 @@ void SRewindDebugger::Construct(const FArguments& InArgs, TSharedRef<FUICommandL
 		]
 	];
 
-	// close all tabs that may be open from restoring the saved layout config
-	for(FName& TabName : TabNames)
+	if (IMainFrameModule::Get().IsWindowInitialized())
 	{
-		CloseTab(TabName);
+		// close all tabs that may be open from restoring the saved layout config
+		CloseAllTabs();
 	}
+	else
+	{
+		// close them later if we are initalizing the layout, to avoid issues with empty windows and crashes
+		IMainFrameModule::Get().OnMainFrameCreationFinished().AddRaw(this, &SRewindDebugger::MainFrameCreationFinished);
+	}
+	
+	bInitializing = false;
 }
 
 void SRewindDebugger::RefreshDebugComponents()
@@ -454,8 +480,13 @@ TSharedRef<SDockTab> SRewindDebugger::SpawnTab(const FSpawnTabArgs& Args, FName
 // returns true if DebugViews contains a view for the ViewName, but there is no matching Pinned view already open
 bool SRewindDebugger::CanSpawnTab(const FSpawnTabArgs& Args, FName ViewName)
 {
+	if (bInitializing)
+	{
+		return true;
+	}
+	
 	TSharedPtr<IRewindDebuggerView>* View = DebugViews.FindByPredicate([ViewName](TSharedPtr<IRewindDebuggerView>& View) { return View->GetName() == ViewName; } );
-
+	
 	if (View!=nullptr)
 	{
 		bool bPinned = nullptr != PinnedDebugViews.FindByPredicate(
@@ -464,7 +495,7 @@ bool SRewindDebugger::CanSpawnTab(const FSpawnTabArgs& Args, FName ViewName)
 				return PinnedView->GetName() == (*View)->GetName() && PinnedView->GetObjectId() == (*View)->GetObjectId();
 			}
 		);
-
+	
 		return !bPinned;
 	}
 	return false;
diff --git a/Engine/Plugins/Animation/GameplayInsights/Source/RewindDebugger/Private/SRewindDebugger.h b/Engine/Plugins/Animation/GameplayInsights/Source/RewindDebugger/Private/SRewindDebugger.h
--- a/Engine/Plugins/Animation/GameplayInsights/Source/RewindDebugger/Private/SRewindDebugger.h
+++ b/Engine/Plugins/Animation/GameplayInsights/Source/RewindDebugger/Private/SRewindDebugger.h
@@ -57,7 +57,10 @@ public:
 
 	void TrackCursor(bool bReverse);
 	void RefreshDebugComponents();
+	void CloseAllTabs();
 private:
+	void MainFrameCreationFinished(TSharedPtr<SWindow> InRootWindow, bool bIsNewProjectWindow);
+	
 	// Time Slider
 	TAttribute<double> ScrubTimeAttribute;
 	TAttribute<bool> TrackScrubbingAttribute;
@@ -104,6 +107,7 @@ private:
 	TArray<FName> TabNames;
 	TArray<FName> HiddenTabs;  // keep track of tabs that have been closed so we don't automatically reopen them when switching components
 	bool bInternalClosingTab = false;
+	bool bInitializing = false;
 
 	TSharedPtr<FTabManager> TabManager;
 };
diff --git a/Engine/Plugins/Animation/GameplayInsights/Source/RewindDebugger/RewindDebugger.Build.cs b/Engine/Plugins/Animation/GameplayInsights/Source/RewindDebugger/RewindDebugger.Build.cs
--- a/Engine/Plugins/Animation/GameplayInsights/Source/RewindDebugger/RewindDebugger.Build.cs
+++ b/Engine/Plugins/Animation/GameplayInsights/Source/RewindDebugger/RewindDebugger.Build.cs
@@ -37,6 +37,7 @@ namespace UnrealBuildTool.Rules
 				"UnrealEd",
 				"RewindDebuggerInterface",
 				"ToolWidgets",
+				"MainFrame"
 			});
 		}
 	}
