Fix for crash in Rewind Debugger if tabs are undocked in a loaded layout, and the main tab is not docked in the level editor

diff --git a/Engine/Plugins/Animation/GameplayInsights/Source/RewindDebugger/Private/SRewindDebugger.cpp b/Engine/Plugins/Animation/GameplayInsights/Source/RewindDebugger/Private/SRewindDebugger.cpp
--- a/Engine/Plugins/Animation/GameplayInsights/Source/RewindDebugger/Private/SRewindDebugger.cpp
+++ b/Engine/Plugins/Animation/GameplayInsights/Source/RewindDebugger/Private/SRewindDebugger.cpp
@@ -417,7 +417,7 @@ void SRewindDebugger::Construct(const FArguments& InArgs, TSharedRef<FUICommandL
 			]
 			+SVerticalBox::Slot().FillHeight(1.0f)
 			[
-				TabManager->RestoreFrom(Layout, ConstructUnderWindow).ToSharedRef()
+				TabManager->RestoreFrom(Layout,TSharedPtr<SWindow>()).ToSharedRef()
 			]
 		]
 	];
