Fix ensure when scrubbing mesh poses in Animation Insights

diff --git a/Engine/Plugins/Animation/GameplayInsights/Source/GameplayInsightsEditor/Private/InsightsSkeletalMeshComponent.cpp b/Engine/Plugins/Animation/GameplayInsights/Source/GameplayInsightsEditor/Private/InsightsSkeletalMeshComponent.cpp
--- a/Engine/Plugins/Animation/GameplayInsights/Source/GameplayInsightsEditor/Private/InsightsSkeletalMeshComponent.cpp
+++ b/Engine/Plugins/Animation/GameplayInsights/Source/GameplayInsightsEditor/Private/InsightsSkeletalMeshComponent.cpp
@@ -42,5 +42,7 @@ void UInsightsSkeletalMeshComponent::InitAnim(bool bForceReInit)
 				BoneVisibilityStates[BaseIndex][BoneIndex] = BVS_Visible;
 			}
 		}
+
+		MorphTargetWeights.SetNumZeroed(SkeletalMesh->GetMorphTargets().Num());
 	}
 }
\ No newline at end of file
