Rewind Debugger timelines

diff --git a/Engine/Plugins/Animation/GameplayInsights/Source/GameplayInsights/GameplayInsights.Build.cs b/Engine/Plugins/Animation/GameplayInsights/Source/GameplayInsights/GameplayInsights.Build.cs
--- a/Engine/Plugins/Animation/GameplayInsights/Source/GameplayInsights/GameplayInsights.Build.cs
+++ b/Engine/Plugins/Animation/GameplayInsights/Source/GameplayInsights/GameplayInsights.Build.cs
@@ -20,7 +20,8 @@ namespace UnrealBuildTool.Rules
 				"TraceInsights",
 				"AssetRegistry",
 				"ApplicationCore",
-				"RewindDebuggerInterface"
+				"RewindDebuggerInterface",
+				"ToolWidgets"
 			});
 
 			if (Target.bCompileAgainstEngine)
diff --git a/Engine/Plugins/Animation/GameplayInsights/Source/GameplayInsights/Private/AnimCurvesTrack.cpp b/Engine/Plugins/Animation/GameplayInsights/Source/GameplayInsights/Private/AnimCurvesTrack.cpp
new file mode 100644
--- /dev/null
+++ b/Engine/Plugins/Animation/GameplayInsights/Source/GameplayInsights/Private/AnimCurvesTrack.cpp
@@ -0,0 +1,234 @@
+// Copyright Epic Games, Inc. All Rights Reserved.
+
+#include "AnimCurvesTrack.h"
+#include "SCurveTimelineView.h"
+#include "IRewindDebugger.h"
+#include "GameplayProvider.h"
+#include "AnimationProvider.h"
+#include "SAnimationCurvesView.h"
+
+#define LOCTEXT_NAMESPACE "AnimCurvesTrack"
+
+namespace RewindDebugger
+{
+
+FAnimCurvesTrack::FAnimCurvesTrack(uint64 InObjectId) : ObjectId(InObjectId)
+{
+	Icon = FSlateIcon("EditorStyle", "AnimGraph.Attribute.Curves.Icon", "AnimGraph.Attribute.Curves.Icon");
+}
+
+TSharedPtr<SWidget> FAnimCurvesTrack::GetDetailsViewInternal()
+{
+	IRewindDebugger* RewindDebugger = IRewindDebugger::Instance();
+	return SNew(SAnimationCurvesView, ObjectId, RewindDebugger->CurrentTraceTime(), *RewindDebugger->GetAnalysisSession())
+				.CurrentTime_Lambda([RewindDebugger]{ return RewindDebugger->CurrentTraceTime(); });
+	
+}
+
+void FAnimCurvesTrack::IterateSubTracksInternal(TFunction<void(TSharedPtr<FRewindDebuggerTrack> SubTrack)> IteratorFunction)
+{
+	for(TSharedPtr<FAnimCurveTrack>& Track : Children)
+	{
+		IteratorFunction(Track);
+	}
+};
+
+bool FAnimCurvesTrack::UpdateInternal()
+{
+	TArray<uint32> UniqueTrackIds;
+
+	IRewindDebugger* RewindDebugger = IRewindDebugger::Instance();
+	const TraceServices::IAnalysisSession* AnalysisSession = RewindDebugger->GetAnalysisSession();
+
+	bool bChanged = false;
+	
+	// convert time range to from rewind debugger times to profiler times
+	TRange<double> TraceTimeRange = RewindDebugger->GetCurrentTraceRange();
+	double StartTime = TraceTimeRange.GetLowerBoundValue();
+	double EndTime = TraceTimeRange.GetUpperBoundValue();
+	
+	// count number of unique animations in the current time range
+	const FGameplayProvider* GameplayProvider = AnalysisSession->ReadProvider<FGameplayProvider>(FGameplayProvider::ProviderName);
+	const FAnimationProvider* AnimationProvider = AnalysisSession->ReadProvider<FAnimationProvider>(FAnimationProvider::ProviderName);
+
+	int CurveCount = 0;
+
+	if(GameplayProvider && AnimationProvider)
+	{
+		TraceServices::FAnalysisSessionReadScope SessionReadScope(*AnalysisSession);
+		UniqueTrackIds.SetNum(0, false);
+
+		AnimationProvider->ReadSkeletalMeshPoseTimeline(ObjectId, [&UniqueTrackIds,AnimationProvider, StartTime, EndTime](const FAnimationProvider::SkeletalMeshPoseTimeline& InTimeline, bool bHasCurves)
+		{
+			// this isn't very efficient, and it gets called every frame.  will need optimizing
+			InTimeline.EnumerateEvents(StartTime, EndTime, [&UniqueTrackIds, StartTime, EndTime, AnimationProvider](double InStartTime, double InEndTime, uint32 InDepth, const FSkeletalMeshPoseMessage& InMessage)
+			{
+				if (InEndTime > StartTime && InStartTime < EndTime)
+				{
+					AnimationProvider->EnumerateSkeletalMeshCurves(InMessage, [&UniqueTrackIds](const FSkeletalMeshNamedCurve& InCurve)
+					{
+						UniqueTrackIds.AddUnique(InCurve.Id);
+					});
+				}
+				return TraceServices::EEventEnumerate::Continue;
+			});
+		});
+
+		UniqueTrackIds.StableSort();
+		
+		CurveCount = UniqueTrackIds.Num();
+
+		if (Children.Num()!=CurveCount)
+			bChanged = true;
+		
+		Children.SetNum(CurveCount);
+		for(int i = 0; i < CurveCount; i++)
+		{
+			if (!Children[i].IsValid() || static_cast<FAnimCurveTrack*>(Children[i].Get())->GetCurveId() != UniqueTrackIds[i])
+			{
+				Children[i] = MakeShared<FAnimCurveTrack>(ObjectId, UniqueTrackIds[i]);
+				bChanged = true;
+			}
+
+			bChanged = bChanged || Children[i]->Update();
+		}
+	}
+
+	return bChanged;
+}
+
+
+FAnimCurveTrack::FAnimCurveTrack(uint64 InObjectId, uint32 InCurveId) :
+	CurveId(InCurveId),
+	ObjectId(InObjectId)
+{
+	Icon = FSlateIcon("EditorStyle", "AnimGraph.Attribute.Curves.Icon", "AnimGraph.Attribute.Curves.Icon");
+}
+
+TSharedPtr<SCurveTimelineView::FTimelineCurveData> FAnimCurveTrack::GetCurveData() const
+{
+	if (!CurveData.IsValid())
+	{
+		CurveData = MakeShared<SCurveTimelineView::FTimelineCurveData>();
+	}
+	
+	CurvesUpdateRequested++;
+	
+	return CurveData;
+}
+
+bool FAnimCurveTrack::UpdateInternal()
+{
+	IRewindDebugger* RewindDebugger = IRewindDebugger::Instance();
+	// compute curve points
+	//
+	// convert time range to from rewind debugger times to profiler times
+	TRange<double> TraceTimeRange = RewindDebugger->GetCurrentTraceRange();
+	double StartTime = TraceTimeRange.GetLowerBoundValue();
+	double EndTime = TraceTimeRange.GetUpperBoundValue();
+	
+	const TraceServices::IAnalysisSession* AnalysisSession = RewindDebugger->GetAnalysisSession();
+	const FGameplayProvider* GameplayProvider = AnalysisSession->ReadProvider<FGameplayProvider>(FGameplayProvider::ProviderName);
+	const FAnimationProvider* AnimationProvider = AnalysisSession->ReadProvider<FAnimationProvider>(FAnimationProvider::ProviderName);
+
+	
+	if(CurvesUpdateRequested > 10 && GameplayProvider && AnimationProvider)
+	{
+		CurvesUpdateRequested = 0;
+		auto& CurvePoints = CurveData->Points;
+		CurvePoints.SetNum(0,false);
+	
+		TraceServices::FAnalysisSessionReadScope SessionReadScope(*AnalysisSession);
+
+		AnimationProvider->ReadSkeletalMeshPoseTimeline(ObjectId, [this, AnimationProvider, StartTime, EndTime, AnalysisSession, &CurvePoints](const FAnimationProvider::SkeletalMeshPoseTimeline& InTimeline, bool bHasCurves)
+		{
+			// this isn't very efficient, and it gets called every frame.  will need optimizing
+			InTimeline.EnumerateEvents(StartTime, EndTime, [this, StartTime, EndTime, AnimationProvider, AnalysisSession, &CurvePoints](double InStartTime, double InEndTime, uint32 InDepth, const FSkeletalMeshPoseMessage& InMessage)
+			{
+				if (InEndTime > StartTime && InStartTime < EndTime)
+				{
+					double Time = InMessage.RecordingTime;
+					
+					AnimationProvider->EnumerateSkeletalMeshCurves(InMessage, [this, Time, &CurvePoints](const FSkeletalMeshNamedCurve& InCurve)
+					{
+						if (InCurve.Id == CurveId)
+						{
+							CurvePoints.Add({Time,InCurve.Value});
+						}
+					});
+				}
+				return TraceServices::EEventEnumerate::Continue;
+			});
+		});
+	}
+		
+	bool bChanged = false;
+
+	if (CurveName.IsEmpty())
+	{
+		CurveName = FText::FromString(AnimationProvider->GetName(CurveId));
+		bChanged = true;
+	}
+
+	return bChanged;
+}
+
+TSharedPtr<SWidget> FAnimCurveTrack::GetDetailsViewInternal() 
+{
+	IRewindDebugger* RewindDebugger = IRewindDebugger::Instance();
+	TSharedPtr<SAnimationCurvesView> AnimationCurvesView = SNew(SAnimationCurvesView, ObjectId, RewindDebugger->CurrentTraceTime(), *RewindDebugger->GetAnalysisSession())
+				.CurrentTime_Lambda([RewindDebugger]{ return RewindDebugger->CurrentTraceTime(); });
+	AnimationCurvesView->SetCurveFilter(CurveId);
+	return AnimationCurvesView;
+}
+
+TSharedPtr<SWidget> FAnimCurveTrack::GetTimelineViewInternal()
+{
+	FLinearColor CurveColor(0.5,0.5,0.5);
+	
+	return SNew(SCurveTimelineView)
+		.CurveColor(CurveColor)
+		.ViewRange_Lambda([]() { return IRewindDebugger::Instance()->GetCurrentViewRange(); })
+		.RenderFill(false)
+		.CurveData_Raw(this, &FAnimCurveTrack::GetCurveData);
+}
+
+
+static const FName AnimationCurvesName("AnimationCurves");
+
+FName FAnimationCurvesTrackCreator::GetTargetTypeNameInternal() const
+{
+	static const FName SkelMeshName("SkeletalMeshComponent");
+	return SkelMeshName;
+}
+	
+FName FAnimationCurvesTrackCreator::GetNameInternal() const
+{
+	return AnimationCurvesName;
+}
+
+TSharedPtr<RewindDebugger::FRewindDebuggerTrack> FAnimationCurvesTrackCreator::CreateTrackInternal(uint64 ObjectId) const
+{
+	return MakeShared<RewindDebugger::FAnimCurvesTrack>(ObjectId);
+}
+
+bool FAnimationCurvesTrackCreator::HasDebugInfoInternal(uint64 ObjectId) const
+{
+	const TraceServices::IAnalysisSession* AnalysisSession = IRewindDebugger::Instance()->GetAnalysisSession();
+	
+	TraceServices::FAnalysisSessionReadScope SessionReadScope(*AnalysisSession);
+	bool bHasData = false;
+	if (const FAnimationProvider* AnimationProvider = AnalysisSession->ReadProvider<FAnimationProvider>(FAnimationProvider::ProviderName))
+	{
+		AnimationProvider->ReadSkeletalMeshPoseTimeline(ObjectId, [&bHasData](const FAnimationProvider::SkeletalMeshPoseTimeline& InTimeline, bool bInHasCurves)
+		{
+			bHasData = bInHasCurves;
+		});
+	}
+	return bHasData;
+}
+
+	
+}
+
+#undef LOCTEXT_NAMESPACE
\ No newline at end of file
diff --git a/Engine/Plugins/Animation/GameplayInsights/Source/GameplayInsights/Private/AnimCurvesTrack.h b/Engine/Plugins/Animation/GameplayInsights/Source/GameplayInsights/Private/AnimCurvesTrack.h
new file mode 100644
--- /dev/null
+++ b/Engine/Plugins/Animation/GameplayInsights/Source/GameplayInsights/Private/AnimCurvesTrack.h
@@ -0,0 +1,76 @@
+// Copyright Epic Games, Inc. All Rights Reserved.
+#pragma once
+
+#include "IAnimationProvider.h"
+#include "IRewindDebuggerTrackCreator.h"
+#include "RewindDebuggerTrack.h"
+#include "SCurveTimelineView.h"
+#include "Textures/SlateIcon.h"
+
+namespace RewindDebugger
+{
+
+// Sub Track for a single curve
+class FAnimCurveTrack : public FRewindDebuggerTrack
+{
+public:
+	FAnimCurveTrack(uint64 InObjectId, uint32 InCurveId);
+
+
+	uint32 GetCurveId() { return CurveId; };
+
+private:
+	virtual bool UpdateInternal() override;
+	TSharedPtr<SWidget> GetDetailsViewInternal() override;
+	virtual TSharedPtr<SWidget> GetTimelineViewInternal() override;
+	virtual FSlateIcon GetIconInternal() override { return Icon; }
+	virtual FName GetNameInternal() const override { return "AnimationCurve"; }
+	virtual FText GetDisplayNameInternal() const override { return CurveName; }
+	virtual uint64 GetObjectIdInternal() const override { return ObjectId; }
+
+	TSharedPtr<SCurveTimelineView::FTimelineCurveData> GetCurveData() const;
+
+	uint32 CurveId;
+
+	FSlateIcon Icon;
+	FText CurveName;
+	uint64 ObjectId;
+
+	mutable TSharedPtr<SCurveTimelineView::FTimelineCurveData> CurveData;
+	mutable int CurvesUpdateRequested = 0;
+};
+
+// Parent track of all the curves
+class FAnimCurvesTrack : public FRewindDebuggerTrack
+{
+public:
+	FAnimCurvesTrack(uint64 InObjectId);
+
+
+private:
+	virtual bool UpdateInternal() override;
+	TSharedPtr<SWidget> GetDetailsViewInternal() override;
+	
+	virtual FSlateIcon GetIconInternal() override { return Icon; }
+	virtual FName GetNameInternal() const override { return "AnimationCurves"; }
+	virtual FText GetDisplayNameInternal() const override { return NSLOCTEXT("RewindDebugger", "AnimationCurvesTrackName", "Animation Curves"); }
+	virtual uint64 GetObjectIdInternal() const override { return ObjectId; }
+	virtual void IterateSubTracksInternal(TFunction<void(TSharedPtr<FRewindDebuggerTrack> SubTrack)> IteratorFunction) override;
+	
+	FSlateIcon Icon;
+	uint64 ObjectId;
+
+	TArray<TSharedPtr<FAnimCurveTrack>> Children;
+};
+
+	
+class FAnimationCurvesTrackCreator : public IRewindDebuggerTrackCreator
+{
+private:
+	virtual FName GetTargetTypeNameInternal() const;
+	virtual FName GetNameInternal() const override;
+	virtual TSharedPtr<FRewindDebuggerTrack> CreateTrackInternal(uint64 ObjectId) const override;
+	virtual bool HasDebugInfoInternal(uint64 ObjectId) const override;
+};
+
+}
\ No newline at end of file
diff --git a/Engine/Plugins/Animation/GameplayInsights/Source/GameplayInsights/Private/AnimationAnalyzer.cpp b/Engine/Plugins/Animation/GameplayInsights/Source/GameplayInsights/Private/AnimationAnalyzer.cpp
--- a/Engine/Plugins/Animation/GameplayInsights/Source/GameplayInsights/Private/AnimationAnalyzer.cpp
+++ b/Engine/Plugins/Animation/GameplayInsights/Source/GameplayInsights/Private/AnimationAnalyzer.cpp
@@ -17,10 +17,12 @@ void FAnimationAnalyzer::OnAnalysisBegin(const FOnAnalysisContext& Context)
 	auto& Builder = Context.InterfaceBuilder;
 
 	Builder.RouteEvent(RouteId_TickRecord, "Animation", "TickRecord");
+	Builder.RouteEvent(RouteId_TickRecord2, "Animation", "TickRecord2");
 	Builder.RouteEvent(RouteId_SkeletalMesh, "Animation", "SkeletalMesh");
 	Builder.RouteEvent(RouteId_SkeletalMesh2, "Animation", "SkeletalMesh2");
 	Builder.RouteEvent(RouteId_SkeletalMeshComponent, "Animation", "SkeletalMeshComponent");
 	Builder.RouteEvent(RouteId_SkeletalMeshComponent2, "Animation", "SkeletalMeshComponent2");
+	Builder.RouteEvent(RouteId_SkeletalMeshComponent3, "Animation", "SkeletalMeshComponent3");
 	Builder.RouteEvent(RouteId_SkeletalMeshFrame, "Animation", "SkeletalMeshFrame");
 	Builder.RouteEvent(RouteId_AnimGraph, "Animation", "AnimGraph");
 	Builder.RouteEvent(RouteId_AnimNodeStart, "Animation", "AnimNodeStart");
@@ -40,6 +42,7 @@ void FAnimationAnalyzer::OnAnalysisBegin(const FOnAnalysisContext& Context)
 	Builder.RouteEvent(RouteId_Notify, "Animation", "Notify");
 	Builder.RouteEvent(RouteId_SyncMarker, "Animation", "SyncMarker");
 	Builder.RouteEvent(RouteId_Montage, "Animation", "Montage");
+	Builder.RouteEvent(RouteId_Montage2, "Animation", "Montage2");
 	Builder.RouteEvent(RouteId_Sync, "Animation", "Sync");
 }
 
@@ -52,6 +55,28 @@ bool FAnimationAnalyzer::OnEvent(uint16 RouteId, EStyle Style, const FOnEventCon
 	const auto& EventData = Context.EventData;
 	switch (RouteId)
 	{
+	case RouteId_TickRecord2:
+	{
+		double RecordingTime = EventData.GetValue<double>("RecordingTime");
+		uint64 Cycle = EventData.GetValue<uint64>("Cycle");
+		uint64 AnimInstanceId = EventData.GetValue<uint64>("AnimInstanceId");
+		uint64 AssetId = EventData.GetValue<uint64>("AssetId");
+		int32 NodeId = EventData.GetValue<int32>("NodeId");
+		float BlendWeight = EventData.GetValue<float>("BlendWeight");
+		float PlaybackTime = EventData.GetValue<float>("PlaybackTime");
+		float RootMotionWeight = EventData.GetValue<float>("RootMotionWeight");
+		float PlayRate = EventData.GetValue<float>("PlayRate");
+		float BlendSpacePositionX = EventData.GetValue<float>("BlendSpacePositionX");
+		float BlendSpacePositionY = EventData.GetValue<float>("BlendSpacePositionY");
+		float BlendSpaceFilteredPositionX = EventData.GetValue<float>("BlendSpaceFilteredPositionX");
+		float BlendSpaceFilteredPositionY = EventData.GetValue<float>("BlendSpaceFilteredPositionY");
+		uint16 FrameCounter = EventData.GetValue<uint16>("FrameCounter");
+		bool bLooping = EventData.GetValue<bool>("Looping");
+		bool bIsBlendSpace = EventData.GetValue<bool>("IsBlendSpace");
+		AnimationProvider.AppendTickRecord(AnimInstanceId, Context.EventTime.AsSeconds(Cycle), RecordingTime, AssetId, NodeId, BlendWeight, PlaybackTime, RootMotionWeight, PlayRate, BlendSpacePositionX, BlendSpacePositionY, BlendSpaceFilteredPositionX, BlendSpaceFilteredPositionY, FrameCounter, bLooping, bIsBlendSpace);
+		break;
+	}
+        	
 	case RouteId_TickRecord:
 	{
 		uint64 Cycle = EventData.GetValue<uint64>("Cycle");
@@ -69,7 +94,7 @@ bool FAnimationAnalyzer::OnEvent(uint16 RouteId, EStyle Style, const FOnEventCon
 		uint16 FrameCounter = EventData.GetValue<uint16>("FrameCounter");
 		bool bLooping = EventData.GetValue<bool>("Looping");
 		bool bIsBlendSpace = EventData.GetValue<bool>("IsBlendSpace");
-		AnimationProvider.AppendTickRecord(AnimInstanceId, Context.EventTime.AsSeconds(Cycle), AssetId, NodeId, BlendWeight, PlaybackTime, RootMotionWeight, PlayRate, BlendSpacePositionX, BlendSpacePositionY, BlendSpaceFilteredPositionX, BlendSpaceFilteredPositionY, FrameCounter, bLooping, bIsBlendSpace);
+		AnimationProvider.AppendTickRecord(AnimInstanceId, Context.EventTime.AsSeconds(Cycle), 0.0, AssetId, NodeId, BlendWeight, PlaybackTime, RootMotionWeight, PlayRate, BlendSpacePositionX, BlendSpacePositionY, BlendSpaceFilteredPositionX, BlendSpaceFilteredPositionY, FrameCounter, bLooping, bIsBlendSpace);
 		break;
 	}
 	case RouteId_SkeletalMesh:
@@ -99,7 +124,7 @@ bool FAnimationAnalyzer::OnEvent(uint16 RouteId, EStyle Style, const FOnEventCon
 		uint16 LodIndex = EventData.GetValue<uint16>("LodIndex");
 		TArrayView<const FTransform> Pose(reinterpret_cast<const FTransform*>(EventData.GetAttachment()), BoneCount);
 		TArrayView<const FSkeletalMeshNamedCurve> Curves(reinterpret_cast<const FSkeletalMeshNamedCurve*>(EventData.GetAttachment() + (sizeof(FTransform) * BoneCount)), CurveCount);
-		AnimationProvider.AppendSkeletalMeshComponent(ComponentId, MeshId, Context.EventTime.AsSeconds(Cycle), LodIndex, FrameCounter, Pose, Curves);
+		AnimationProvider.AppendSkeletalMeshComponent(ComponentId, MeshId, Context.EventTime.AsSeconds(Cycle), 0.0, LodIndex, FrameCounter, Pose, Curves);
 		break;
 	}
 	case RouteId_SkeletalMeshComponent2:
@@ -118,7 +143,28 @@ bool FAnimationAnalyzer::OnEvent(uint16 RouteId, EStyle Style, const FOnEventCon
 		TArrayView<const float> CurveValues = EventData.GetArrayView<float>("CurveValues");
 		check(CurveIds.Num() == CurveValues.Num());
 
-		AnimationProvider.AppendSkeletalMeshComponent(ComponentId, MeshId, Context.EventTime.AsSeconds(Cycle), LodIndex, FrameCounter, ComponentToWorldFloatArray, PoseFloatArray, CurveIds, CurveValues);
+		AnimationProvider.AppendSkeletalMeshComponent(ComponentId, MeshId, Context.EventTime.AsSeconds(Cycle), 0.0,  LodIndex, FrameCounter, ComponentToWorldFloatArray, PoseFloatArray, CurveIds, CurveValues);
+		break;
+	}
+	case RouteId_SkeletalMeshComponent3:
+	{
+		uint64 Cycle = EventData.GetValue<uint64>("Cycle");
+		double RecordingTime = EventData.GetValue<double>("RecordingTime");
+		
+		uint64 ComponentId = EventData.GetValue<uint64>("ComponentId");
+		uint64 MeshId = EventData.GetValue<uint64>("MeshId");
+		uint32 BoneCount = EventData.GetValue<uint32>("BoneCount");
+		uint32 CurveCount = EventData.GetValue<uint32>("CurveCount");
+		uint16 FrameCounter = EventData.GetValue<uint16>("FrameCounter");
+		uint16 LodIndex = EventData.GetValue<uint16>("LodIndex");
+
+		TArrayView<const float> ComponentToWorldFloatArray = EventData.GetArrayView<float>("ComponentToWorld");
+		TArrayView<const float> PoseFloatArray = EventData.GetArrayView<float>("Pose");
+		TArrayView<const uint32> CurveIds = EventData.GetArrayView<uint32>("CurveIds");
+		TArrayView<const float> CurveValues = EventData.GetArrayView<float>("CurveValues");
+		check(CurveIds.Num() == CurveValues.Num());
+
+		AnimationProvider.AppendSkeletalMeshComponent(ComponentId, MeshId, Context.EventTime.AsSeconds(Cycle), RecordingTime, LodIndex, FrameCounter, ComponentToWorldFloatArray, PoseFloatArray, CurveIds, CurveValues);
 		break;
 	}
 	case RouteId_Name:
@@ -311,7 +357,22 @@ bool FAnimationAnalyzer::OnEvent(uint16 RouteId, EStyle Style, const FOnEventCon
 		float DesiredWeight = EventData.GetValue<float>("DesiredWeight");
 		float Position = EventData.GetValue<float>("Position");
 		uint16 FrameCounter = EventData.GetValue<uint16>("FrameCounter");
-		AnimationProvider.AppendMontage(AnimInstanceId, Context.EventTime.AsSeconds(Cycle), MontageId, CurrentSectionNameId, NextSectionNameId, Weight, DesiredWeight, Position, FrameCounter);
+		AnimationProvider.AppendMontage(AnimInstanceId, Context.EventTime.AsSeconds(Cycle), 0.0, MontageId, CurrentSectionNameId, NextSectionNameId, Weight, DesiredWeight, Position, FrameCounter);
+		break;
+	}
+	case RouteId_Montage2:
+	{
+		uint64 Cycle = EventData.GetValue<uint64>("Cycle");
+		double RecordingTime = EventData.GetValue<double>("RecordingTime");
+		uint64 AnimInstanceId = EventData.GetValue<uint64>("AnimInstanceId");
+		uint64 MontageId = EventData.GetValue<uint64>("MontageId");
+		uint32 CurrentSectionNameId = EventData.GetValue<uint32>("CurrentSectionNameId");
+		uint32 NextSectionNameId = EventData.GetValue<uint32>("NextSectionNameId");
+		float Weight = EventData.GetValue<float>("Weight");
+		float DesiredWeight = EventData.GetValue<float>("DesiredWeight");
+		float Position = EventData.GetValue<float>("Position");
+		uint16 FrameCounter = EventData.GetValue<uint16>("FrameCounter");
+		AnimationProvider.AppendMontage(AnimInstanceId, Context.EventTime.AsSeconds(Cycle), RecordingTime, MontageId, CurrentSectionNameId, NextSectionNameId, Weight, DesiredWeight, Position, FrameCounter);
 		break;
 	}
 	case RouteId_Sync:
diff --git a/Engine/Plugins/Animation/GameplayInsights/Source/GameplayInsights/Private/AnimationAnalyzer.h b/Engine/Plugins/Animation/GameplayInsights/Source/GameplayInsights/Private/AnimationAnalyzer.h
--- a/Engine/Plugins/Animation/GameplayInsights/Source/GameplayInsights/Private/AnimationAnalyzer.h
+++ b/Engine/Plugins/Animation/GameplayInsights/Source/GameplayInsights/Private/AnimationAnalyzer.h
@@ -20,10 +20,12 @@ private:
 	enum : uint16
 	{
 		RouteId_TickRecord,
+		RouteId_TickRecord2,
 		RouteId_SkeletalMesh,
 		RouteId_SkeletalMesh2,
 		RouteId_SkeletalMeshComponent,
 		RouteId_SkeletalMeshComponent2,
+		RouteId_SkeletalMeshComponent3,
 		RouteId_SkeletalMeshFrame,
 		RouteId_AnimGraph,
 		RouteId_AnimNodeStart,
@@ -43,6 +45,7 @@ private:
 		RouteId_Notify,
 		RouteId_SyncMarker,
 		RouteId_Montage,
+		RouteId_Montage2,
 		RouteId_Sync,
 	};
 
diff --git a/Engine/Plugins/Animation/GameplayInsights/Source/GameplayInsights/Private/AnimationProvider.cpp b/Engine/Plugins/Animation/GameplayInsights/Source/GameplayInsights/Private/AnimationProvider.cpp
--- a/Engine/Plugins/Animation/GameplayInsights/Source/GameplayInsights/Private/AnimationProvider.cpp
+++ b/Engine/Plugins/Animation/GameplayInsights/Source/GameplayInsights/Private/AnimationProvider.cpp
@@ -440,7 +440,7 @@ bool FAnimationProvider::HasAnyData() const
 	return bHasAnyData;
 }
 
-void FAnimationProvider::AppendTickRecord(uint64 InAnimInstanceId, double InTime, uint64 InAssetId, int32 InNodeId, float InBlendWeight, float InPlaybackTime, float InRootMotionWeight, float InPlayRate, float InBlendSpacePositionX, float InBlendSpacePositionY, float InBlendSpaceFilteredPositionX, float InBlendSpaceFilteredPositionY, uint16 InFrameCounter, bool bInLooping, bool bInIsBlendSpace)
+void FAnimationProvider::AppendTickRecord(uint64 InAnimInstanceId, double InProfileTime, double InRecordingTime, uint64 InAssetId, int32 InNodeId, float InBlendWeight, float InPlaybackTime, float InRootMotionWeight, float InPlayRate, float InBlendSpacePositionX, float InBlendSpacePositionY, float InBlendSpaceFilteredPositionX, float InBlendSpaceFilteredPositionY, uint16 InFrameCounter, bool bInLooping, bool bInIsBlendSpace)
 {
 	Session.WriteAccessCheck();
 
@@ -467,6 +467,7 @@ void FAnimationProvider::AppendTickRecord(uint64 InAnimInstanceId, double InTime
 	FTickRecordMessage Message;
 	Message.AnimInstanceId = InAnimInstanceId;
 	Message.AssetId = InAssetId;
+	Message.RecordingTime = InRecordingTime;
 	Message.NodeId = InNodeId;
 	Message.BlendWeight = InBlendWeight;
 	Message.PlaybackTime = InPlaybackTime;
@@ -480,9 +481,9 @@ void FAnimationProvider::AppendTickRecord(uint64 InAnimInstanceId, double InTime
 	Message.bLooping = bInLooping;
 	Message.bIsBlendSpace = bInIsBlendSpace;
 
-	TimelineStorage->Timeline->AppendEvent(InTime, Message);
+	TimelineStorage->Timeline->AppendEvent(InProfileTime, Message);
 
-	Session.UpdateDurationSeconds(InTime);
+	Session.UpdateDurationSeconds(InProfileTime);
 }
 
 void FAnimationProvider::AppendSkeletalMesh(uint64 InObjectId, const TArrayView<const int32>& InParentIndices)
@@ -508,7 +509,7 @@ void FAnimationProvider::AppendSkeletalMesh(uint64 InObjectId, const TArrayView<
 	}
 }
 
-void FAnimationProvider::AppendSkeletalMeshComponent(uint64 InObjectId, uint64 InMeshId, double InTime, uint16 InLodIndex, uint16 InFrameCounter, const TArrayView<const FTransform>& InPose, const TArrayView<const FSkeletalMeshNamedCurve>& InCurves)
+void FAnimationProvider::AppendSkeletalMeshComponent(uint64 InObjectId, uint64 InMeshId, double InProfileTime, double InRecordingTime, uint16 InLodIndex, uint16 InFrameCounter, const TArrayView<const FTransform>& InPose, const TArrayView<const FSkeletalMeshNamedCurve>& InCurves)
 {
 	Session.WriteAccessCheck();
 
@@ -533,10 +534,11 @@ void FAnimationProvider::AppendSkeletalMeshComponent(uint64 InObjectId, uint64 I
 	if(NumEvents > 0)
 	{
 		// Add end event at current time
-		TimelineStorage->Timeline->EndEvent(NumEvents - 1, InTime);
+		TimelineStorage->Timeline->EndEvent(NumEvents - 1, InProfileTime);
 	}
 
 	FSkeletalMeshPoseMessage Message;
+	Message.RecordingTime = InRecordingTime;
 	Message.ComponentToWorld = InPose[0];
 	Message.TransformStartIndex = SkeletalMeshPoseTransforms.Num();
 	Message.CurveStartIndex = SkeletalMeshCurves.Num();
@@ -548,7 +550,7 @@ void FAnimationProvider::AppendSkeletalMeshComponent(uint64 InObjectId, uint64 I
 	Message.LodIndex = InLodIndex;
 	Message.FrameCounter = InFrameCounter;
 
-	TimelineStorage->Timeline->AppendBeginEvent(InTime, Message);
+	TimelineStorage->Timeline->AppendBeginEvent(InProfileTime, Message);
 
 	for(int32 TransformIndex = 1; TransformIndex < InPose.Num(); ++TransformIndex)
 	{
@@ -561,7 +563,7 @@ void FAnimationProvider::AppendSkeletalMeshComponent(uint64 InObjectId, uint64 I
 		TimelineStorage->AllCurveIds.Add(Curve.Id);
 	}
 
-	Session.UpdateDurationSeconds(InTime);
+	Session.UpdateDurationSeconds(InProfileTime);
 }
 
 // support for deserializing LWC/non LWC Transforms from builds with non-matching LWC setting
@@ -609,7 +611,7 @@ static FTransform ConvertTransform(int TransformSize, const float* TransformFloa
 	return FTransform(Rotation, Translation, Scale3D);
 }
 
-void FAnimationProvider::AppendSkeletalMeshComponent(uint64 InObjectId, uint64 InMeshId, double InTime, uint16 InLodIndex, uint16 InFrameCounter, const TArrayView<const float>& InComponentToWorldRaw, const TArrayView<const float>& InPoseRaw, const TArrayView<const uint32>& InCurveIds, const TArrayView<const float>& InCurveValues)
+void FAnimationProvider::AppendSkeletalMeshComponent(uint64 InObjectId, uint64 InMeshId, double InProfileTime, double InRecordingTime, uint16 InLodIndex, uint16 InFrameCounter, const TArrayView<const float>& InComponentToWorldRaw, const TArrayView<const float>& InPoseRaw, const TArrayView<const uint32>& InCurveIds, const TArrayView<const float>& InCurveValues)
 {
 	Session.WriteAccessCheck();
 
@@ -634,7 +636,7 @@ void FAnimationProvider::AppendSkeletalMeshComponent(uint64 InObjectId, uint64 I
 	if(NumEvents > 0)
 	{
 		// Add end event at current time
-		TimelineStorage->Timeline->EndEvent(NumEvents - 1, InTime);
+		TimelineStorage->Timeline->EndEvent(NumEvents - 1, InProfileTime);
 	}
 
 	const int32 NumCurves = InCurveIds.Num();
@@ -656,6 +658,7 @@ void FAnimationProvider::AppendSkeletalMeshComponent(uint64 InObjectId, uint64 I
 	const int PoseTransformCount = InPoseRaw.Num()/CaptureTransformSize;
 
 	FSkeletalMeshPoseMessage Message;
+	Message.RecordingTime = InRecordingTime;
 	Message.ComponentToWorld = ComponentToWorld;
 	Message.TransformStartIndex = SkeletalMeshPoseTransforms.Num();
 	Message.CurveStartIndex = SkeletalMeshCurves.Num();
@@ -667,7 +670,7 @@ void FAnimationProvider::AppendSkeletalMeshComponent(uint64 InObjectId, uint64 I
 	Message.LodIndex = InLodIndex;
 	Message.FrameCounter = InFrameCounter;
 
-	TimelineStorage->Timeline->AppendBeginEvent(InTime, Message);
+	TimelineStorage->Timeline->AppendBeginEvent(InProfileTime, Message);
 
 	if (CaptureTransformSize == LocalTransformSize)
 	{
@@ -690,7 +693,7 @@ void FAnimationProvider::AppendSkeletalMeshComponent(uint64 InObjectId, uint64 I
 		TimelineStorage->AllCurveIds.Add(InCurveIds[CurveIndex]);
 	}
 
-	Session.UpdateDurationSeconds(InTime);
+	Session.UpdateDurationSeconds(InProfileTime);
 }
 
 void FAnimationProvider::AppendName(uint32 InId, const TCHAR* InName)
@@ -1183,7 +1186,7 @@ void FAnimationProvider::AppendNotify(uint64 InAnimInstanceId, double InTime, ui
 	Session.UpdateDurationSeconds(InTime);
 }
 
-void FAnimationProvider::AppendMontage(uint64 InAnimInstanceId, double InTime, uint64 InMontageId, uint32 InCurrentSectionNameId, uint32 InNextSectionNameId, float InWeight, float InDesiredWeight, float InPosition, uint16 InFrameCounter)
+void FAnimationProvider::AppendMontage(uint64 InAnimInstanceId, double InProfileTime, double InRecoringTime, uint64 InMontageId, uint32 InCurrentSectionNameId, uint32 InNextSectionNameId, float InWeight, float InDesiredWeight, float InPosition, uint16 InFrameCounter)
 {
 	Session.WriteAccessCheck();
 
@@ -1207,6 +1210,7 @@ void FAnimationProvider::AppendMontage(uint64 InAnimInstanceId, double InTime, u
 	TimelineStorage->AllMontageIds.Add(InMontageId);
 
 	FAnimMontageMessage Message;
+	Message.RecordingTime = InRecoringTime;
 	Message.AnimInstanceId = InAnimInstanceId;
 	Message.MontageId = InMontageId;
 	Message.CurrentSectionNameId = InCurrentSectionNameId;
@@ -1216,9 +1220,9 @@ void FAnimationProvider::AppendMontage(uint64 InAnimInstanceId, double InTime, u
 	Message.Position = InPosition;
 	Message.FrameCounter = InFrameCounter;
 
-	TimelineStorage->Timeline->AppendEvent(InTime, Message);
+	TimelineStorage->Timeline->AppendEvent(InProfileTime, Message);
 
-	Session.UpdateDurationSeconds(InTime);
+	Session.UpdateDurationSeconds(InProfileTime);
 }
 
 void FAnimationProvider::AppendSync(uint64 InAnimInstanceId, double InTime, int32 InSourceNodeId, uint32 InGroupNameId)
diff --git a/Engine/Plugins/Animation/GameplayInsights/Source/GameplayInsights/Private/AnimationProvider.h b/Engine/Plugins/Animation/GameplayInsights/Source/GameplayInsights/Private/AnimationProvider.h
--- a/Engine/Plugins/Animation/GameplayInsights/Source/GameplayInsights/Private/AnimationProvider.h
+++ b/Engine/Plugins/Animation/GameplayInsights/Source/GameplayInsights/Private/AnimationProvider.h
@@ -52,14 +52,14 @@ public:
 	bool HasAnyData() const;
 
 	/** Add a tick record */
-	void AppendTickRecord(uint64 InAnimInstanceId, double InTime, uint64 InAssetId, int32 InNodeId, float InBlendWeight, float InPlaybackTime, float InRootMotionWeight, float InPlayRate, float InBlendSpacePositionX, float InBlendSpacePositionY, float InBlendSpaceFilteredPositionX, float InBlendSpaceFilteredPositionY, uint16 InFrameCounter, bool bInLooping, bool bInIsBlendSpace);
+	void AppendTickRecord(uint64 InAnimInstanceId, double InProfileTime, double InRecordingTime, uint64 InAssetId, int32 InNodeId, float InBlendWeight, float InPlaybackTime, float InRootMotionWeight, float InPlayRate, float InBlendSpacePositionX, float InBlendSpacePositionY, float InBlendSpaceFilteredPositionX, float InBlendSpaceFilteredPositionY, uint16 InFrameCounter, bool bInLooping, bool bInIsBlendSpace);
 
 	/** Add a skeletal mesh */
 	void AppendSkeletalMesh(uint64 InObjectId, const TArrayView<const int32>& ParentIndices);
 
 	/** Add a skeletal mesh pose/curves etc. */
-	void AppendSkeletalMeshComponent(uint64 InObjectId, uint64 InMeshId, double InTime, uint16 InLodIndex, uint16 InFrameCounter, const TArrayView<const FTransform>& InPose, const TArrayView<const FSkeletalMeshNamedCurve>& InCurves);
-	void AppendSkeletalMeshComponent(uint64 InObjectId, uint64 InMeshId, double InTime, uint16 InLodIndex, uint16 InFrameCounter, const TArrayView<const float>& InComponentToWorldRaw, const TArrayView<const float>& InPoseRaw, const TArrayView<const uint32>& InCurveIds, const TArrayView<const float>& InCurveValues);
+	void AppendSkeletalMeshComponent(uint64 InObjectId, uint64 InMeshId, double InProfileTime, double InRecordingTime, uint16 InLodIndex, uint16 InFrameCounter, const TArrayView<const FTransform>& InPose, const TArrayView<const FSkeletalMeshNamedCurve>& InCurves);
+	void AppendSkeletalMeshComponent(uint64 InObjectId, uint64 InMeshId, double InProfileTime, double InRecordingTime, uint16 InLodIndex, uint16 InFrameCounter, const TArrayView<const float>& InComponentToWorldRaw, const TArrayView<const float>& InPoseRaw, const TArrayView<const uint32>& InCurveIds, const TArrayView<const float>& InCurveValues);
 
 	/** Get a skeletal mesh for a specified path. If the mesh no longer exists, make a fake one */
 	USkeletalMesh* GetSkeletalMesh(const TCHAR* InPath);
@@ -102,7 +102,7 @@ public:
 	void AppendNotify(uint64 InAnimInstanceId, double InTime, uint64 InAssetId, uint64 InNotifyId, uint32 InNameId, float InNotifyTime, float InNotifyDuration, EAnimNotifyMessageType InNotifyEventType);
 
 	/** Append montage data */
-	void AppendMontage(uint64 InAnimInstanceId, double InTime, uint64 InMontageId, uint32 InCurrentSectionNameId, uint32 InNextSectionNameId, float InWeight, float InDesiredWeight, float InPosition, uint16 InFrameCounter);
+	void AppendMontage(uint64 InAnimInstanceId, double InProfileTime, double InRecordingTime, uint64 InMontageId, uint32 InCurrentSectionNameId, uint32 InNextSectionNameId, float InWeight, float InDesiredWeight, float InPosition, uint16 InFrameCounter);
 
 	/** Append sync data */
 	void AppendSync(uint64 InAnimInstanceId, double InTime, int32 InSourceNodeId, uint32 InGroupNameId);
diff --git a/Engine/Plugins/Animation/GameplayInsights/Source/GameplayInsights/Private/BlendWeightsTrack.cpp b/Engine/Plugins/Animation/GameplayInsights/Source/GameplayInsights/Private/BlendWeightsTrack.cpp
new file mode 100644
--- /dev/null
+++ b/Engine/Plugins/Animation/GameplayInsights/Source/GameplayInsights/Private/BlendWeightsTrack.cpp
@@ -0,0 +1,373 @@
+// Copyright Epic Games, Inc. All Rights Reserved.
+
+#include "BlendWeightsTrack.h"
+#include "SCurveTimelineView.h"
+#include "IRewindDebugger.h"
+#include "GameplayProvider.h"
+#include "AnimationProvider.h"
+#include "Styling/SlateIconFinder.h"
+#include "SBlendWeightsView.h"
+
+#if WITH_EDITOR
+#include "Animation/AnimInstance.h"
+#include "Animation/AnimSequence.h"
+#include "Animation/BlendSpace.h"
+#endif
+
+#define LOCTEXT_NAMESPACE "BlendWeightsTrack"
+
+namespace RewindDebugger
+{
+
+FBlendWeightsTrack::FBlendWeightsTrack(uint64 InObjectId) : ObjectId(InObjectId)
+{
+#if WITH_EDITOR
+	Icon = FSlateIconFinder::FindIconForClass(UAnimSequence::StaticClass());
+#endif
+}
+
+void FBlendWeightsTrack::IterateSubTracksInternal(TFunction<void(TSharedPtr<FRewindDebuggerTrack> SubTrack)> IteratorFunction)
+{
+	for(TSharedPtr<FBlendWeightTrack>& Track : Children)
+	{
+		IteratorFunction(Track);
+	}
+};
+
+bool FBlendWeightsTrack::UpdateInternal()
+{
+	IRewindDebugger* RewindDebugger = IRewindDebugger::Instance();
+
+	struct TrackId
+	{
+		TrackId() : AssetId(0), NodeId(0) {}
+		TrackId(uint64 InAssetId, uint32 InNodeId) : AssetId(InAssetId), NodeId(InNodeId) { }
+		bool operator ==(const TrackId& Other) const { return AssetId == Other.AssetId && NodeId == Other.NodeId; }
+		
+		uint64 AssetId;
+		uint32 NodeId;
+	};	
+
+	TArray<TrackId> UniqueTrackIds;
+	
+	const TraceServices::IAnalysisSession* AnalysisSession = RewindDebugger->GetAnalysisSession();
+
+	bool bChanged = false;
+	
+	// convert time range to from rewind debugger times to profiler times
+	TRange<double> TraceTimeRange = RewindDebugger->GetCurrentTraceRange();
+	double StartTime = TraceTimeRange.GetLowerBoundValue();
+	double EndTime = TraceTimeRange.GetUpperBoundValue();
+	
+	// count number of unique animations in the current time range
+	const FGameplayProvider* GameplayProvider = AnalysisSession->ReadProvider<FGameplayProvider>(FGameplayProvider::ProviderName);
+	const FAnimationProvider* AnimationProvider = AnalysisSession->ReadProvider<FAnimationProvider>(FAnimationProvider::ProviderName);
+
+	int AnimationCount = 0;
+
+	if(GameplayProvider && AnimationProvider)
+	{
+		TraceServices::FAnalysisSessionReadScope SessionReadScope(*AnalysisSession);
+		UniqueTrackIds.SetNum(0, false);
+
+		AnimationProvider->ReadTickRecordTimeline(ObjectId, [&UniqueTrackIds,&GameplayProvider, StartTime, EndTime](const FAnimationProvider::TickRecordTimeline& InTimeline)
+		{
+			// this isn't very efficient, and it gets called every frame.  will need optimizing
+			InTimeline.EnumerateEvents(StartTime, EndTime, [&UniqueTrackIds, StartTime, EndTime](double InStartTime, double InEndTime, uint32 InDepth, const FTickRecordMessage& InMessage)
+			{
+				if (InEndTime > StartTime && InStartTime < EndTime)
+				{
+					UniqueTrackIds.AddUnique(TrackId(InMessage.AssetId, InMessage.NodeId));
+				}
+				return TraceServices::EEventEnumerate::Continue;
+			});
+		});
+
+		UniqueTrackIds.StableSort([](const TrackId &A, const TrackId &B)  { return A.NodeId > B.NodeId; });
+		
+		AnimationCount = UniqueTrackIds.Num();
+
+		if (Children.Num()!=AnimationCount)
+			bChanged = true;
+		
+		Children.SetNum(AnimationCount);
+		for(int i = 0; i < AnimationCount; i++)
+		{
+			if (!Children[i].IsValid() || Children[i].Get()->GetAssetId() != UniqueTrackIds[i].AssetId)
+			{
+				Children[i] = MakeShared<FBlendWeightTrack>(ObjectId, UniqueTrackIds[i].AssetId, UniqueTrackIds[i].NodeId);
+				bChanged = true;
+			}
+
+			bChanged = bChanged || Children[i]->Update();
+		}
+	}
+
+	return bChanged;
+}
+
+TSharedPtr<SWidget> FBlendWeightsTrack::GetDetailsViewInternal()
+{
+	IRewindDebugger* RewindDebugger = IRewindDebugger::Instance();
+	return SNew(SBlendWeightsView, ObjectId, RewindDebugger->CurrentTraceTime(), *RewindDebugger->GetAnalysisSession())
+				.CurrentTime_Lambda([RewindDebugger]{ return RewindDebugger->CurrentTraceTime(); });
+}
+
+
+FBlendWeightTrack::FBlendWeightTrack(uint64 InObjectId, uint64 InAssetId, uint32 InNodeId, FBlendWeightTrack::ECurveType InCurveType) :
+	 AssetId(InAssetId)
+	, NodeId(InNodeId)
+	, CurveType(InCurveType)
+	, ObjectId(InObjectId)
+{
+	SetIsExpanded(false);
+}
+
+
+void FBlendWeightTrack::IterateSubTracksInternal(TFunction<void(TSharedPtr<FRewindDebuggerTrack> SubTrack)> IteratorFunction)
+{
+	for (TSharedPtr<FBlendWeightTrack>& Track : Children)
+	{
+		IteratorFunction(Track);
+	}
+};
+
+TSharedPtr<SCurveTimelineView::FTimelineCurveData> FBlendWeightTrack::GetCurveData() const
+{
+	if (!CurveData.IsValid())
+	{
+		CurveData = MakeShared<SCurveTimelineView::FTimelineCurveData>();
+	}
+	
+	CurvesUpdateRequested++;
+	
+	return CurveData;
+}
+
+bool FBlendWeightTrack::UpdateInternal()
+{
+	IRewindDebugger* RewindDebugger = IRewindDebugger::Instance();
+	// compute curve points
+	//
+	TRange<double> TraceTimeRange = RewindDebugger->GetCurrentTraceRange();
+	double StartTime = TraceTimeRange.GetLowerBoundValue();
+	double EndTime = TraceTimeRange.GetUpperBoundValue();
+
+	const TraceServices::IAnalysisSession* AnalysisSession = RewindDebugger->GetAnalysisSession();
+	const FGameplayProvider* GameplayProvider = AnalysisSession->ReadProvider<FGameplayProvider>(FGameplayProvider::ProviderName);
+	const FAnimationProvider* AnimationProvider = AnalysisSession->ReadProvider<FAnimationProvider>(FAnimationProvider::ProviderName);
+
+	int AnimationCount = 0;
+
+	if(CurvesUpdateRequested > 10 && GameplayProvider && AnimationProvider)
+	{
+		auto& CurvePoints = CurveData->Points;
+		CurvePoints.SetNum(0,false);
+		
+		TraceServices::FAnalysisSessionReadScope SessionReadScope(*AnalysisSession);
+		
+		AnimationProvider->ReadTickRecordTimeline(ObjectId, [AnalysisSession, StartTime, EndTime, &CurvePoints, this](const FAnimationProvider::TickRecordTimeline& InTimeline)
+		{
+			InTimeline.EnumerateEvents(StartTime, EndTime, [&CurvePoints, AnalysisSession, this](double InStartTime, double InEndTime, uint32 InDepth, const FTickRecordMessage& InMessage)
+			{
+				if (InMessage.AssetId == AssetId && InMessage.NodeId == NodeId)
+				{
+					float Weight = 0;
+					switch (CurveType)
+					{
+						case ECurveType::BlendWeight:					Weight = InMessage.BlendWeight; break; 
+						case ECurveType::PlaybackTime:					Weight = InMessage.PlaybackTime; break; 
+						case ECurveType::RootMotionWeight:				Weight = InMessage.RootMotionWeight; break; 
+						case ECurveType::PlayRate:						Weight = InMessage.PlayRate; break; 
+						case ECurveType::BlendSpacePositionX:			Weight = InMessage.BlendSpacePositionX; break; 
+						case ECurveType::BlendSpacePositionY:			Weight = InMessage.BlendSpacePositionY; break; 
+						case ECurveType::BlendSpaceFilteredPositionX:	Weight = InMessage.BlendSpaceFilteredPositionX; break; 
+						case ECurveType::BlendSpaceFilteredPositionY:	Weight = InMessage.BlendSpaceFilteredPositionY; break; 
+					}
+					
+					CurvePoints.Add({ InMessage.RecordingTime,	Weight });
+				}
+								
+				return TraceServices::EEventEnumerate::Continue;
+			});
+		});
+
+		CurvesUpdateRequested = 0;
+	}
+
+	// update Icon:
+	
+	bool bChanged = false;
+
+	if (CurveType == ECurveType::BlendWeight)
+	{
+		// Blend Weight track gets name/icon of animation
+		if (CurveName.IsEmpty())
+		{
+			if (const FObjectInfo* ObjectInfo = GameplayProvider->FindObjectInfo(AssetId))
+			{
+				CurveName = FText::FromString(ObjectInfo->Name);
+				bChanged = true;
+				bool bIsBlendSpace = false;
+
+				if (const UClass* FoundClass = GameplayProvider->FindClass(ObjectInfo->ClassId))
+				{
+					Icon = FSlateIconFinder::FindIconForClass(FoundClass);
+#if WITH_EDITOR
+					if (FoundClass->IsChildOf(UBlendSpace::StaticClass()))
+					{
+						bIsBlendSpace = true;	
+					}
+#endif
+				}
+
+				// Blend Weight track gets child tracks for each of the other bits of extra data
+				if (Children.Num() == 0)
+				{
+					int MaxCurve = static_cast<int>(bIsBlendSpace ? ECurveType::BlendSpaceFilteredPositionY : ECurveType::PlayRate);
+					for (int i = static_cast<int>(ECurveType::BlendWeight) + 1; i <= MaxCurve; i++)
+					{
+						Children.Add(MakeShared<FBlendWeightTrack>(ObjectId, AssetId, NodeId, static_cast<ECurveType>(i)));
+					}
+				}
+			}
+		}
+
+	}
+	else
+	{
+		// other tracks get the curve name
+
+		if (CurveName.IsEmpty())
+		{
+			switch (CurveType)
+			{
+				case ECurveType::BlendWeight:					break;
+				case ECurveType::PlaybackTime:					CurveName = LOCTEXT("PlaybackTime", "Playback Time"); break;
+				case ECurveType::RootMotionWeight:				CurveName = LOCTEXT("RootMotionWeight", "Root Motion Weight"); break;
+				case ECurveType::PlayRate:						CurveName = LOCTEXT("PlaybackRate", "Playback Rate"); break;
+				case ECurveType::BlendSpacePositionX:			CurveName = LOCTEXT("BlendSpaceX", "X"); break;
+				case ECurveType::BlendSpacePositionY:			CurveName = LOCTEXT("BlendSpaceY", "Y"); break;
+				case ECurveType::BlendSpaceFilteredPositionX:	CurveName = LOCTEXT("BlendSpaceXF", "X (Filtered)"); break;
+				case ECurveType::BlendSpaceFilteredPositionY:	CurveName = LOCTEXT("BlendSpaceYF", "Y (Filtered)"); break;
+			}
+
+			Icon = FSlateIcon("EditorStyle", "AnimGraph.Attribute.Curves.Icon", "AnimGraph.Attribute.Curves.Icon");
+
+			bChanged = true;
+		}
+	}
+
+	for (auto& Child : Children)
+	{
+		bChanged |= Child->Update();
+	}
+
+	return bChanged;
+}
+
+static FLinearColor MakeBlendWeightCurveColor(uint32 InSeed, bool bInLine = false)
+{
+	FRandomStream Stream(InSeed);
+	const uint8 Hue = (uint8)(Stream.FRand() * 255.0f);
+	const uint8 SatVal = bInLine ? 196 : 128;
+	return FLinearColor::MakeFromHSV8(Hue, SatVal, SatVal);
+}
+
+TSharedPtr<SWidget> FBlendWeightTrack::GetTimelineViewInternal()
+{
+	FLinearColor Color;
+	switch(CurveType)
+	{
+	case ECurveType::BlendWeight:
+		Color = MakeBlendWeightCurveColor(CityHash32(reinterpret_cast<char*>(&AssetId), 8));
+		Color.A = 0.5f;
+		break; 
+	case ECurveType::PlaybackTime:
+		Color = FLinearColor::MakeFromHSV8(0, 50, 50);
+		break;
+	case ECurveType::RootMotionWeight:
+		Color = FLinearColor::MakeFromHSV8(60, 50, 50);
+		break;
+	case ECurveType::PlayRate:		
+		Color = FLinearColor::MakeFromHSV8(120, 50, 50);
+		break;
+	case ECurveType::BlendSpacePositionX:
+		Color = FLinearColor::MakeFromHSV8(180, 50, 50);
+		break;
+	case ECurveType::BlendSpacePositionY:
+		Color = FLinearColor::MakeFromHSV8(240, 50, 50);
+		break;
+	case ECurveType::BlendSpaceFilteredPositionX:
+		Color = FLinearColor::MakeFromHSV8(180, 50, 80);
+		break;
+	case ECurveType::BlendSpaceFilteredPositionY:
+		Color = FLinearColor::MakeFromHSV8(240, 50, 80);
+		break;
+	}
+
+	FLinearColor CurveColor = Color;
+	CurveColor.R *= 0.5;
+	CurveColor.G *= 0.5;
+	CurveColor.B *= 0.5;
+
+	TSharedPtr<SCurveTimelineView> CurveTimelineView = SNew(SCurveTimelineView)
+		.FillColor(Color)
+		.CurveColor(CurveColor)
+		.ViewRange_Lambda([]() { return IRewindDebugger::Instance()->GetCurrentViewRange(); })
+		.RenderFill(CurveType == ECurveType::BlendWeight)
+		.CurveData_Raw(this, &FBlendWeightTrack::GetCurveData);
+
+	if (CurveType == ECurveType::BlendWeight)
+	{
+		CurveTimelineView->SetFixedRange(0, 1);
+	}
+
+	return CurveTimelineView;
+}
+
+TSharedPtr<SWidget> FBlendWeightTrack::GetDetailsViewInternal()
+{
+	IRewindDebugger* RewindDebugger = IRewindDebugger::Instance();
+	TSharedPtr<SBlendWeightsView> BlendWeightsView = SNew(SBlendWeightsView, ObjectId, RewindDebugger->CurrentTraceTime(), *RewindDebugger->GetAnalysisSession())
+				.CurrentTime_Lambda([RewindDebugger]{ return RewindDebugger->CurrentTraceTime(); });
+	BlendWeightsView->SetAssetFilter(AssetId,NodeId);
+	return BlendWeightsView;
+}
+
+
+FName FBlendWeightsTrackCreator::GetTargetTypeNameInternal() const
+{
+	static const FName TargetTypeName = "AnimInstance";
+	return TargetTypeName;
+}
+
+FName FBlendWeightsTrackCreator::GetNameInternal() const
+{
+	static const FName BlendWeightsName("BlendWeights");
+	return BlendWeightsName;
+}
+
+TSharedPtr<RewindDebugger::FRewindDebuggerTrack> FBlendWeightsTrackCreator::CreateTrackInternal(uint64 ObjectId) const
+{
+ 	return MakeShared<RewindDebugger::FBlendWeightsTrack>(ObjectId);
+}
+
+bool FBlendWeightsTrackCreator::HasDebugInfoInternal(uint64 ObjectId) const
+{
+	const TraceServices::IAnalysisSession* AnalysisSession = IRewindDebugger::Instance()->GetAnalysisSession();
+	TraceServices::FAnalysisSessionReadScope SessionReadScope(*AnalysisSession);
+	bool bHasData = false;
+	if (const FAnimationProvider* AnimationProvider = AnalysisSession->ReadProvider<FAnimationProvider>(FAnimationProvider::ProviderName))
+	{
+		AnimationProvider->ReadAnimGraphTimeline(ObjectId, [&bHasData](const FAnimationProvider::AnimGraphTimeline& InGraphTimeline)
+		{
+			bHasData = true;
+		});
+	}
+	
+	return bHasData;
+}
+
+}
+#undef LOCTEXT_NAMESPACE
diff --git a/Engine/Plugins/Animation/GameplayInsights/Source/GameplayInsights/Private/BlendWeightsTrack.h b/Engine/Plugins/Animation/GameplayInsights/Source/GameplayInsights/Private/BlendWeightsTrack.h
new file mode 100644
--- /dev/null
+++ b/Engine/Plugins/Animation/GameplayInsights/Source/GameplayInsights/Private/BlendWeightsTrack.h
@@ -0,0 +1,92 @@
+// Copyright Epic Games, Inc. All Rights Reserved.
+#pragma once
+
+#include "IAnimationProvider.h"
+#include "IRewindDebugger.h"
+#include "RewindDebuggerTrack.h"
+#include "SCurveTimelineView.h"
+#include "Textures/SlateIcon.h"
+#include "IRewindDebuggerTrackCreator.h"
+
+namespace RewindDebugger
+{
+
+class FBlendWeightTrack : public FRewindDebuggerTrack
+{
+public:
+	enum class ECurveType
+	{
+		BlendWeight,
+		PlaybackTime,
+		RootMotionWeight,
+		PlayRate,
+		BlendSpacePositionX,
+		BlendSpacePositionY,
+		BlendSpaceFilteredPositionX,
+		BlendSpaceFilteredPositionY,
+	};
+	
+	FBlendWeightTrack(uint64 InObjectId, uint64 AssetId, uint32 NodeId, ECurveType CurveType = ECurveType::BlendWeight);
+
+	uint32 GetNodeId() const { return NodeId; };
+	uint64 GetAssetId() const { return AssetId; }
+
+private:
+	virtual bool UpdateInternal() override;
+	virtual TSharedPtr<SWidget> GetTimelineViewInternal() override;
+	virtual TSharedPtr<SWidget> GetDetailsViewInternal() override;
+	virtual void IterateSubTracksInternal(TFunction<void(TSharedPtr<FRewindDebuggerTrack> SubTrack)> IteratorFunction) override;
+	
+	virtual FSlateIcon GetIconInternal() override { return Icon; }
+	virtual FName GetNameInternal() const override { return "BlendWeight"; }
+	virtual FText GetDisplayNameInternal() const override { return CurveName; }
+	virtual uint64 GetObjectIdInternal() const override { return ObjectId; }
+
+	TSharedPtr<SCurveTimelineView::FTimelineCurveData> GetCurveData() const;
+	
+	uint64 AssetId;
+	uint32 NodeId;
+
+	ECurveType CurveType;
+	
+	FSlateIcon Icon;
+	FText CurveName;
+	uint64 ObjectId;
+
+	mutable TSharedPtr<SCurveTimelineView::FTimelineCurveData> CurveData;
+	mutable int CurvesUpdateRequested = 0;
+	
+	TArray<TSharedPtr<FBlendWeightTrack>> Children;
+};
+
+class FBlendWeightsTrack : public FRewindDebuggerTrack
+{
+public:
+	FBlendWeightsTrack(uint64 InObjectId);
+
+private:
+	virtual bool UpdateInternal() override;
+	virtual TSharedPtr<SWidget> GetDetailsViewInternal() override;
+	virtual void IterateSubTracksInternal(TFunction<void(TSharedPtr<FRewindDebuggerTrack> SubTrack)> IteratorFunction) override;
+	
+	virtual FName GetNameInternal() const override { return "BlendWeights"; }
+	virtual FText GetDisplayNameInternal() const override { return NSLOCTEXT("RewindDebugger", "Blend Weights Track Name", "Blend Weights"); }
+	virtual uint64 GetObjectIdInternal() const override { return ObjectId; }
+
+	virtual FSlateIcon GetIconInternal() override { return Icon; }
+
+	FSlateIcon Icon;
+	uint64 ObjectId;
+	TArray<TSharedPtr<FBlendWeightTrack>> Children;
+};
+	
+class FBlendWeightsTrackCreator : public IRewindDebuggerTrackCreator
+{
+public:
+	virtual FName GetTargetTypeNameInternal() const;
+	virtual FName GetNameInternal() const override;
+	virtual TSharedPtr<FRewindDebuggerTrack> CreateTrackInternal(uint64 ObjectId) const override;
+	virtual bool HasDebugInfoInternal(uint64 ObjectId) const override;
+};
+
+}
\ No newline at end of file
diff --git a/Engine/Plugins/Animation/GameplayInsights/Source/GameplayInsights/Private/GameplayAnalyzer.cpp b/Engine/Plugins/Animation/GameplayInsights/Source/GameplayInsights/Private/GameplayAnalyzer.cpp
--- a/Engine/Plugins/Animation/GameplayInsights/Source/GameplayInsights/Private/GameplayAnalyzer.cpp
+++ b/Engine/Plugins/Animation/GameplayInsights/Source/GameplayInsights/Private/GameplayAnalyzer.cpp
@@ -19,8 +19,10 @@ void FGameplayAnalyzer::OnAnalysisBegin(const FOnAnalysisContext& Context)
 	Builder.RouteEvent(RouteId_Class, "Object", "Class");
 	Builder.RouteEvent(RouteId_Object, "Object", "Object");
 	Builder.RouteEvent(RouteId_ObjectEvent, "Object", "ObjectEvent");
-	Builder.RouteEvent(RoutId_ObjectLifetimeBegin, "Object", "ObjectLifetimeBegin");
-	Builder.RouteEvent(RoutId_ObjectLifetimeEnd, "Object", "ObjectLifetimeEnd");
+	Builder.RouteEvent(RouteId_ObjectLifetimeBegin, "Object", "ObjectLifetimeBegin");
+	Builder.RouteEvent(RouteId_ObjectLifetimeBegin2, "Object", "ObjectLifetimeBegin2");
+	Builder.RouteEvent(RouteId_ObjectLifetimeEnd, "Object", "ObjectLifetimeEnd");
+	Builder.RouteEvent(RouteId_ObjectLifetimeEnd2, "Object", "ObjectLifetimeEnd2");
 	Builder.RouteEvent(RouteId_PawnPossess, "Object", "PawnPossess");
 	Builder.RouteEvent(RouteId_World, "Object", "World");
 	Builder.RouteEvent(RouteId_View, "Object", "View");
@@ -105,18 +107,34 @@ bool FGameplayAnalyzer::OnEvent(uint16 RouteId, EStyle Style, const FOnEventCont
 		GameplayProvider.AppendObjectEvent(Id, Context.EventTime.AsSeconds(Cycle), *Event);
 		break;
 	}
-	case RoutId_ObjectLifetimeBegin:
+	case RouteId_ObjectLifetimeBegin:
 	{
 		uint64 Cycle = EventData.GetValue<uint64>("Cycle");
 		uint64 Id = EventData.GetValue<uint64>("Id");
-		GameplayProvider.AppendObjectLifetimeBegin(Id, Context.EventTime.AsSeconds(Cycle));
+		GameplayProvider.AppendObjectLifetimeBegin(Id, Context.EventTime.AsSeconds(Cycle), 0.0);
 		break;
 	}
-	case RoutId_ObjectLifetimeEnd:
+	case RouteId_ObjectLifetimeBegin2:
 	{
 		uint64 Cycle = EventData.GetValue<uint64>("Cycle");
+		double RecordingTime = EventData.GetValue<double>("RecordingTime");
 		uint64 Id = EventData.GetValue<uint64>("Id");
-		GameplayProvider.AppendObjectLifetimeEnd(Id, Context.EventTime.AsSeconds(Cycle));
+		GameplayProvider.AppendObjectLifetimeBegin(Id, Context.EventTime.AsSeconds(Cycle), RecordingTime);
+		break;
+	}
+	case RouteId_ObjectLifetimeEnd:
+	{
+		uint64 Cycle = EventData.GetValue<uint64>("Cycle");
+		uint64 Id = EventData.GetValue<uint64>("Id");
+		GameplayProvider.AppendObjectLifetimeEnd(Id, Context.EventTime.AsSeconds(Cycle), 0.0);
+		break;
+	}
+	case RouteId_ObjectLifetimeEnd2:
+	{
+		uint64 Cycle = EventData.GetValue<uint64>("Cycle");
+		double RecordingTime = EventData.GetValue<double>("RecordingTime");
+		uint64 Id = EventData.GetValue<uint64>("Id");
+		GameplayProvider.AppendObjectLifetimeEnd(Id, Context.EventTime.AsSeconds(Cycle), RecordingTime);
 		break;
 	}
 	case RouteId_PawnPossess:
diff --git a/Engine/Plugins/Animation/GameplayInsights/Source/GameplayInsights/Private/GameplayAnalyzer.h b/Engine/Plugins/Animation/GameplayInsights/Source/GameplayInsights/Private/GameplayAnalyzer.h
--- a/Engine/Plugins/Animation/GameplayInsights/Source/GameplayInsights/Private/GameplayAnalyzer.h
+++ b/Engine/Plugins/Animation/GameplayInsights/Source/GameplayInsights/Private/GameplayAnalyzer.h
@@ -23,8 +23,10 @@ private:
 		RouteId_Class,
 		RouteId_Object,
 		RouteId_ObjectEvent,
-		RoutId_ObjectLifetimeBegin,
-		RoutId_ObjectLifetimeEnd,
+		RouteId_ObjectLifetimeBegin,
+		RouteId_ObjectLifetimeBegin2,
+		RouteId_ObjectLifetimeEnd,
+		RouteId_ObjectLifetimeEnd2,
 		RouteId_PawnPossess,
 		RouteId_View,
 		RouteId_ClassPropertyStringId,
diff --git a/Engine/Plugins/Animation/GameplayInsights/Source/GameplayInsights/Private/GameplayInsightsModule.cpp b/Engine/Plugins/Animation/GameplayInsights/Source/GameplayInsights/Private/GameplayInsightsModule.cpp
--- a/Engine/Plugins/Animation/GameplayInsights/Source/GameplayInsights/Private/GameplayInsightsModule.cpp
+++ b/Engine/Plugins/Animation/GameplayInsights/Source/GameplayInsights/Private/GameplayInsightsModule.cpp
@@ -15,15 +15,16 @@
 #include "Trace/StoreClient.h"
 #include "Stats/Stats.h"
 #include "ObjectPropertyTrace.h"
-#include "SAnimationCurvesView.h"
 #include "SBlendWeightsView.h"
 #include "SMontageView.h"
 #include "SObjectPropertiesView.h"
 #include "SNotifiesView.h"
+#include "AnimCurvesTrack.h"
+#include "BlendWeightsTrack.h"
+#include "MontagesTrack.h"
 
 #if WITH_EDITOR
 #include "IAnimationBlueprintEditorModule.h"
-#include "Editor.h"
 #include "ToolMenus.h"
 #include "Engine/Selection.h"
 #include "SubobjectEditorMenuContext.h"
@@ -69,14 +70,15 @@ void FGameplayInsightsModule::StartupModule()
 	IModularFeatures::Get().RegisterModularFeature(IRewindDebuggerViewCreator::ModularFeatureName, &ObjectPropertiesViewCreator);
 	static FAnimGraphSchematicViewCreator AnimGraphSchematicViewCreator;
 	IModularFeatures::Get().RegisterModularFeature(IRewindDebuggerViewCreator::ModularFeatureName, &AnimGraphSchematicViewCreator);
-	static FBlendWeightsViewCreator BlendWeightsViewCreator;
-	IModularFeatures::Get().RegisterModularFeature(IRewindDebuggerViewCreator::ModularFeatureName, &BlendWeightsViewCreator);
-	static FMontageViewCreator MontageViewCreator;
-	IModularFeatures::Get().RegisterModularFeature(IRewindDebuggerViewCreator::ModularFeatureName, &MontageViewCreator);
 	static FNotifiesViewCreator NotifiesViewCreator;
 	IModularFeatures::Get().RegisterModularFeature(IRewindDebuggerViewCreator::ModularFeatureName, &NotifiesViewCreator);
-	static FAnimationCurvesViewCreator AnimationCurvesViewCreator;
-	IModularFeatures::Get().RegisterModularFeature(IRewindDebuggerViewCreator::ModularFeatureName, &AnimationCurvesViewCreator);
+	
+	static RewindDebugger::FAnimationCurvesTrackCreator AnimationCurvesTrackCreator;
+	IModularFeatures::Get().RegisterModularFeature(RewindDebugger::IRewindDebuggerTrackCreator::ModularFeatureName, &AnimationCurvesTrackCreator);
+	static RewindDebugger::FBlendWeightsTrackCreator BlendWeightsTrackCreator;
+	IModularFeatures::Get().RegisterModularFeature(RewindDebugger::IRewindDebuggerTrackCreator::ModularFeatureName, &BlendWeightsTrackCreator);
+	static RewindDebugger::FMontagesTrackCreator MontagesTrackCreator;
+	IModularFeatures::Get().RegisterModularFeature(RewindDebugger::IRewindDebuggerTrackCreator::ModularFeatureName, &MontagesTrackCreator);
 
 
 	if (!IsRunningCommandlet())
diff --git a/Engine/Plugins/Animation/GameplayInsights/Source/GameplayInsights/Private/GameplayProvider.cpp b/Engine/Plugins/Animation/GameplayInsights/Source/GameplayInsights/Private/GameplayProvider.cpp
--- a/Engine/Plugins/Animation/GameplayInsights/Source/GameplayInsights/Private/GameplayProvider.cpp
+++ b/Engine/Plugins/Animation/GameplayInsights/Source/GameplayInsights/Private/GameplayProvider.cpp
@@ -2,17 +2,118 @@
 
 #include "GameplayProvider.h"
 
+#include "AssetRegistry/AssetRegistryModule.h"
 #include "Model/AsyncEnumerateTask.h"
 
 FName FGameplayProvider::ProviderName("GameplayProvider");
 
 #define LOCTEXT_NAMESPACE "GameplayProvider"
 
+#if WITH_EDITOR
+// This is copied from EditorUtilitySubsystem (where it is not public)
+// Should probably be somewhere shared
+static UClass* FindBlueprintClass(const FString& TargetNameRaw)
+{
+	FAssetRegistryModule& AssetRegistryModule = FModuleManager::LoadModuleChecked<FAssetRegistryModule>("AssetRegistry");
+	IAssetRegistry& AssetRegistry = AssetRegistryModule.Get();
+	if (AssetRegistry.IsLoadingAssets())
+	{
+		AssetRegistry.SearchAllAssets(true);
+	}
+
+	FString TargetName = TargetNameRaw;
+	TargetName.RemoveFromEnd(TEXT("_C"), ESearchCase::CaseSensitive);
+
+	FARFilter Filter;
+	Filter.bRecursiveClasses = true;
+	Filter.ClassNames.Add(UBlueprintCore::StaticClass()->GetFName());
+
+	// We enumerate all assets to find any blueprints who inherit from native classes directly - or
+	// from other blueprints.
+	UClass* FoundClass = nullptr;
+	AssetRegistry.EnumerateAssets(Filter, [&FoundClass, TargetName](const FAssetData& AssetData)
+	{
+		if ((AssetData.AssetName.ToString() == TargetName) || (AssetData.ObjectPath.ToString() == TargetName))
+		{
+			if (UBlueprint* BP = Cast<UBlueprint>(AssetData.GetAsset()))
+			{
+				FoundClass = BP->GeneratedClass;
+				return false;
+			}
+		}
+
+		return true;
+	});
+
+	return FoundClass;
+}
+
+// This is copied from EditorUtilitySubsystem (where it is not public)
+// Should probably be somewhere shared
+static UClass* FindClassByName(const FString& RawTargetName)
+{
+	FString TargetName = RawTargetName;
+
+	// Check native classes and loaded assets first before resorting to the asset registry
+	bool bIsValidClassName = true;
+	if (TargetName.IsEmpty() || TargetName.Contains(TEXT(" ")))
+	{
+		bIsValidClassName = false;
+	}
+	else if (!FPackageName::IsShortPackageName(TargetName))
+	{
+		if (TargetName.Contains(TEXT(".")))
+		{
+			// Convert type'path' to just path (will return the full string if it doesn't have ' in it)
+			TargetName = FPackageName::ExportTextPathToObjectPath(TargetName);
+
+			FString PackageName;
+			FString ObjectName;
+			TargetName.Split(TEXT("."), &PackageName, &ObjectName);
+
+			const bool bIncludeReadOnlyRoots = true;
+			FText Reason;
+			if (!FPackageName::IsValidLongPackageName(PackageName, bIncludeReadOnlyRoots, &Reason))
+			{
+				bIsValidClassName = false;
+			}
+		}
+		else
+		{
+			bIsValidClassName = false;
+		}
+	}
+
+	UClass* ResultClass = nullptr;
+	if (bIsValidClassName)
+	{
+		if (FPackageName::IsShortPackageName(TargetName))
+		{
+			ResultClass = FindObject<UClass>(ANY_PACKAGE, *TargetName);
+		}
+		else
+		{
+			ResultClass = FindObject<UClass>(nullptr, *TargetName);
+		}
+	}
+
+	// If we still haven't found anything yet, try the asset registry for blueprints that match the requirements
+	if (ResultClass == nullptr)
+	{
+		ResultClass = FindBlueprintClass(TargetName);
+	}
+
+	return ResultClass;
+}
+#endif
+
+
 FGameplayProvider::FGameplayProvider(TraceServices::IAnalysisSession& InSession)
 	: Session(InSession)
 	, EndPlayEvent(nullptr)
 	, PawnPossession(InSession.GetLinearAllocator())
 	, ObjectLifetimes(InSession.GetLinearAllocator())
+	, ObjectRecordingLifetimes(InSession.GetLinearAllocator())
 	, bHasAnyData(false)
 	, bHasObjectProperties(false)
 {
@@ -119,6 +220,19 @@ const FClassInfo* FGameplayProvider::FindClassInfo(uint64 InClassId) const
 	return nullptr;
 }
 
+const UClass* FGameplayProvider::FindClass(uint64 InClassId) const
+{
+#if WITH_EDITOR
+	if (const FClassInfo* ClassInfo = FindClassInfo(InClassId))
+	{
+		return FindClassByName(ClassInfo->Name);
+	}
+	return nullptr;
+#else
+	return nullptr;
+#endif
+}
+
 const FClassInfo* FGameplayProvider::FindClassInfo(const TCHAR* InClassPath) const
 {
 	Session.ReadAccessCheck();
@@ -291,7 +405,7 @@ void FGameplayProvider::AppendObject(uint64 InObjectId, uint64 InOuterId, uint64
 	}
 }
 
-void FGameplayProvider::AppendObjectLifetimeBegin(uint64 InObjectId, double InTime)
+void FGameplayProvider::AppendObjectLifetimeBegin(uint64 InObjectId, double InProfileTime, double InRecordingTime)
 {
 	Session.WriteAccessCheck();
 	bHasAnyData = true;
@@ -300,24 +414,31 @@ void FGameplayProvider::AppendObjectLifetimeBegin(uint64 InObjectId, double InTi
 	{
 		FObjectExistsMessage Message;
 		Message.ObjectId = InObjectId;
-		ActiveObjectLifetimes.Add(InObjectId, ObjectLifetimes.AppendBeginEvent(InTime, Message));
+		ActiveObjectLifetimes.Add(InObjectId, ObjectLifetimes.AppendBeginEvent(InProfileTime, Message));
+		ActiveObjectRecordingLifetimes.Add(InObjectId, ObjectRecordingLifetimes.AppendBeginEvent(InRecordingTime, Message));
 	}
 }
 
-void FGameplayProvider::AppendObjectLifetimeEnd(uint64 InObjectId, double InTime)
+void FGameplayProvider::AppendObjectLifetimeEnd(uint64 InObjectId, double InProfileTime, double InRecordingTime)
 {
 	Session.WriteAccessCheck();
 	bHasAnyData = true;
 
 	if (const uint64 *FoundIndex = ActiveObjectLifetimes.Find(InObjectId))
 	{
-		ObjectLifetimes.EndEvent(*FoundIndex, InTime);
+		ObjectLifetimes.EndEvent(*FoundIndex, InProfileTime);
 		ActiveObjectLifetimes.Remove(InObjectId);
 	}
 
+	if (const uint64 *FoundIndex = ActiveObjectRecordingLifetimes.Find(InObjectId))
+	{
+		ObjectRecordingLifetimes.EndEvent(*FoundIndex, InRecordingTime);
+		// do not remove from ActiveObjectRecordingLifetimes - lifetimes can be queried by Object Id in GetObjectRecordingLifetime
+	}
+
 	if(int32* ObjectInfoIndex = ObjectIdToIndexMap.Find(InObjectId))
 	{
-		OnObjectEndPlayDelegate.Broadcast(InObjectId, InTime, ObjectInfos[*ObjectInfoIndex]);
+		OnObjectEndPlayDelegate.Broadcast(InObjectId, InProfileTime, ObjectInfos[*ObjectInfoIndex]);
 	}
 }
 
@@ -399,6 +520,18 @@ uint64 FGameplayProvider::FindPossessingController(uint64 PawnId, double Time) c
 	return ControllerId;
 }
 
+TRange<double> FGameplayProvider::GetObjectRecordingLifetime(uint64 ObjectId) const
+{
+	if (const uint64 *FoundIndex = ActiveObjectRecordingLifetimes.Find(ObjectId))
+	{
+		return TRange<double>(ObjectRecordingLifetimes.GetEventStartTime(*FoundIndex), ObjectRecordingLifetimes.GetEventEndTime(*FoundIndex));
+	}
+	else
+	{
+		return TRange<double>(0,0);
+	}
+}
+
 void FGameplayProvider::ReadViewTimeline(TFunctionRef<void(const IGameplayProvider::ViewTimeline&)> Callback) const
 {
 	Session.ReadAccessCheck();
diff --git a/Engine/Plugins/Animation/GameplayInsights/Source/GameplayInsights/Private/GameplayProvider.h b/Engine/Plugins/Animation/GameplayInsights/Source/GameplayInsights/Private/GameplayProvider.h
--- a/Engine/Plugins/Animation/GameplayInsights/Source/GameplayInsights/Private/GameplayProvider.h
+++ b/Engine/Plugins/Animation/GameplayInsights/Source/GameplayInsights/Private/GameplayProvider.h
@@ -26,6 +26,7 @@ public:
 	virtual void EnumerateObjects(TFunctionRef<void(const FObjectInfo&)> Callback) const override;
 	virtual void EnumerateObjects(double StartTime, double EndTime, TFunctionRef<void(const FObjectInfo&)> Callback) const override;
 	virtual const FClassInfo* FindClassInfo(uint64 InClassId) const override;
+	virtual const UClass* FindClass(uint64 InClassId) const override;
 	virtual const FClassInfo* FindClassInfo(const TCHAR* InClassPath) const override;
 	virtual const FObjectInfo* FindObjectInfo(uint64 InObjectId) const override;
 	virtual const FWorldInfo* FindWorldInfo(uint64 InObjectId) const override;
@@ -39,6 +40,7 @@ public:
 	virtual const RecordingInfoTimeline* GetRecordingInfo(uint32 RecordingId) const override; 
 	virtual void ReadViewTimeline(TFunctionRef<void(const ViewTimeline&)> Callback) const override;
 
+
 	/** Add a class message */
 	void AppendClass(uint64 InClassId, uint64 InSuperId, const TCHAR* InClassName, const TCHAR* InClassPathName);
 
@@ -46,10 +48,10 @@ public:
 	void AppendObject(uint64 InObjectId, uint64 InOuterId, uint64 InClassId, const TCHAR* InObjectName, const TCHAR* InObjectPathName);
 
 	/** Add an object create message */
-	void AppendObjectLifetimeBegin(uint64 InObjectId, double InTime);
+	void AppendObjectLifetimeBegin(uint64 InObjectId, double InProfileTime, double InRecordingTime);
 	
 	/** Add an object destroy message */
-	void AppendObjectLifetimeEnd(uint64 InObjectId, double InTime);	
+	void AppendObjectLifetimeEnd(uint64 InObjectId, double InProfileTime, double InRecordingTime);	
 
 	/** Add an object event message */
 	void AppendObjectEvent(uint64 InObjectId, double InTime, const TCHAR* InEvent);
@@ -86,6 +88,9 @@ public:
 
 	/** Search PawnPossession timeline to find the Controller Object Id for a Pawn (or 0 if none)*/
 	virtual uint64 FindPossessingController(uint64 Pawn, double Time) const override;
+	
+	/** find the time range for which an object existed */
+	virtual TRange<double> GetObjectRecordingLifetime(uint64 ObjectId) const override;
 
 private:
 	TraceServices::IAnalysisSession& Session;
@@ -157,6 +162,9 @@ private:
 	
 	// Timeline containing intervals where an object exists
 	TraceServices::TIntervalTimeline<FObjectExistsMessage> ObjectLifetimes;
+	
+	// Timeline containing intervals where an object exists - In Rewind Debugger recording time units, to avoid excessive conversions
+	TraceServices::TIntervalTimeline<FObjectExistsMessage> ObjectRecordingLifetimes;
 
 
 	// map from controller id to index in the PawnPossession timeline, for lookup when ending events
@@ -164,6 +172,9 @@ private:
 
 	// map from object id to index in ObjectLifetimes, for lookup when objects are destroyed
 	TMap<uint64, uint64> ActiveObjectLifetimes;
+	
+	// map from object id to index in ObjectRecordingLifetimes, for lookup
+	TMap<uint64, uint64> ActiveObjectRecordingLifetimes;
 
 	/** Whether we have any data */
 	bool bHasAnyData;
diff --git a/Engine/Plugins/Animation/GameplayInsights/Source/GameplayInsights/Private/MontagesTrack.cpp b/Engine/Plugins/Animation/GameplayInsights/Source/GameplayInsights/Private/MontagesTrack.cpp
new file mode 100644
--- /dev/null
+++ b/Engine/Plugins/Animation/GameplayInsights/Source/GameplayInsights/Private/MontagesTrack.cpp
@@ -0,0 +1,328 @@
+// Copyright Epic Games, Inc. All Rights Reserved.
+
+#include "MontagesTrack.h"
+#include "SCurveTimelineView.h"
+#include "IRewindDebugger.h"
+#include "GameplayProvider.h"
+#include "AnimationProvider.h"
+#include "Styling/SlateIconFinder.h"
+#include "SMontageView.h"
+#if WITH_EDITOR
+#include "Animation/AnimMontage.h"
+#endif
+
+#define LOCTEXT_NAMESPACE "MontagesTrack"
+
+namespace RewindDebugger
+{
+
+FMontagesTrack::FMontagesTrack(uint64 InObjectId) : ObjectId(InObjectId)
+{
+#if WITH_EDITOR
+	Icon = FSlateIconFinder::FindIconForClass(UAnimMontage::StaticClass());
+#endif
+}
+
+void FMontagesTrack::IterateSubTracksInternal(TFunction<void(TSharedPtr<FRewindDebuggerTrack> SubTrack)> IteratorFunction)
+{
+	for(TSharedPtr<FMontageTrack>& Track : Children)
+	{
+		IteratorFunction(Track);
+	}
+};
+
+bool FMontagesTrack::UpdateInternal()
+{
+	TArray<uint64> UniqueTrackIds;
+
+	IRewindDebugger* RewindDebugger = IRewindDebugger::Instance();
+	const TraceServices::IAnalysisSession* AnalysisSession = RewindDebugger->GetAnalysisSession();
+
+	bool bChanged = false;
+	
+	// convert time range to from rewind debugger times to profiler times
+	TRange<double> TraceTimeRange = RewindDebugger->GetCurrentTraceRange();
+	double StartTime = TraceTimeRange.GetLowerBoundValue();
+	double EndTime = TraceTimeRange.GetUpperBoundValue();
+	
+	// count number of unique animations in the current time range
+	const FGameplayProvider* GameplayProvider = AnalysisSession->ReadProvider<FGameplayProvider>(FGameplayProvider::ProviderName);
+	const FAnimationProvider* AnimationProvider = AnalysisSession->ReadProvider<FAnimationProvider>(FAnimationProvider::ProviderName);
+
+	int AnimationCount = 0;
+
+	if(GameplayProvider && AnimationProvider)
+	{
+		TraceServices::FAnalysisSessionReadScope SessionReadScope(*AnalysisSession);
+		UniqueTrackIds.SetNum(0, false);
+
+		AnimationProvider->ReadMontageTimeline(ObjectId, [&UniqueTrackIds,&GameplayProvider, StartTime, EndTime](const FAnimationProvider::AnimMontageTimeline& InTimeline)
+		{
+			// this isn't very efficient, and it gets called every frame.  will need optimizing
+			InTimeline.EnumerateEvents(StartTime, EndTime, [&UniqueTrackIds, StartTime, EndTime](double InStartTime, double InEndTime, uint32 InDepth, const FAnimMontageMessage& InMessage)
+			{
+				if (InEndTime > StartTime && InStartTime < EndTime)
+				{
+					UniqueTrackIds.AddUnique(InMessage.MontageId);
+				}
+				return TraceServices::EEventEnumerate::Continue;
+			});
+		});
+
+		UniqueTrackIds.StableSort();
+		
+		AnimationCount = UniqueTrackIds.Num();
+
+		if (Children.Num()!=AnimationCount)
+			bChanged = true;
+		
+		Children.SetNum(AnimationCount);
+		for(int i = 0; i < AnimationCount; i++)
+		{
+			if (!Children[i].IsValid() || Children[i].Get()->GetAssetId() != UniqueTrackIds[i])
+			{
+				Children[i] = MakeShared<FMontageTrack>(ObjectId, UniqueTrackIds[i]);
+				bChanged = true;
+			}
+
+			bChanged = bChanged || Children[i]->Update();
+		}
+	}
+
+	return bChanged;
+}
+
+TSharedPtr<SWidget> FMontagesTrack::GetDetailsViewInternal()
+{
+	IRewindDebugger* RewindDebugger = IRewindDebugger::Instance();
+	return SNew(SMontageView, ObjectId, RewindDebugger->CurrentTraceTime(), *RewindDebugger->GetAnalysisSession())
+		.CurrentTime_Lambda([RewindDebugger]{ return RewindDebugger->CurrentTraceTime(); });
+}
+
+
+FMontageTrack::FMontageTrack(uint64 InObjectId, uint64 InAssetId, FMontageTrack::ECurveType InCurveType) :
+	 AssetId(InAssetId)
+	, CurveType(InCurveType)
+	, ObjectId(InObjectId)
+{
+	SetIsExpanded(false);
+}
+
+
+void FMontageTrack::IterateSubTracksInternal(TFunction<void(TSharedPtr<FRewindDebuggerTrack> SubTrack)> IteratorFunction)
+{
+	for (TSharedPtr<FMontageTrack>& Track : Children)
+	{
+		IteratorFunction(Track);
+	}
+};
+
+TSharedPtr<SCurveTimelineView::FTimelineCurveData> FMontageTrack::GetCurveData() const
+{
+	if (!CurveData.IsValid())
+	{
+		CurveData = MakeShared<SCurveTimelineView::FTimelineCurveData>();
+	}
+	
+	CurvesUpdateRequested++;
+	
+	return CurveData;
+}
+
+bool FMontageTrack::UpdateInternal()
+{
+	IRewindDebugger* RewindDebugger = IRewindDebugger::Instance();
+	// compute curve points
+	//
+	TRange<double> TraceTimeRange = RewindDebugger->GetCurrentTraceRange();
+	double StartTime = TraceTimeRange.GetLowerBoundValue();
+	double EndTime = TraceTimeRange.GetUpperBoundValue();
+
+	const TraceServices::IAnalysisSession* AnalysisSession = RewindDebugger->GetAnalysisSession();
+	const FGameplayProvider* GameplayProvider = AnalysisSession->ReadProvider<FGameplayProvider>(FGameplayProvider::ProviderName);
+	const FAnimationProvider* AnimationProvider = AnalysisSession->ReadProvider<FAnimationProvider>(FAnimationProvider::ProviderName);
+
+	int AnimationCount = 0;
+
+	if(CurvesUpdateRequested > 10 && GameplayProvider && AnimationProvider)
+	{
+		auto& CurvePoints = CurveData->Points;
+		CurvePoints.SetNum(0,false);
+		
+		TraceServices::FAnalysisSessionReadScope SessionReadScope(*AnalysisSession);
+		
+		AnimationProvider->ReadMontageTimeline(ObjectId, [AnalysisSession, StartTime, EndTime, &CurvePoints, this](const FAnimationProvider::AnimMontageTimeline& InTimeline)
+		{
+			InTimeline.EnumerateEvents(StartTime, EndTime, [&CurvePoints, AnalysisSession, this](double InStartTime, double InEndTime, uint32 InDepth, const FAnimMontageMessage& InMessage)
+			{
+				if (InMessage.MontageId == AssetId)
+				{
+					float Weight = 0;
+					switch (CurveType)
+					{
+						case ECurveType::BlendWeight:					Weight = InMessage.Weight; break;
+						case ECurveType::DesiredWeight:					Weight = InMessage.DesiredWeight; break; 
+						case ECurveType::Position:						Weight = InMessage.Position; break; 
+					}
+					
+					CurvePoints.Add({ InMessage.RecordingTime,	Weight });
+				}
+								
+				return TraceServices::EEventEnumerate::Continue;
+			});
+		});
+
+		CurvesUpdateRequested = 0;
+	}
+
+	// update Icon:
+	
+	bool bChanged = false;
+
+	if (CurveType == ECurveType::BlendWeight)
+	{
+		// Blend Weight track gets name/icon of animation
+		if (CurveName.IsEmpty())
+		{
+			if (const FObjectInfo* ObjectInfo = GameplayProvider->FindObjectInfo(AssetId))
+			{
+				CurveName = FText::FromString(ObjectInfo->Name);
+				bChanged = true;
+
+				if (const UClass* FoundClass = GameplayProvider->FindClass(ObjectInfo->ClassId))
+				{
+					Icon = FSlateIconFinder::FindIconForClass(FoundClass);
+				}
+
+				if (Children.Num() == 0)
+				{
+					for (int i = static_cast<int>(ECurveType::BlendWeight) + 1; i <= static_cast<int>(ECurveType::Position); i++)
+					{
+						Children.Add(MakeShared<FMontageTrack>(ObjectId, AssetId, static_cast<ECurveType>(i)));
+					}
+				}
+			}
+		}
+	}
+	else
+	{
+		// other tracks get the curve name
+
+		if (CurveName.IsEmpty())
+		{
+			switch (CurveType)
+			{
+				case ECurveType::BlendWeight:					break;
+				case ECurveType::DesiredWeight:					CurveName = LOCTEXT("Desired Weight", "Desired Weight"); break;
+				case ECurveType::Position:						CurveName = LOCTEXT("Position", "Position"); break;
+			}
+
+			Icon = FSlateIcon("EditorStyle", "AnimGraph.Attribute.Curves.Icon", "AnimGraph.Attribute.Curves.Icon");
+
+			bChanged = true;
+		}
+	}
+
+	for (auto& Child : Children)
+	{
+		bChanged |= Child->Update();
+	}
+
+	return bChanged;
+}
+
+static FLinearColor MakeMontageCurveColor(uint32 InSeed, bool bInLine = false)
+{
+	FRandomStream Stream(InSeed);
+	const uint8 Hue = (uint8)(Stream.FRand() * 255.0f);
+	const uint8 SatVal = bInLine ? 196 : 128;
+	return FLinearColor::MakeFromHSV8(Hue, SatVal, SatVal);
+}
+
+TSharedPtr<SWidget> FMontageTrack::GetTimelineViewInternal()
+{
+	FLinearColor Color;
+	switch(CurveType)
+	{
+	case ECurveType::BlendWeight:
+		Color = MakeMontageCurveColor(CityHash32(reinterpret_cast<char*>(&AssetId), 8));
+		Color.A = 0.5f;
+		break;
+	case ECurveType::DesiredWeight:
+		Color = FLinearColor::MakeFromHSV8(0, 50, 50);
+		break;
+	case ECurveType::Position:
+		Color = FLinearColor::MakeFromHSV8(0, 50, 50);
+		break;
+	}
+
+	FLinearColor CurveColor = Color;
+	CurveColor.R *= 0.5;
+	CurveColor.G *= 0.5;
+	CurveColor.B *= 0.5;
+
+	TSharedPtr<SCurveTimelineView> CurveTimelineView = SNew(SCurveTimelineView)
+		.FillColor(Color)
+		.CurveColor(CurveColor)
+		.ViewRange_Lambda([]() { return IRewindDebugger::Instance()->GetCurrentViewRange(); })
+		.RenderFill(CurveType == ECurveType::BlendWeight)
+		.CurveData_Raw(this, &FMontageTrack::GetCurveData);
+
+	if (CurveType == ECurveType::BlendWeight || CurveType == ECurveType::DesiredWeight)
+	{
+		CurveTimelineView->SetFixedRange(0, 1);
+	}
+
+	return CurveTimelineView;
+}
+
+TSharedPtr<SWidget> FMontageTrack::GetDetailsViewInternal()
+{
+	IRewindDebugger* RewindDebugger = IRewindDebugger::Instance();
+	TSharedPtr<SMontageView> MontageView = SNew(SMontageView, ObjectId, RewindDebugger->CurrentTraceTime(), *RewindDebugger->GetAnalysisSession())
+			.CurrentTime_Lambda([RewindDebugger]{ return RewindDebugger->CurrentTraceTime(); });
+	MontageView->SetAssetFilter(AssetId);
+	return MontageView;
+}
+
+
+FName FMontagesTrackCreator::GetTargetTypeNameInternal() const
+{
+	static const FName TargetTypeName = "AnimInstance";
+	return TargetTypeName;
+}
+
+static const FName MontagesName("Montages");
+FName FMontagesTrackCreator::GetNameInternal() const
+{
+	return MontagesName;
+}
+	
+FName FMontagesTrack::GetNameInternal() const
+{
+	return MontagesName;
+}
+
+TSharedPtr<RewindDebugger::FRewindDebuggerTrack> FMontagesTrackCreator::CreateTrackInternal(uint64 ObjectId) const
+{
+ 	return MakeShared<RewindDebugger::FMontagesTrack>(ObjectId);
+}
+
+bool FMontagesTrackCreator::HasDebugInfoInternal(uint64 ObjectId) const
+{
+	const TraceServices::IAnalysisSession* AnalysisSession = IRewindDebugger::Instance()->GetAnalysisSession();
+	TraceServices::FAnalysisSessionReadScope SessionReadScope(*AnalysisSession);
+	bool bHasData = false;
+	if (const FAnimationProvider* AnimationProvider = AnalysisSession->ReadProvider<FAnimationProvider>(FAnimationProvider::ProviderName))
+	{
+		AnimationProvider->ReadMontageTimeline(ObjectId, [&bHasData](const FAnimationProvider::AnimMontageTimeline& InTimeline)
+		{
+			bHasData = true;
+		});
+	}
+	
+	return bHasData;
+}
+
+}
+#undef LOCTEXT_NAMESPACE
diff --git a/Engine/Plugins/Animation/GameplayInsights/Source/GameplayInsights/Private/MontagesTrack.h b/Engine/Plugins/Animation/GameplayInsights/Source/GameplayInsights/Private/MontagesTrack.h
new file mode 100644
--- /dev/null
+++ b/Engine/Plugins/Animation/GameplayInsights/Source/GameplayInsights/Private/MontagesTrack.h
@@ -0,0 +1,85 @@
+// Copyright Epic Games, Inc. All Rights Reserved.
+#pragma once
+
+#include "IAnimationProvider.h"
+#include "IRewindDebugger.h"
+#include "RewindDebuggerTrack.h"
+#include "SCurveTimelineView.h"
+#include "Textures/SlateIcon.h"
+#include "IRewindDebuggerTrackCreator.h"
+
+namespace RewindDebugger
+{
+
+class FMontageTrack : public FRewindDebuggerTrack
+{
+public:
+	enum class ECurveType
+	{
+		BlendWeight,
+		DesiredWeight,
+		Position
+	};
+	
+	FMontageTrack(uint64 InObjectId, uint64 AssetId, ECurveType CurveType = ECurveType::BlendWeight);
+
+	uint64 GetAssetId() const { return AssetId; }
+
+private:
+	virtual bool UpdateInternal() override;
+	virtual TSharedPtr<SWidget> GetTimelineViewInternal() override;
+	virtual TSharedPtr<SWidget> GetDetailsViewInternal() override;
+	virtual void IterateSubTracksInternal(TFunction<void(TSharedPtr<FRewindDebuggerTrack> SubTrack)> IteratorFunction) override;
+	
+	virtual FSlateIcon GetIconInternal() override { return Icon; }
+	virtual FName GetNameInternal() const override { return "Montage"; }
+	virtual FText GetDisplayNameInternal() const override { return CurveName; }
+	virtual uint64 GetObjectIdInternal() const override { return ObjectId; }
+
+	TSharedPtr<SCurveTimelineView::FTimelineCurveData> GetCurveData() const;
+	
+	uint64 AssetId;
+
+	ECurveType CurveType;
+	
+	FSlateIcon Icon;
+	FText CurveName;
+	uint64 ObjectId;
+
+	mutable TSharedPtr<SCurveTimelineView::FTimelineCurveData> CurveData;
+	mutable int CurvesUpdateRequested = 0;
+	
+	TArray<TSharedPtr<FMontageTrack>> Children;
+};
+
+class FMontagesTrack : public FRewindDebuggerTrack
+{
+public:
+	FMontagesTrack(uint64 InObjectId);
+
+private:
+	virtual bool UpdateInternal() override;
+	virtual TSharedPtr<SWidget> GetDetailsViewInternal() override;
+	virtual void IterateSubTracksInternal(TFunction<void(TSharedPtr<FRewindDebuggerTrack> SubTrack)> IteratorFunction) override;
+	
+	virtual FName GetNameInternal() const override;
+	virtual FText GetDisplayNameInternal() const override { return NSLOCTEXT("RewindDebugger", "Montages Track Name", "Montages"); }
+	virtual uint64 GetObjectIdInternal() const override { return ObjectId; }
+
+	virtual FSlateIcon GetIconInternal() override { return Icon; }
+
+	FSlateIcon Icon;
+	uint64 ObjectId;
+	TArray<TSharedPtr<FMontageTrack>> Children;
+};
+	
+class FMontagesTrackCreator : public IRewindDebuggerTrackCreator
+{
+public:
+	virtual FName GetTargetTypeNameInternal() const;
+	virtual FName GetNameInternal() const override;
+	virtual TSharedPtr<FRewindDebuggerTrack> CreateTrackInternal(uint64 ObjectId) const override;
+	virtual bool HasDebugInfoInternal(uint64 ObjectId) const override;
+};
+
+}
\ No newline at end of file
diff --git a/Engine/Plugins/Animation/GameplayInsights/Source/GameplayInsights/Private/SAnimGraphSchematicView.cpp b/Engine/Plugins/Animation/GameplayInsights/Source/GameplayInsights/Private/SAnimGraphSchematicView.cpp
--- a/Engine/Plugins/Animation/GameplayInsights/Source/GameplayInsights/Private/SAnimGraphSchematicView.cpp
+++ b/Engine/Plugins/Animation/GameplayInsights/Source/GameplayInsights/Private/SAnimGraphSchematicView.cpp
@@ -10,6 +10,7 @@
 #include "Widgets/SBoxPanel.h"
 #include "Styling/CoreStyle.h"
 #include "GameplayInsightsStyle.h"
+#include "IRewindDebugger.h"
 #include "Widgets/Layout/SScrollBorder.h"
 #include "Framework/MultiBox/MultiBoxBuilder.h"
 #include "Widgets/Input/SEditableTextBox.h"
@@ -1033,4 +1034,21 @@ TSharedPtr<IRewindDebuggerView> FAnimGraphSchematicViewCreator::CreateDebugView(
 	return SNew(SAnimGraphSchematicView, ObjectId, CurrentTime, AnalysisSession);
 }
 
+bool FAnimGraphSchematicViewCreator::HasDebugInfo(uint64 ObjectId) const
+{
+	const TraceServices::IAnalysisSession* AnalysisSession = IRewindDebugger::Instance()->GetAnalysisSession();
+	
+	TraceServices::FAnalysisSessionReadScope SessionReadScope(*AnalysisSession);
+	bool bHasData = false;
+	if (const FAnimationProvider* AnimationProvider = AnalysisSession->ReadProvider<FAnimationProvider>(FAnimationProvider::ProviderName))
+	{
+		AnimationProvider->ReadAnimGraphTimeline(ObjectId, [&bHasData](const FAnimationProvider::AnimGraphTimeline& InGraphTimeline)
+		{
+			bHasData = true;
+		});
+	}
+	
+	return bHasData;
+}
+
 #undef LOCTEXT_NAMESPACE
diff --git a/Engine/Plugins/Animation/GameplayInsights/Source/GameplayInsights/Private/SAnimGraphSchematicView.h b/Engine/Plugins/Animation/GameplayInsights/Source/GameplayInsights/Private/SAnimGraphSchematicView.h
--- a/Engine/Plugins/Animation/GameplayInsights/Source/GameplayInsights/Private/SAnimGraphSchematicView.h
+++ b/Engine/Plugins/Animation/GameplayInsights/Source/GameplayInsights/Private/SAnimGraphSchematicView.h
@@ -124,4 +124,5 @@ class FAnimGraphSchematicViewCreator : public IRewindDebuggerViewCreator
 		virtual FText GetTitle() const override;
 		virtual FSlateIcon GetIcon() const override;
 		virtual TSharedPtr<IRewindDebuggerView> CreateDebugView(uint64 ObjectId, double CurrentTime, const TraceServices::IAnalysisSession& InAnalysisSession) const override;
+		virtual bool HasDebugInfo(uint64 ObjectId) const override;
 };
\ No newline at end of file
diff --git a/Engine/Plugins/Animation/GameplayInsights/Source/GameplayInsights/Private/SAnimationCurvesView.cpp b/Engine/Plugins/Animation/GameplayInsights/Source/GameplayInsights/Private/SAnimationCurvesView.cpp
--- a/Engine/Plugins/Animation/GameplayInsights/Source/GameplayInsights/Private/SAnimationCurvesView.cpp
+++ b/Engine/Plugins/Animation/GameplayInsights/Source/GameplayInsights/Private/SAnimationCurvesView.cpp
@@ -4,14 +4,9 @@
 #include "AnimationProvider.h"
 #include "TraceServices/Model/AnalysisSession.h"
 #include "GameplayProvider.h"
-#include "Styling/SlateIconFinder.h"
 #include "TraceServices/Model/Frames.h"
 #include "VariantTreeNode.h"
 
-#if WITH_EDITOR
-#include "Animation/AnimInstance.h"
-#endif
-
 #define LOCTEXT_NAMESPACE "SAnimationCurvesView"
 
 void SAnimationCurvesView::GetVariantsAtFrame(const TraceServices::FFrame& InFrame, TArray<TSharedRef<FVariantTreeNode>>& OutVariants) const
@@ -25,16 +20,19 @@ void SAnimationCurvesView::GetVariantsAtFrame(const TraceServices::FFrame& InFra
 
 		TSharedRef<FVariantTreeNode> Header = OutVariants.Add_GetRef(FVariantTreeNode::MakeHeader(LOCTEXT("Animation Curves", "Animation Curves"), 0));
 
-		AnimationProvider->ReadSkeletalMeshPoseTimeline(ObjectId, [&Header, &AnimationProvider, &InFrame](const FAnimationProvider::SkeletalMeshPoseTimeline& InTimeline, bool bInHasCurves)
+		AnimationProvider->ReadSkeletalMeshPoseTimeline(ObjectId, [this, &Header, &AnimationProvider, &InFrame](const FAnimationProvider::SkeletalMeshPoseTimeline& InTimeline, bool bInHasCurves)
 		{
-			InTimeline.EnumerateEvents(InFrame.StartTime, InFrame.EndTime, [&Header, &AnimationProvider, &InFrame](double InStartTime, double InEndTime, uint32 InDepth, const FSkeletalMeshPoseMessage& InMessage)
+			InTimeline.EnumerateEvents(InFrame.StartTime, InFrame.EndTime, [this ,&Header, &AnimationProvider, &InFrame](double InStartTime, double InEndTime, uint32 InDepth, const FSkeletalMeshPoseMessage& InMessage)
 			{
 				if(InStartTime >= InFrame.StartTime && InStartTime <= InFrame.EndTime)
 				{
-					AnimationProvider->EnumerateSkeletalMeshCurves(InMessage, [&Header, &AnimationProvider](const FSkeletalMeshNamedCurve& InCurve)
+					AnimationProvider->EnumerateSkeletalMeshCurves(InMessage, [this, &Header, &AnimationProvider](const FSkeletalMeshNamedCurve& InCurve)
 					{
-						const TCHAR* CurveName = AnimationProvider->GetName(InCurve.Id);
-						Header->AddChild(FVariantTreeNode::MakeFloat(FText::FromString(CurveName), InCurve.Value));
+						if (!bCurveFilterSet || InCurve.Id == CurveFilter)
+						{
+							const TCHAR* CurveName = AnimationProvider->GetName(InCurve.Id);
+							Header->AddChild(FVariantTreeNode::MakeFloat(FText::FromString(CurveName), InCurve.Value));
+						}
 					});
 				}
 				return TraceServices::EEventEnumerate::Continue;
@@ -43,42 +41,10 @@ void SAnimationCurvesView::GetVariantsAtFrame(const TraceServices::FFrame& InFra
 	}
 }
 
-FName FAnimationCurvesViewCreator::GetTargetTypeName() const
-{
-	static FName TargetTypeName = "SkeletalMeshComponent";
-	return TargetTypeName;
-}
-
 static const FName AnimationCurvesName("AnimationCurves");
-
 FName SAnimationCurvesView::GetName() const
 {
 	return AnimationCurvesName;
 }
 
-FName FAnimationCurvesViewCreator::GetName() const
-{
-	return AnimationCurvesName;
-}
-
-FText FAnimationCurvesViewCreator::GetTitle() const
-{
-	return LOCTEXT("Animation Curves", "Animation Curves");
-}
-
-FSlateIcon FAnimationCurvesViewCreator::GetIcon() const
-{
-#if WITH_EDITOR
-	return FSlateIconFinder::FindIconForClass(UAnimInstance::StaticClass());
-#else
-	return FSlateIcon();
-#endif
-}
-
-TSharedPtr<IRewindDebuggerView> FAnimationCurvesViewCreator::CreateDebugView(uint64 ObjectId, double CurrentTime, const TraceServices::IAnalysisSession& AnalysisSession) const
-{
-	return SNew(SAnimationCurvesView, ObjectId, CurrentTime, AnalysisSession);
-}
-
-
 #undef LOCTEXT_NAMESPACE
diff --git a/Engine/Plugins/Animation/GameplayInsights/Source/GameplayInsights/Private/SAnimationCurvesView.h b/Engine/Plugins/Animation/GameplayInsights/Source/GameplayInsights/Private/SAnimationCurvesView.h
--- a/Engine/Plugins/Animation/GameplayInsights/Source/GameplayInsights/Private/SAnimationCurvesView.h
+++ b/Engine/Plugins/Animation/GameplayInsights/Source/GameplayInsights/Private/SAnimationCurvesView.h
@@ -3,8 +3,6 @@
 #pragma once
 
 #include "CoreMinimal.h"
-#include "IRewindDebuggerView.h"
-#include "IRewindDebuggerViewCreator.h"
 #include "SPropertiesDebugViewBase.h"
 
 namespace TraceServices { class IAnalysisSession; }
@@ -14,15 +12,19 @@ class SAnimationCurvesView : public SPropertiesDebugViewBase
 public:
 	virtual void GetVariantsAtFrame(const TraceServices::FFrame& InFrame, TArray<TSharedRef<FVariantTreeNode>>& OutVariants) const override;
 	virtual FName GetName() const override;
-};
+	
+	void SetCurveFilter(uint32 InCurveFilter)
+	{
+		CurveFilter = InCurveFilter;
+		bCurveFilterSet = true;
+	}
 
-class FAnimationCurvesViewCreator : public IRewindDebuggerViewCreator
-{
-	public:
+	void ClearCurveFilter()
+	{
+		bCurveFilterSet = false;
+	}
 
-		virtual FName GetTargetTypeName() const;
-		virtual FName GetName() const override;
-		virtual FText GetTitle() const override;
-		virtual FSlateIcon GetIcon() const override;
-		virtual TSharedPtr<IRewindDebuggerView> CreateDebugView(uint64 ObjectId, double CurrentTime, const TraceServices::IAnalysisSession& InAnalysisSession) const override;
+private:
+	uint32 CurveFilter = 0;
+	bool bCurveFilterSet = false;
 };
\ No newline at end of file
diff --git a/Engine/Plugins/Animation/GameplayInsights/Source/GameplayInsights/Private/SBlendWeightsView.cpp b/Engine/Plugins/Animation/GameplayInsights/Source/GameplayInsights/Private/SBlendWeightsView.cpp
--- a/Engine/Plugins/Animation/GameplayInsights/Source/GameplayInsights/Private/SBlendWeightsView.cpp
+++ b/Engine/Plugins/Animation/GameplayInsights/Source/GameplayInsights/Private/SBlendWeightsView.cpp
@@ -4,14 +4,9 @@
 #include "AnimationProvider.h"
 #include "TraceServices/Model/AnalysisSession.h"
 #include "GameplayProvider.h"
-#include "Styling/SlateIconFinder.h"
 #include "TraceServices/Model/Frames.h"
 #include "VariantTreeNode.h"
 
-#if WITH_EDITOR
-#include "Animation/AnimInstance.h"
-#endif
-
 #define LOCTEXT_NAMESPACE "SBlendWeightsView"
 
 void SBlendWeightsView::GetVariantsAtFrame(const TraceServices::FFrame& InFrame, TArray<TSharedRef<FVariantTreeNode>>& OutVariants) const
@@ -23,12 +18,20 @@ void SBlendWeightsView::GetVariantsAtFrame(const TraceServices::FFrame& InFrame,
 	{
 		TraceServices::FAnalysisSessionReadScope SessionReadScope(*AnalysisSession);
 
-		AnimationProvider->ReadTickRecordTimeline(ObjectId, [&GameplayProvider, &OutVariants, &InFrame](const FAnimationProvider::TickRecordTimeline& InTimeline)
+		AnimationProvider->ReadTickRecordTimeline(ObjectId, [this, &GameplayProvider, &OutVariants, &InFrame](const FAnimationProvider::TickRecordTimeline& InTimeline)
 		{
-			InTimeline.EnumerateEvents(InFrame.StartTime, InFrame.EndTime, [&GameplayProvider, &OutVariants, &InFrame](double InStartTime, double InEndTime, uint32 InDepth, const FTickRecordMessage& InMessage)
+			InTimeline.EnumerateEvents(InFrame.StartTime, InFrame.EndTime, [this, &GameplayProvider, &OutVariants, &InFrame](double InStartTime, double InEndTime, uint32 InDepth, const FTickRecordMessage& InMessage)
 			{
 				if(InStartTime >= InFrame.StartTime && InEndTime <= InFrame.EndTime)
 				{
+					if (bAssetFilterSet)
+					{
+						if (InMessage.NodeId != NodeIdFilter || InMessage.AssetId != AssetIdFilter)
+						{
+							return TraceServices::EEventEnumerate::Continue;
+						}
+					}
+					
 					const FClassInfo& ClassInfo = GameplayProvider->GetClassInfoFromObject(InMessage.AssetId);
 					TSharedRef<FVariantTreeNode> Header = OutVariants.Add_GetRef(FVariantTreeNode::MakeObject(FText::FromString(ClassInfo.Name), InMessage.AssetId, InMessage.AssetId));
 					
@@ -50,7 +53,6 @@ void SBlendWeightsView::GetVariantsAtFrame(const TraceServices::FFrame& InFrame,
 	}
 }
 
-
 static const FName BlendWeightsName("BlendWeights");
 
 FName SBlendWeightsView::GetName() const
@@ -58,35 +60,4 @@ FName SBlendWeightsView::GetName() const
 	return BlendWeightsName;
 }
 
-FName FBlendWeightsViewCreator::GetTargetTypeName() const
-{
-	static FName TargetTypeName = "AnimInstance";
-	return TargetTypeName;
-}
-
-FName FBlendWeightsViewCreator::GetName() const
-{
-	return BlendWeightsName;
-}
-
-FText FBlendWeightsViewCreator::GetTitle() const
-{
-	return LOCTEXT("Blend Weights", "Blend Weights");
-}
-
-FSlateIcon FBlendWeightsViewCreator::GetIcon() const
-{
-#if WITH_EDITOR
-	return FSlateIconFinder::FindIconForClass(UAnimInstance::StaticClass());
-#else
-	return FSlateIcon();
-#endif
-}
-
-TSharedPtr<IRewindDebuggerView> FBlendWeightsViewCreator::CreateDebugView(uint64 ObjectId, double CurrentTime, const TraceServices::IAnalysisSession& AnalysisSession) const
-{
-	return SNew(SBlendWeightsView, ObjectId, CurrentTime, AnalysisSession);
-}
-
-
 #undef LOCTEXT_NAMESPACE
diff --git a/Engine/Plugins/Animation/GameplayInsights/Source/GameplayInsights/Private/SBlendWeightsView.h b/Engine/Plugins/Animation/GameplayInsights/Source/GameplayInsights/Private/SBlendWeightsView.h
--- a/Engine/Plugins/Animation/GameplayInsights/Source/GameplayInsights/Private/SBlendWeightsView.h
+++ b/Engine/Plugins/Animation/GameplayInsights/Source/GameplayInsights/Private/SBlendWeightsView.h
@@ -3,8 +3,6 @@
 #pragma once
 
 #include "CoreMinimal.h"
-#include "IRewindDebuggerView.h"
-#include "IRewindDebuggerViewCreator.h"
 #include "SPropertiesDebugViewBase.h"
 
 namespace TraceServices { class IAnalysisSession; }
@@ -14,14 +12,21 @@ class SBlendWeightsView : public SPropertiesDebugViewBase
 public:
 	virtual void GetVariantsAtFrame(const TraceServices::FFrame& InFrame, TArray<TSharedRef<FVariantTreeNode>>& OutVariants) const override;
 	virtual FName GetName() const override;
-};
 
-class FBlendWeightsViewCreator : public IRewindDebuggerViewCreator
-{
-	public:
-		virtual FName GetTargetTypeName() const;
-		virtual FName GetName() const override;
-		virtual FText GetTitle() const override;
-		virtual FSlateIcon GetIcon() const override;
-		virtual TSharedPtr<IRewindDebuggerView> CreateDebugView(uint64 ObjectId, double CurrentTime, const TraceServices::IAnalysisSession& InAnalysisSession) const override;
+	void SetAssetFilter(uint64 InAssetIdFilter, uint32 InNodeIdFilter)
+	{
+		AssetIdFilter = InAssetIdFilter;
+		NodeIdFilter = InNodeIdFilter;
+		bAssetFilterSet = true;
+	}
+
+	void ClearAssetFilter()
+	{
+		bAssetFilterSet = false;
+	}
+
+private:
+	uint64 AssetIdFilter = 0;
+	uint32 NodeIdFilter = 0;
+	bool bAssetFilterSet = false;
 };
\ No newline at end of file
diff --git a/Engine/Plugins/Animation/GameplayInsights/Source/GameplayInsights/Private/SCurveTimelineView.cpp b/Engine/Plugins/Animation/GameplayInsights/Source/GameplayInsights/Private/SCurveTimelineView.cpp
new file mode 100644
--- /dev/null
+++ b/Engine/Plugins/Animation/GameplayInsights/Source/GameplayInsights/Private/SCurveTimelineView.cpp
@@ -0,0 +1,189 @@
+// Copyright Epic Games, Inc. All Rights Reserved.
+
+#include "SCurveTimelineView.h"
+#include "Framework/Application/SlateApplication.h"
+#include "Rendering/DrawElements.h"
+#include "SSimpleTimeSlider.h"
+
+#define LOCTEXT_NAMESPACE "SCurveTimelineView"
+
+
+void SCurveTimelineView::Construct( const SCurveTimelineView::FArguments& InArgs )
+{
+    ViewRange = InArgs._ViewRange;
+	CurveColor = InArgs._CurveColor;
+	FillColor = InArgs._FillColor;
+	RenderFill = InArgs._RenderFill;
+	DesiredSize = InArgs._DesiredSize;
+	CurveData = InArgs._CurveData;
+
+	// set clipping on by default, since the OnPaint function is drawing outside the bounds
+	Clipping = EWidgetClipping::ClipToBounds;
+}
+
+int32 SCurveTimelineView::OnPaint( const FPaintArgs& Args, const FGeometry& AllottedGeometry, const FSlateRect& MyCullingRect, FSlateWindowElementList& OutDrawElements, int32 LayerId, const FWidgetStyle& InWidgetStyle, bool bParentEnabled ) const
+{
+	int32 NewLayer = PaintCurve( AllottedGeometry, MyCullingRect, OutDrawElements, LayerId, InWidgetStyle, bParentEnabled );
+
+	return FMath::Max( NewLayer, SCompoundWidget::OnPaint(Args, AllottedGeometry, MyCullingRect, OutDrawElements, NewLayer, InWidgetStyle, ShouldBeEnabled( bParentEnabled ) ) );
+}
+
+
+FVector2D SCurveTimelineView::ComputeDesiredSize( float ) const
+{
+	return DesiredSize.Get();
+}
+
+int32 SCurveTimelineView::PaintCurve(const FGeometry& AllottedGeometry, const FSlateRect& MyCullingRect, FSlateWindowElementList& OutDrawElements, int32 LayerId, const FWidgetStyle& InWidgetStyle, bool bParentEnabled ) const
+{
+	static const float LargeFrameTime = 0.125;  // ideally, we could check the recorded frame data for actual frame length
+	 
+	// convert time range to from rewind debugger times to profiler times
+	TRange<double> DebugTimeRange = ViewRange.Get();
+	FLinearColor LineColor = CurveColor.Get();
+	
+	SSimpleTimeSlider::FScrubRangeToScreen RangeToScreen( DebugTimeRange, AllottedGeometry.GetLocalSize() );
+	FVector2D Size = AllottedGeometry.GetLocalSize();
+
+	if (TSharedPtr<FTimelineCurveData> Curve = CurveData.Get())
+	{
+		// Min/Max for normalization
+		// maybe expose these as attributes too, so they are recomputed less frequently
+		float MinValue = 1.0e10;
+		float MaxValue = -1.0e10;
+
+		if (bUseFixedRange)
+		{
+			MinValue = FixedRangeMin;
+			MaxValue = FixedRangeMax;
+		}
+		else
+		{
+			for(auto& Point : Curve->Points)
+			{
+				MinValue = FMath::Min(MinValue, Point.Value);
+				MaxValue = FMath::Max(MaxValue, Point.Value);
+			}
+		}
+		
+		float Range = FMath::Max(1.0,MaxValue - MinValue);
+		
+		auto &CurvePoints = Curve->Points;
+		const int NumPoints = CurvePoints.Num();
+		
+		const ESlateDrawEffect LineDrawEffects = ESlateDrawEffect::NoPixelSnapping;	
+
+		if (NumPoints > 0)
+		{
+			if (RenderFill.Get())
+			{
+				// don't call this every frame
+				FSlateResourceHandle ResourceHandle = FSlateApplication::Get().GetRenderer()->GetResourceHandle(*FAppStyle::GetBrush("Sequencer.SectionArea.Background"));
+			
+				const FSlateRenderTransform& Transform = AllottedGeometry.GetAccumulatedRenderTransform();
+				FColor Color = FillColor.Get().ToFColor(true);
+
+				// filled polygons
+				TArray<FSlateVertex> Vertices;
+				TArray<SlateIndex> Indices;
+				Vertices.Reserve((NumPoints-1)*4);
+				Indices.Reserve((NumPoints-1)*6);
+				
+				for(int i=1; i<NumPoints; i++)
+				{
+					const FTimelineCurveData::CurvePoint& Point1 = CurvePoints[i-1];
+					const FTimelineCurveData::CurvePoint& Point2 = CurvePoints[i];
+				
+					if (Point2.Time - Point1.Time < LargeFrameTime)
+					{
+						float X1 = RangeToScreen.InputToLocalX(Point1.Time);
+						const float NormalizedValue1 = (Point1.Value - MinValue) / Range;
+						float Y1 = Size.Y - (Size.Y - 2) * NormalizedValue1;
+						float X2 = RangeToScreen.InputToLocalX(Point2.Time);
+						const float NormalizedValue2 = (Point2.Value - MinValue) / Range;
+						float Y2 = Size.Y - (Size.Y - 2) * NormalizedValue2;
+					
+						Vertices.Add(FSlateVertex::Make<ESlateVertexRounding::Disabled>(Transform, FVector2f(X1,Size.Y-1) ,FVector2f(0,0), Color));
+						Vertices.Add(FSlateVertex::Make<ESlateVertexRounding::Disabled>(Transform, FVector2f(X1,Y1),FVector2f(0,0), Color));
+						Vertices.Add(FSlateVertex::Make<ESlateVertexRounding::Disabled>(Transform, FVector2f(X2,Y2),FVector2f(0,0), Color));
+						Vertices.Add(FSlateVertex::Make<ESlateVertexRounding::Disabled>(Transform, FVector2f(X2,Size.Y-1),FVector2f(0,0), Color));
+						
+						int32 Index0 = Vertices.Num() - 4;
+                        int32 Index1 = Vertices.Num() - 3;
+                        int32 Index2 = Vertices.Num() - 2;
+                        int32 Index3 = Vertices.Num() - 1;
+                       
+						Indices.Add(Index0);
+						Indices.Add(Index1);
+						Indices.Add(Index2);
+						Indices.Add(Index0);
+						Indices.Add(Index2);
+						Indices.Add(Index3);
+					}
+					
+					FSlateDrawElement::MakeCustomVerts(
+								OutDrawElements,
+										LayerId++,
+										ResourceHandle,
+										Vertices,
+										Indices,
+										nullptr,
+										0,
+										0);
+				}
+			}
+
+
+			// Lines rendering
+			
+			TArray<FVector2D> Points;
+			Points.Reserve(NumPoints);
+
+			float PrevTime = CurvePoints[0].Time;
+			for(int i=0; i<NumPoints; i++)
+			{
+				const FTimelineCurveData::CurvePoint& Point = CurvePoints[i];
+				if (Point.Time - PrevTime > LargeFrameTime && Points.Num()>1)
+				{
+					// break the line list - data has stopped and started again
+					FSlateDrawElement::MakeLines(
+								OutDrawElements,
+										LayerId++,
+										AllottedGeometry.ToPaintGeometry(),
+										Points,
+										LineDrawEffects,
+										LineColor,
+										false
+										);
+					Points.SetNum(0,false);
+				}
+			
+				float X = RangeToScreen.InputToLocalX(Point.Time);
+				PrevTime = Point.Time;
+				const float NormalizedValue = (Point.Value - MinValue) / Range;
+				float Y = Size.Y - (Size.Y - 2) * NormalizedValue;
+				Points.Add(FVector2D(X,Y));
+			}
+			
+			FSlateDrawElement::MakeLines(
+						OutDrawElements,
+								LayerId,
+								AllottedGeometry.ToPaintGeometry(),
+								Points,
+								LineDrawEffects,
+								LineColor,
+								false
+								);
+
+
+								
+		}
+	
+
+
+	}
+
+	return LayerId;
+}
+
+#undef LOCTEXT_NAMESPACE
diff --git a/Engine/Plugins/Animation/GameplayInsights/Source/GameplayInsights/Private/SCurveTimelineView.h b/Engine/Plugins/Animation/GameplayInsights/Source/GameplayInsights/Private/SCurveTimelineView.h
new file mode 100644
--- /dev/null
+++ b/Engine/Plugins/Animation/GameplayInsights/Source/GameplayInsights/Private/SCurveTimelineView.h
@@ -0,0 +1,89 @@
+// Copyright Epic Games, Inc. All Rights Reserved.
+
+#pragma once
+
+#include "CoreMinimal.h"
+#include "Widgets/DeclarativeSyntaxSupport.h"
+
+class FPaintArgs;
+class FSlateWindowElementList;
+class FTimeSliderController;
+
+
+class SCurveTimelineView : public SCompoundWidget
+{
+public:
+	struct FTimelineCurveData
+	{
+	public:
+		struct CurvePoint
+		{
+			double Time;
+			float Value;
+		};
+
+		TArray<CurvePoint> Points;
+	};
+	
+	SLATE_BEGIN_ARGS(SCurveTimelineView)
+		: _ViewRange(TRange<double>(0,10))
+		, _DesiredSize(FVector2D(100.f,20.f))
+		, _RenderFill(false)
+	{}
+    	/** View time range */
+    	SLATE_ATTRIBUTE(TRange<double>, ViewRange);
+
+		/** Data for curve to render */
+		SLATE_ATTRIBUTE(TSharedPtr<FTimelineCurveData>, CurveData);
+
+		/** Desired widget size */
+		SLATE_ATTRIBUTE(FVector2D, DesiredSize);
+	
+		/** Curve Color */
+		SLATE_ATTRIBUTE(FLinearColor, CurveColor);
+	
+		/** Fill Color */
+		SLATE_ATTRIBUTE(FLinearColor, FillColor);
+	
+		/** Enable Fill rendering */
+		SLATE_ATTRIBUTE(bool, RenderFill);
+     
+	SLATE_END_ARGS()
+
+
+	/**
+	 * Construct the widget
+	 * 
+	 * @param InArgs   A declaration from which to construct the widget
+	 */
+	void Construct( const FArguments& InArgs );
+
+	// render the curve with a fixed min/max range, rather than using the min/max from data
+	void SetFixedRange(float Min, float Max)
+	{
+		bUseFixedRange = true;
+		FixedRangeMin = Min;
+		FixedRangeMax = Max;
+	}
+
+protected:
+	// SWidget interface
+	virtual FVector2D ComputeDesiredSize(float) const override;
+	virtual int32 OnPaint( const FPaintArgs& Args, const FGeometry& AllottedGeometry, const FSlateRect& MyCullingRect, FSlateWindowElementList& OutDrawElements, int32 LayerId, const FWidgetStyle& InWidgetStyle, bool bParentEnabled ) const override;
+	
+	int32 PaintCurve(const FGeometry& AllottedGeometry, const FSlateRect& MyCullingRect, FSlateWindowElementList& OutDrawElements, int32 LayerId, const FWidgetStyle& InWidgetStyle, bool bParentEnabled ) const;
+	
+	TAttribute<TRange<double>> ViewRange;
+	TAttribute<TSharedPtr<FTimelineCurveData>> CurveData;
+	TAttribute<FVector2D> DesiredSize;
+	TAttribute<FLinearColor> CurveColor;
+	TAttribute<FLinearColor> FillColor;
+	TAttribute<bool> RenderFill;
+	
+	const FSlateBrush* WhiteBrush;
+
+	bool bUseFixedRange;
+	float FixedRangeMin;
+	float FixedRangeMax;
+};
+
diff --git a/Engine/Plugins/Animation/GameplayInsights/Source/GameplayInsights/Private/SMontageView.cpp b/Engine/Plugins/Animation/GameplayInsights/Source/GameplayInsights/Private/SMontageView.cpp
--- a/Engine/Plugins/Animation/GameplayInsights/Source/GameplayInsights/Private/SMontageView.cpp
+++ b/Engine/Plugins/Animation/GameplayInsights/Source/GameplayInsights/Private/SMontageView.cpp
@@ -4,6 +4,7 @@
 #include "AnimationProvider.h"
 #include "TraceServices/Model/AnalysisSession.h"
 #include "GameplayProvider.h"
+#include "IRewindDebugger.h"
 #include "Styling/SlateIconFinder.h"
 #include "TraceServices/Model/Frames.h"
 #include "VariantTreeNode.h"
@@ -29,18 +30,21 @@ void SMontageView::GetVariantsAtFrame(const TraceServices::FFrame& InFrame, TArr
 			{
 				if(InStartTime >= InFrame.StartTime && InEndTime <= InFrame.EndTime)
 				{
-					const FObjectInfo& MontageInfo = GameplayProvider->GetObjectInfo(InMessage.MontageId);
+					if (!bAssetFilterSet || InMessage.MontageId == AssetIdFilter)
+					{
+						const FObjectInfo& MontageInfo = GameplayProvider->GetObjectInfo(InMessage.MontageId);
 					
-					TSharedRef<FVariantTreeNode> MontageHeader = OutVariants.Add_GetRef(FVariantTreeNode::MakeObject(FText::FromString(MontageInfo.Name), InMessage.MontageId, InMessage.MontageId));
+						TSharedRef<FVariantTreeNode> MontageHeader = OutVariants.Add_GetRef(FVariantTreeNode::MakeObject(FText::FromString(MontageInfo.Name), InMessage.MontageId, InMessage.MontageId));
 
-					MontageHeader->AddChild(FVariantTreeNode::MakeFloat(LOCTEXT("EventWeight", "Weight"), InMessage.Weight));
-					MontageHeader->AddChild(FVariantTreeNode::MakeFloat(LOCTEXT("EventDesiredWeight", "Desired Weight"), InMessage.DesiredWeight));
-					MontageHeader->AddChild(FVariantTreeNode::MakeFloat(LOCTEXT("EventPosition", "Position"), InMessage.Position));
+						MontageHeader->AddChild(FVariantTreeNode::MakeFloat(LOCTEXT("EventWeight", "Weight"), InMessage.Weight));
+						MontageHeader->AddChild(FVariantTreeNode::MakeFloat(LOCTEXT("EventDesiredWeight", "Desired Weight"), InMessage.DesiredWeight));
+						MontageHeader->AddChild(FVariantTreeNode::MakeFloat(LOCTEXT("EventPosition", "Position"), InMessage.Position));
 
-					const TCHAR* CurrentSectionName = AnimationProvider->GetName(InMessage.CurrentSectionNameId);
-					MontageHeader->AddChild(FVariantTreeNode::MakeString(LOCTEXT("CurrentSectionName", "Current Section"), CurrentSectionName));
-					const TCHAR* NextSectionName = AnimationProvider->GetName(InMessage.NextSectionNameId);
-					MontageHeader->AddChild(FVariantTreeNode::MakeString(LOCTEXT("NextSectionName", "Next Section"), NextSectionName));
+						const TCHAR* CurrentSectionName = AnimationProvider->GetName(InMessage.CurrentSectionNameId);
+						MontageHeader->AddChild(FVariantTreeNode::MakeString(LOCTEXT("CurrentSectionName", "Current Section"), CurrentSectionName));
+						const TCHAR* NextSectionName = AnimationProvider->GetName(InMessage.NextSectionNameId);
+						MontageHeader->AddChild(FVariantTreeNode::MakeString(LOCTEXT("NextSectionName", "Next Section"), NextSectionName));
+					}
 				}
 				return TraceServices::EEventEnumerate::Continue;
 			});
@@ -57,35 +61,4 @@ FName SMontageView::GetName() const
 	return MontageName;
 }
 
-FName FMontageViewCreator::GetTargetTypeName() const
-{
-	static FName TargetTypeName = "AnimInstance";
-	return TargetTypeName;
-}
-
-FName FMontageViewCreator::GetName() const
-{
-	return MontageName;
-}
-
-FText FMontageViewCreator::GetTitle() const
-{
-	return LOCTEXT("Montages", "Montages");
-}
-
-FSlateIcon FMontageViewCreator::GetIcon() const
-{
-#if WITH_EDITOR
-	return FSlateIconFinder::FindIconForClass(UAnimInstance::StaticClass());
-#else
-	return FSlateIcon();
-#endif
-}
-
-TSharedPtr<IRewindDebuggerView> FMontageViewCreator::CreateDebugView(uint64 ObjectId, double CurrentTime, const TraceServices::IAnalysisSession& AnalysisSession) const
-{
-	return SNew(SMontageView, ObjectId, CurrentTime, AnalysisSession);
-}
-
-
 #undef LOCTEXT_NAMESPACE
diff --git a/Engine/Plugins/Animation/GameplayInsights/Source/GameplayInsights/Private/SMontageView.h b/Engine/Plugins/Animation/GameplayInsights/Source/GameplayInsights/Private/SMontageView.h
--- a/Engine/Plugins/Animation/GameplayInsights/Source/GameplayInsights/Private/SMontageView.h
+++ b/Engine/Plugins/Animation/GameplayInsights/Source/GameplayInsights/Private/SMontageView.h
@@ -14,14 +14,19 @@ class SMontageView : public SPropertiesDebugViewBase
 public:
 	virtual void GetVariantsAtFrame(const TraceServices::FFrame& InFrame, TArray<TSharedRef<FVariantTreeNode>>& OutVariants) const override;
 	virtual FName GetName() const override;
-};
+	
+	void SetAssetFilter(uint64 InAssetIdFilter)
+	{
+		AssetIdFilter = InAssetIdFilter;
+		bAssetFilterSet = true;
+	}
 
-class FMontageViewCreator : public IRewindDebuggerViewCreator
-{
-	public:
-		virtual FName GetTargetTypeName() const;
-		virtual FName GetName() const override;
-		virtual FText GetTitle() const override;
-		virtual FSlateIcon GetIcon() const override;
-		virtual TSharedPtr<IRewindDebuggerView> CreateDebugView(uint64 ObjectId, double CurrentTime, const TraceServices::IAnalysisSession& InAnalysisSession) const override;
+	void ClearAssetFilter()
+	{
+		bAssetFilterSet = false;
+	}
+
+private:
+	uint64 AssetIdFilter = 0;
+	bool bAssetFilterSet = false;
 };
\ No newline at end of file
diff --git a/Engine/Plugins/Animation/GameplayInsights/Source/GameplayInsights/Private/SNotifiesView.cpp b/Engine/Plugins/Animation/GameplayInsights/Source/GameplayInsights/Private/SNotifiesView.cpp
--- a/Engine/Plugins/Animation/GameplayInsights/Source/GameplayInsights/Private/SNotifiesView.cpp
+++ b/Engine/Plugins/Animation/GameplayInsights/Source/GameplayInsights/Private/SNotifiesView.cpp
@@ -4,6 +4,7 @@
 #include "AnimationProvider.h"
 #include "TraceServices/Model/AnalysisSession.h"
 #include "GameplayProvider.h"
+#include "IRewindDebugger.h"
 #include "Styling/SlateIconFinder.h"
 #include "TraceServices/Model/Frames.h"
 #include "VariantTreeNode.h"
@@ -142,5 +143,22 @@ TSharedPtr<IRewindDebuggerView> FNotifiesViewCreator::CreateDebugView(uint64 Obj
 	return SNew(SNotifiesView, ObjectId, CurrentTime, AnalysisSession);
 }
 
+bool FNotifiesViewCreator::HasDebugInfo(uint64 ObjectId) const
+{
+	const TraceServices::IAnalysisSession* AnalysisSession = IRewindDebugger::Instance()->GetAnalysisSession();
+	
+	TraceServices::FAnalysisSessionReadScope SessionReadScope(*AnalysisSession);
+	bool bHasData = false;
+	if (const FAnimationProvider* AnimationProvider = AnalysisSession->ReadProvider<FAnimationProvider>(FAnimationProvider::ProviderName))
+	{
+		AnimationProvider->ReadNotifyTimeline(ObjectId, [&bHasData](const FAnimationProvider::AnimNotifyTimeline& InTimeline)
+		{
+			bHasData = true;
+		});
+	}
+
+	return bHasData;
+}
+
 
 #undef LOCTEXT_NAMESPACE
diff --git a/Engine/Plugins/Animation/GameplayInsights/Source/GameplayInsights/Private/SNotifiesView.h b/Engine/Plugins/Animation/GameplayInsights/Source/GameplayInsights/Private/SNotifiesView.h
--- a/Engine/Plugins/Animation/GameplayInsights/Source/GameplayInsights/Private/SNotifiesView.h
+++ b/Engine/Plugins/Animation/GameplayInsights/Source/GameplayInsights/Private/SNotifiesView.h
@@ -22,4 +22,5 @@ class FNotifiesViewCreator : public IRewindDebuggerViewCreator
 		virtual FText GetTitle() const override;
 		virtual FSlateIcon GetIcon() const override;
 		virtual TSharedPtr<IRewindDebuggerView> CreateDebugView(uint64 ObjectId, double CurrentTime, const TraceServices::IAnalysisSession& InAnalysisSession) const override;
+		virtual bool HasDebugInfo(uint64 ObjectId) const override;
 };
\ No newline at end of file
diff --git a/Engine/Plugins/Animation/GameplayInsights/Source/GameplayInsights/Private/SObjectPropertiesView.cpp b/Engine/Plugins/Animation/GameplayInsights/Source/GameplayInsights/Private/SObjectPropertiesView.cpp
--- a/Engine/Plugins/Animation/GameplayInsights/Source/GameplayInsights/Private/SObjectPropertiesView.cpp
+++ b/Engine/Plugins/Animation/GameplayInsights/Source/GameplayInsights/Private/SObjectPropertiesView.cpp
@@ -4,6 +4,7 @@
 #include "AnimationProvider.h"
 #include "TraceServices/Model/AnalysisSession.h"
 #include "GameplayProvider.h"
+#include "IRewindDebugger.h"
 #include "Styling/SlateIconFinder.h"
 #include "TraceServices/Model/Frames.h"
 #include "VariantTreeNode.h"
@@ -78,5 +79,22 @@ TSharedPtr<IRewindDebuggerView> FObjectPropertiesViewCreator::CreateDebugView(ui
 	return SNew(SObjectPropertiesView, ObjectId, CurrentTime, AnalysisSession);
 }
 
+bool FObjectPropertiesViewCreator::HasDebugInfo(uint64 ObjectId) const
+{
+	const TraceServices::IAnalysisSession* AnalysisSession = IRewindDebugger::Instance()->GetAnalysisSession();
+	
+	TraceServices::FAnalysisSessionReadScope SessionReadScope(*AnalysisSession);
+	bool bHasData = false;
+	if (const FGameplayProvider* GameplayProvider = AnalysisSession->ReadProvider<FGameplayProvider>(FGameplayProvider::ProviderName))
+	{
+		GameplayProvider->ReadObjectPropertiesTimeline(ObjectId, [this, &bHasData, GameplayProvider](const FGameplayProvider::ObjectPropertiesTimeline& InTimeline)
+		{
+			bHasData = true;
+		});
+	}
+	
+	return bHasData;
+}
+
 
 #undef LOCTEXT_NAMESPACE
diff --git a/Engine/Plugins/Animation/GameplayInsights/Source/GameplayInsights/Private/SObjectPropertiesView.h b/Engine/Plugins/Animation/GameplayInsights/Source/GameplayInsights/Private/SObjectPropertiesView.h
--- a/Engine/Plugins/Animation/GameplayInsights/Source/GameplayInsights/Private/SObjectPropertiesView.h
+++ b/Engine/Plugins/Animation/GameplayInsights/Source/GameplayInsights/Private/SObjectPropertiesView.h
@@ -22,4 +22,5 @@ class FObjectPropertiesViewCreator : public IRewindDebuggerViewCreator
 		virtual FText GetTitle() const override;
 		virtual FSlateIcon GetIcon() const override;
 		virtual TSharedPtr<IRewindDebuggerView> CreateDebugView(uint64 ObjectId, double CurrentTime, const TraceServices::IAnalysisSession& InAnalysisSession) const override;
+		virtual bool HasDebugInfo(uint64 ObjectId) const override;
 };
\ No newline at end of file
diff --git a/Engine/Plugins/Animation/GameplayInsights/Source/GameplayInsights/Private/SPropertiesDebugViewBase.cpp b/Engine/Plugins/Animation/GameplayInsights/Source/GameplayInsights/Private/SPropertiesDebugViewBase.cpp
--- a/Engine/Plugins/Animation/GameplayInsights/Source/GameplayInsights/Private/SPropertiesDebugViewBase.cpp
+++ b/Engine/Plugins/Animation/GameplayInsights/Source/GameplayInsights/Private/SPropertiesDebugViewBase.cpp
@@ -14,6 +14,7 @@ void SPropertiesDebugViewBase::Construct(const FArguments& InArgs, uint64 InObje
 {
 	ObjectId = InObjectId;
 	AnalysisSession = &InAnalysisSession;
+	CurrentTime = InArgs._CurrentTime;
 
 	View = SNew(SVariantValueView, InAnalysisSession).OnGetVariantValues(this, &SPropertiesDebugViewBase::GetVariantsAtFrame);
 
@@ -23,7 +24,6 @@ void SPropertiesDebugViewBase::Construct(const FArguments& InArgs, uint64 InObje
 	[
 		View.ToSharedRef()
 	];
-
 }
 
 void SPropertiesDebugViewBase::SetTimeMarker(double Time)
@@ -41,4 +41,14 @@ void SPropertiesDebugViewBase::SetTimeMarker(double Time)
 	}
 }
 
+void SPropertiesDebugViewBase::Tick( const FGeometry& AllottedGeometry, const double InCurrentTime, const float InDeltaTime )
+{
+	if (CurrentTime.IsBound())
+	{
+		SetTimeMarker(CurrentTime.Get());
+	}
+
+	SCompoundWidget::Tick(AllottedGeometry, InCurrentTime, InDeltaTime);
+}
+
 #undef LOCTEXT_NAMESPACE
diff --git a/Engine/Plugins/Animation/GameplayInsights/Source/GameplayInsights/Private/SPropertiesDebugViewBase.h b/Engine/Plugins/Animation/GameplayInsights/Source/GameplayInsights/Private/SPropertiesDebugViewBase.h
--- a/Engine/Plugins/Animation/GameplayInsights/Source/GameplayInsights/Private/SPropertiesDebugViewBase.h
+++ b/Engine/Plugins/Animation/GameplayInsights/Source/GameplayInsights/Private/SPropertiesDebugViewBase.h
@@ -12,6 +12,7 @@ namespace TraceServices { class IAnalysisSession; }
 class SPropertiesDebugViewBase : public IRewindDebuggerView
 {
 	SLATE_BEGIN_ARGS(SPropertiesDebugViewBase) {}
+	SLATE_ATTRIBUTE(double, CurrentTime)
 	SLATE_END_ARGS()
 public:
 	void Construct(const FArguments& InArgs, uint64 InObjectId, double InTimeMarker, const TraceServices::IAnalysisSession& InAnalysisSession);
@@ -21,10 +22,13 @@ public:
 
 	virtual void GetVariantsAtFrame(const TraceServices::FFrame& InFrame, TArray<TSharedRef<FVariantTreeNode>>& OutVariants) const = 0;
 
+	virtual void Tick( const FGeometry& AllottedGeometry, const double InCurrentTime, const float InDeltaTime ) override;
 protected:
 	TSharedPtr<SVariantValueView> View;
 
 	uint64 ObjectId;
 	double TimeMarker;
+
+	TAttribute<double> CurrentTime;
 	const TraceServices::IAnalysisSession* AnalysisSession;
 };
\ No newline at end of file
diff --git a/Engine/Plugins/Animation/GameplayInsights/Source/GameplayInsights/Public/IAnimationProvider.h b/Engine/Plugins/Animation/GameplayInsights/Source/GameplayInsights/Public/IAnimationProvider.h
--- a/Engine/Plugins/Animation/GameplayInsights/Source/GameplayInsights/Public/IAnimationProvider.h
+++ b/Engine/Plugins/Animation/GameplayInsights/Source/GameplayInsights/Public/IAnimationProvider.h
@@ -24,6 +24,7 @@ struct FSkeletalMeshNamedCurve
 struct FSkeletalMeshPoseMessage
 {
 	FTransform ComponentToWorld;
+	double RecordingTime = 0.0;
 	uint64 TransformStartIndex = 0;
 	uint64 CurveStartIndex = 0;
 	uint64 ComponentId = 0;	
@@ -46,6 +47,7 @@ struct FTickRecordMessage
 	uint64 ComponentId = 0;
 	uint64 AnimInstanceId = 0;
 	uint64 AssetId = 0;
+	double RecordingTime = 0.0;
 	int32 NodeId = -1;
 	float BlendWeight = 0.0f;
 	float PlaybackTime = 0.0f;
@@ -225,6 +227,7 @@ struct FAnimMontageMessage
 {
 	uint64 AnimInstanceId = 0;
 	uint64 MontageId = 0;
+	double RecordingTime = 0.0;
 	uint32 CurrentSectionNameId = 0;
 	uint32 NextSectionNameId = 0;
 	float Weight = 0.0f;
diff --git a/Engine/Plugins/Animation/GameplayInsights/Source/GameplayInsights/Public/IGameplayProvider.h b/Engine/Plugins/Animation/GameplayInsights/Source/GameplayInsights/Public/IGameplayProvider.h
--- a/Engine/Plugins/Animation/GameplayInsights/Source/GameplayInsights/Public/IGameplayProvider.h
+++ b/Engine/Plugins/Animation/GameplayInsights/Source/GameplayInsights/Public/IGameplayProvider.h
@@ -113,6 +113,7 @@ public:
 	virtual void EnumerateObjects(TFunctionRef<void(const FObjectInfo&)> Callback) const = 0;
 	virtual void EnumerateObjects(double StartTime, double EndTime, TFunctionRef<void(const FObjectInfo&)> Callback) const = 0;
 	virtual const FClassInfo* FindClassInfo(uint64 InClassId) const = 0;
+	virtual const UClass* FindClass(uint64 InClassId) const = 0;
 	virtual const FClassInfo* FindClassInfo(const TCHAR* InClassPath) const = 0;
 	virtual const FObjectInfo* FindObjectInfo(uint64 InObjectId) const = 0;
 	virtual const FWorldInfo* FindWorldInfo(uint64 InObjectId) const = 0;
@@ -126,5 +127,5 @@ public:
 	virtual const RecordingInfoTimeline* GetRecordingInfo(uint32 RecordingId) const = 0; 
 	virtual void ReadViewTimeline(TFunctionRef<void(const ViewTimeline&)> Callback) const = 0;
 	virtual uint64 FindPossessingController(uint64 Pawn, double Time) const = 0;
-};
-
+	virtual TRange<double> GetObjectRecordingLifetime(uint64 ObjectId) const = 0;
+};
\ No newline at end of file
diff --git a/Engine/Plugins/Animation/GameplayInsights/Source/RewindDebugger/Private/RewindDebugger.cpp b/Engine/Plugins/Animation/GameplayInsights/Source/RewindDebugger/Private/RewindDebugger.cpp
--- a/Engine/Plugins/Animation/GameplayInsights/Source/RewindDebugger/Private/RewindDebugger.cpp
+++ b/Engine/Plugins/Animation/GameplayInsights/Source/RewindDebugger/Private/RewindDebugger.cpp
@@ -1,6 +1,7 @@
 // Copyright Epic Games, Inc. All Rights Reserved.
 
 #include "RewindDebugger.h"
+
 #include "Components/SkeletalMeshComponent.h"
 #include "Editor.h"
 #include "IAnimationProvider.h"
@@ -12,12 +13,13 @@
 #include "SLevelViewport.h"
 #include "IRewindDebuggerExtension.h"
 #include "IRewindDebuggerDoubleClickHandler.h"
+#include "RewindDebuggerObjectTrack.h"
+#include "RewindDebuggerViewCreators.h"
+#include "RewindDebuggerTrackCreators.h"
 #include "Animation/AnimBlueprint.h"
 #include "Animation/AnimBlueprintGeneratedClass.h"
 #include "ToolMenus.h"
 
-FRewindDebugger* FRewindDebugger::InternalInstance = nullptr;
-
 static void IterateExtensions(TFunction<void(IRewindDebuggerExtension* Extension)> IteratorFunction)
 {
 	// update extensions
@@ -31,6 +33,95 @@ static void IterateExtensions(TFunction<void(IRewindDebuggerExtension* Extension
 	}
 }
 
+static double FindProfileTime(double InDebugTime, const TraceServices::IAnalysisSession* AnalysisSession)
+{
+	const IGameplayProvider* GameplayProvider = AnalysisSession->ReadProvider<IGameplayProvider>("GameplayProvider");
+	const IAnimationProvider* AnimationProvider = AnalysisSession->ReadProvider<IAnimationProvider>("AnimationProvider");
+
+	if (GameplayProvider && AnimationProvider)
+	{
+		TraceServices::FAnalysisSessionReadScope SessionReadScope(*AnalysisSession);
+
+		int ScrubFrameIndex = 0;
+		int RecordingIndex = 0;
+		while (GameplayProvider->GetRecordingInfo(RecordingIndex))
+		{
+			// find latest recording. (hack!)
+			RecordingIndex++;
+		}
+		RecordingIndex--;
+
+		if (const IGameplayProvider::RecordingInfoTimeline* Recording = GameplayProvider->GetRecordingInfo(RecordingIndex))
+		{
+			uint64 EventCount = Recording->GetEventCount();
+
+			if (EventCount > 0)
+			{
+				// check if we are outside of the recorded range, and apply the first or last frame
+				const FRecordingInfoMessage& FirstEvent = Recording->GetEvent(0);
+				const FRecordingInfoMessage& LastEvent = Recording->GetEvent(EventCount - 1);
+				if (InDebugTime <= FirstEvent.ElapsedTime)
+				{
+					ScrubFrameIndex = FMath::Min<uint64>(1, EventCount - 1);
+				}
+				else if (InDebugTime >= LastEvent.ElapsedTime)
+				{
+					ScrubFrameIndex = EventCount - 1;
+				}
+				else
+				{
+					// find the two keys surrounding the CurrentScrubTime, and pick the nearest to update ProfileTime
+					if (Recording->GetEvent(ScrubFrameIndex).ElapsedTime > InDebugTime)
+					{
+						for (uint64 EventIndex = ScrubFrameIndex; EventIndex > 0; EventIndex--)
+						{
+							const FRecordingInfoMessage& Event = Recording->GetEvent(EventIndex);
+							const FRecordingInfoMessage& NextEvent = Recording->GetEvent(EventIndex - 1);
+							if (Event.ElapsedTime >= InDebugTime && NextEvent.ElapsedTime <= InDebugTime)
+							{
+								if (Event.ElapsedTime - InDebugTime < InDebugTime - NextEvent.ElapsedTime)
+								{
+									ScrubFrameIndex = EventIndex;
+								}
+								else
+								{
+									ScrubFrameIndex = EventIndex - 1;
+								}
+								break;
+							}
+						}
+					}
+					else
+					{
+						for (uint64 EventIndex = ScrubFrameIndex; EventIndex < EventCount - 1; EventIndex++)
+						{
+							const FRecordingInfoMessage& Event = Recording->GetEvent(EventIndex);
+							const FRecordingInfoMessage& NextEvent = Recording->GetEvent(EventIndex + 1);
+							if (Event.ElapsedTime <= InDebugTime && NextEvent.ElapsedTime >= InDebugTime)
+							{
+								if (InDebugTime - Event.ElapsedTime < NextEvent.ElapsedTime - InDebugTime)
+								{
+									ScrubFrameIndex = EventIndex;
+								}
+								else
+								{
+									ScrubFrameIndex = EventIndex + 1;
+								}
+								break;
+							}
+						}
+					}
+				}
+
+				const FRecordingInfoMessage& Event = Recording->GetEvent(ScrubFrameIndex);
+				return Event.ProfileTime;
+			}
+		}
+	}
+
+	return 0;
+}
+
 FRewindDebugger::FRewindDebugger()  :
 	ControlState(FRewindDebugger::EControlState::Pause),
 	bPIEStarted(false),
@@ -39,6 +130,8 @@ FRewindDebugger::FRewindDebugger()  :
 	bRecording(false),
 	PlaybackRate(1),
 	CurrentScrubTime(0),
+	CurrentViewRange(0,0),
+	CurrentTraceRange(0,0),
 	ScrubFrameIndex(0),
 	RecordingIndex(0),
 	bTargetActorPositionValid(false)
@@ -56,7 +149,7 @@ FRewindDebugger::FRewindDebugger()  :
 	FEditorDelegates::EndPIE.AddRaw(this, &FRewindDebugger::OnPIEStopped);
 	FEditorDelegates::SingleStepPIE.AddRaw(this, &FRewindDebugger::OnPIESingleStepped);
 
-	DebugTargetActor.OnPropertyChanged = DebugTargetActor.OnPropertyChanged.CreateLambda([this](FString Target) { RefreshDebugComponents(); });
+	DebugTargetActor.OnPropertyChanged = DebugTargetActor.OnPropertyChanged.CreateLambda([this](FString Target) { RefreshDebugTracks(); });
 
 	UnrealInsightsModule = &FModuleManager::LoadModuleChecked<IUnrealInsightsModule>("TraceInsights");
 
@@ -179,80 +272,6 @@ void FRewindDebugger::OnPIEStopped(bool bSimulating)
 	SetCurrentScrubTime(0);
 }
 
-bool FRewindDebugger::UpdateComponentList(uint64 ParentId, TArray<TSharedPtr<FDebugObjectInfo>>& ComponentList, bool bAddController)
-{
-	TSharedPtr<const TraceServices::IAnalysisSession> Session = UnrealInsightsModule->GetAnalysisSession();
-	TraceServices::FAnalysisSessionReadScope SessionReadScope(*Session.Get());
-	UWorld* World = GetWorldToVisualize();
-
-	const IGameplayProvider* GameplayProvider = Session->ReadProvider<IGameplayProvider>("GameplayProvider");
-
-	bool bChanged = false;
-
-	double Time = CurrentTraceTime();
-
-	TArray<uint64, TInlineAllocator<32>> FoundObjects;
-
-	GameplayProvider->EnumerateObjects(Time, Time, [this, &bChanged, ParentId, GameplayProvider, &ComponentList, &FoundObjects](const FObjectInfo& InObjectInfo)
-	{
-		if (InObjectInfo.OuterId == ParentId)
-		{
-			const int32 FoundIndex = ComponentList.FindLastByPredicate([&InObjectInfo](const TSharedPtr<FDebugObjectInfo>& Info) { return Info->ObjectName == InObjectInfo.Name; });
-
-			if (FoundIndex >= 0 && ComponentList[FoundIndex]->ObjectName == InObjectInfo.Name)
-			{
-				// we need to make sure we reusing existing FDebugObjectInfo instances for the treeview selection to be stable
-				ComponentList[FoundIndex]->ObjectId = InObjectInfo.Id; // there is an issue with skeletal mesh components changing ids
-				bChanged = bChanged || UpdateComponentList(InObjectInfo.Id, ComponentList[FoundIndex]->Children);
-			}
-			else
-			{
-				bChanged = true;
-				ComponentList.Add(MakeShared<FDebugObjectInfo>(InObjectInfo.Id,InObjectInfo.Name));
-				UpdateComponentList(InObjectInfo.Id, ComponentList.Last()->Children);
-			}
-			
-			FoundObjects.Add(InObjectInfo.Id);
-		}
-	});
-
-	// add controller and it's component hierarchy if one is attached
-	if (bAddController)
-	{
-		if (uint64 ControllerId = GameplayProvider->FindPossessingController(ParentId, CurrentTraceTime()))
-		{
-			const FObjectInfo& ObjectInfo = GameplayProvider->GetObjectInfo(ControllerId);
-			const int32 FoundIndex = ComponentList.FindLastByPredicate([ObjectInfo](const TSharedPtr<FDebugObjectInfo>& Info) { return Info->ObjectName == ObjectInfo.Name; });
-
-			if (FoundIndex >= 0 && ComponentList[FoundIndex]->ObjectName == ObjectInfo.Name)
-			{
-				ComponentList[FoundIndex]->ObjectId = ObjectInfo.Id;
-				bChanged = bChanged || UpdateComponentList(ObjectInfo.Id, ComponentList[FoundIndex]->Children);
-			}
-			else
-			{
-				bChanged = true;
-				ComponentList.Add(MakeShared<FDebugObjectInfo>(ObjectInfo.Id,ObjectInfo.Name));
-				UpdateComponentList(ObjectInfo.Id, ComponentList.Last()->Children);
-			}
-			
-			FoundObjects.Add(ControllerId);
-		}
-	}
-	
-	// remove any components previously in the list that were not found in this time range.
-	for(int Index=ComponentList.Num()-1; Index>=0; Index--)
-	{
-		if (!FoundObjects.Contains(ComponentList[Index]->ObjectId))
-		{
-			bChanged = true;
-			ComponentList.RemoveAt(Index);
-		}
-	}
-	
-	return bChanged;
-}
-
 bool FRewindDebugger::GetTargetActorPosition(FVector& OutPosition) const
 {
 	OutPosition = TargetActorPosition;
@@ -286,7 +305,7 @@ uint64 FRewindDebugger::GetTargetActorId() const
 	return TargetActorId;
 }
 
-void FRewindDebugger::RefreshDebugComponents()
+void FRewindDebugger::RefreshDebugTracks()
 {
 	if (const TraceServices::IAnalysisSession* Session = GetAnalysisSession())
 	{
@@ -305,23 +324,24 @@ void FRewindDebugger::RefreshDebugComponents()
 			bool bChanged = false;
 
 			// add actor (even if it isn't found in the gameplay provider)
-			if (DebugComponents.Num() == 0)
+			if (DebugTracks.Num() == 0)
 			{
 				bChanged = true;
-				DebugComponents.Add(MakeShared<FDebugObjectInfo>(TargetActorId, DebugTargetActor.Get()));
+				const bool bAddController = true;
+				DebugTracks.Add(MakeShared<RewindDebugger::FRewindDebuggerObjectTrack>(TargetActorId, DebugTargetActor.Get(), bAddController));
 			}
 			else
 			{
-				if (DebugComponents[0]->ObjectName != DebugTargetActor.Get() || DebugComponents[0]->ObjectId != TargetActorId)
+				if (DebugTracks[0]->GetDisplayName().ToString() != DebugTargetActor.Get() || DebugTracks[0]->GetObjectId() != TargetActorId)
 				{
 					bChanged = true;
-					DebugComponents[0] = MakeShared<FDebugObjectInfo>(TargetActorId, DebugTargetActor.Get());
+					DebugTracks[0] = MakeShared<RewindDebugger::FRewindDebuggerObjectTrack>(TargetActorId, DebugTargetActor.Get());
 				}
 			}
 
-			if (TargetActorId != 0 && DebugComponents.Num() > 0)
+			if (TargetActorId != 0 && DebugTracks.Num() > 0)
 			{
-				bChanged = bChanged || UpdateComponentList(TargetActorId, DebugComponents[0]->Children, true);
+				bChanged = bChanged || DebugTracks[0]->Update();
 			}
 
 			if (bChanged)
@@ -488,7 +508,7 @@ void FRewindDebugger::Step(int frames)
 						ScrubFrameIndex = FMath::Clamp<int64>(ScrubFrameIndex + frames, 0, (int64)EventCount - 1);
 
 						const FRecordingInfoMessage& Event = Recording->GetEvent(ScrubFrameIndex);
-						CurrentScrubTime = Event.ElapsedTime;
+						SetCurrentScrubTime(Event.ElapsedTime);
 						TraceTime.Set(Event.ProfileTime);
 						TrackCursorDelegate.ExecuteIfBound(false);
 					}
@@ -541,89 +561,28 @@ UWorld* FRewindDebugger::GetWorldToVisualize() const
 	return World;
 }
 
-void  FRewindDebugger::SetCurrentScrubTime(double Time)
+void FRewindDebugger::SetCurrentViewRange(const TRange<double>& Range)
+{
+	CurrentViewRange = Range;
+	if (const TraceServices::IAnalysisSession* Session = GetAnalysisSession())
+	{
+		CurrentTraceRange.SetLowerBoundValue(FindProfileTime(CurrentViewRange.GetLowerBoundValue(), Session));
+		CurrentTraceRange.SetUpperBoundValue(FindProfileTime(CurrentViewRange.GetUpperBoundValue(), Session));
+	}
+}
+
+void FRewindDebugger::SetCurrentScrubTime(double Time)
 {
 	CurrentScrubTime = Time;
 	UpdateTraceTime();
 }
 
+
 void FRewindDebugger::UpdateTraceTime()
 {
-	// find the current Insights trace time for the current scrub time.
 	if (const TraceServices::IAnalysisSession* Session = GetAnalysisSession())
 	{
-		TraceServices::FAnalysisSessionReadScope SessionReadScope(*Session);
-		UWorld* World = GetWorldToVisualize();
-
-		if (const IGameplayProvider* GameplayProvider = Session->ReadProvider<IGameplayProvider>("GameplayProvider"))
-		{
-			if (const IGameplayProvider::RecordingInfoTimeline* Recording = GameplayProvider->GetRecordingInfo(RecordingIndex))
-			{
-				uint64 EventCount = Recording->GetEventCount();
-
-				if (EventCount > 0)
-				{
-					// check if we are outside of the recorded range, and apply the first or last frame
-					const FRecordingInfoMessage& FirstEvent = Recording->GetEvent(0);
-					const FRecordingInfoMessage& LastEvent = Recording->GetEvent(EventCount - 1);
-					if (CurrentScrubTime <= FirstEvent.ElapsedTime)
-					{
-						ScrubFrameIndex = FMath::Min<uint64>(1, EventCount - 1);
-					}
-					else if (CurrentScrubTime >= LastEvent.ElapsedTime)
-					{
-						ScrubFrameIndex = EventCount - 1;
-					}
-					else
-					{
-						// find the two keys surrounding the CurrentScrubTime, and pick the nearest to update ProfileTime
-						if (Recording->GetEvent(ScrubFrameIndex).ElapsedTime > CurrentScrubTime)
-						{
-							for (uint64 EventIndex = ScrubFrameIndex; EventIndex > 0; EventIndex--)
-							{
-								const FRecordingInfoMessage& Event = Recording->GetEvent(EventIndex);
-								const FRecordingInfoMessage& NextEvent = Recording->GetEvent(EventIndex - 1);
-								if (Event.ElapsedTime >= CurrentScrubTime && NextEvent.ElapsedTime <= CurrentScrubTime)
-								{
-									if (Event.ElapsedTime - CurrentScrubTime < CurrentScrubTime - NextEvent.ElapsedTime)
-									{
-										ScrubFrameIndex = EventIndex;
-									}
-									else
-									{
-										ScrubFrameIndex = EventIndex - 1;
-									}
-									break;
-								}
-							}
-						}
-						else
-						{
-							for (uint64 EventIndex = ScrubFrameIndex; EventIndex < EventCount - 1; EventIndex++)
-							{
-								const FRecordingInfoMessage& Event = Recording->GetEvent(EventIndex);
-								const FRecordingInfoMessage& NextEvent = Recording->GetEvent(EventIndex + 1);
-								if (Event.ElapsedTime <= CurrentScrubTime && NextEvent.ElapsedTime >= CurrentScrubTime)
-								{
-									if (CurrentScrubTime - Event.ElapsedTime < NextEvent.ElapsedTime - CurrentScrubTime)
-									{
-										ScrubFrameIndex = EventIndex;
-									}
-									else
-									{
-										ScrubFrameIndex = EventIndex + 1;
-									}
-									break;
-								}
-							}
-						}
-					}
-
-					const FRecordingInfoMessage& Event = Recording->GetEvent(ScrubFrameIndex);
-					TraceTime.Set(Event.ProfileTime);
-				}
-			}
-		}
+		TraceTime.Set(FindProfileTime(CurrentScrubTime, Session));
 	}
 }
 
@@ -644,7 +603,7 @@ void FRewindDebugger::Tick(float DeltaTime)
 		if (bRecording)
 		{
 			// if you select a debug target before you start recording, update component list when it becomes valid
-			RefreshDebugComponents();
+			RefreshDebugTracks();
 		}
 		
 		const IAnimationProvider* AnimationProvider = Session->ReadProvider<IAnimationProvider>("AnimationProvider");
@@ -900,19 +859,46 @@ void FRewindDebugger::Tick(float DeltaTime)
 	}
 }
 
-void FRewindDebugger::ComponentSelectionChanged(TSharedPtr<FDebugObjectInfo> SelectedObject)
+void FRewindDebugger::ComponentSelectionChanged(TSharedPtr<RewindDebugger::FRewindDebuggerTrack> SelectedObject)
 {
-	SelectedComponent = SelectedObject;
+	SelectedTrack = SelectedObject;
+
+	// todo: if not hidden (add a toggle button to toolbar to hide/show details)
+	TSharedPtr<SDockTab> DetailsTab = FGlobalTabmanager::Get()->TryInvokeTab(FName("RewindDebuggerDetails"));
+
+	TSharedPtr<SWidget> DetailsView;
+
+	if (SelectedObject)
+	{
+		if (const TraceServices::IAnalysisSession* Session = GetAnalysisSession())
+		{
+			DetailsView = SelectedObject->GetDetailsView();
+		}
+	}
+	
+	if (DetailsView)
+	{
+		DetailsTab->SetContent(DetailsView.ToSharedRef());
+	}
+	else
+	{
+		static TSharedPtr<SWidget> EmptyDetails;
+		if (EmptyDetails == nullptr)
+		{
+			EmptyDetails = 	SNew(SSpacer);
+		}
+		DetailsTab->SetContent(EmptyDetails.ToSharedRef());
+	}
 }
 
-void FRewindDebugger::ComponentDoubleClicked(TSharedPtr<FDebugObjectInfo> SelectedObject)
+void FRewindDebugger::ComponentDoubleClicked(TSharedPtr<RewindDebugger::FRewindDebuggerTrack> SelectedObject)
 {
 	if (!SelectedObject.IsValid())
 	{
 		return;
 	}
 	
-	SelectedComponent = SelectedObject;
+	SelectedTrack = SelectedObject;
 	
 	IModularFeatures& ModularFeatures = IModularFeatures::Get();
 	static const FName HandlerFeatureName = IRewindDebuggerDoubleClickHandler::ModularFeatureName;
@@ -922,7 +908,7 @@ void FRewindDebugger::ComponentDoubleClicked(TSharedPtr<FDebugObjectInfo> Select
 		TraceServices::FAnalysisSessionReadScope SessionReadScope(*Session);
 	
 		const IGameplayProvider* GameplayProvider = Session->ReadProvider<IGameplayProvider>("GameplayProvider");
-		const FObjectInfo& ObjectInfo = GameplayProvider->GetObjectInfo(SelectedObject->ObjectId);
+		const FObjectInfo& ObjectInfo = GameplayProvider->GetObjectInfo(SelectedObject->GetObjectId());
 		uint64 ClassId = ObjectInfo.ClassId;
 		bool bHandled = false;
 		
@@ -954,9 +940,9 @@ void FRewindDebugger::ComponentDoubleClicked(TSharedPtr<FDebugObjectInfo> Select
 TSharedPtr<SWidget> FRewindDebugger::BuildComponentContextMenu()
 {
 	UComponentContextMenuContext* MenuContext = NewObject<UComponentContextMenuContext>();
-	MenuContext->SelectedObject = SelectedComponent;
+	MenuContext->SelectedObject = GetSelectedComponent();
 
-	if (SelectedComponent.IsValid())
+	if (SelectedTrack.IsValid())
 	{
 		// build a list of class hierarchy names to make it easier for extensions to enable menu entries by type
 		if (const TraceServices::IAnalysisSession* Session = GetAnalysisSession())
@@ -964,7 +950,7 @@ TSharedPtr<SWidget> FRewindDebugger::BuildComponentContextMenu()
 			TraceServices::FAnalysisSessionReadScope SessionReadScope(*Session);
 	
 			const IGameplayProvider* GameplayProvider = Session->ReadProvider<IGameplayProvider>("GameplayProvider");
-			const FObjectInfo& ObjectInfo = GameplayProvider->GetObjectInfo(SelectedComponent->ObjectId);
+			const FObjectInfo& ObjectInfo = GameplayProvider->GetObjectInfo(SelectedTrack->GetObjectId());
 			uint64 ClassId = ObjectInfo.ClassId;
 			while (ClassId != 0)
 			{
@@ -977,3 +963,35 @@ TSharedPtr<SWidget> FRewindDebugger::BuildComponentContextMenu()
 
 	return UToolMenus::Get()->GenerateWidget("RewindDebugger.ComponentContextMenu", FToolMenuContext(MenuContext));
  }
+
+
+TSharedPtr<FDebugObjectInfo> FRewindDebugger::GetSelectedComponent() const
+{
+	if (!SelectedComponent.IsValid())
+	{
+		SelectedComponent = MakeShared<FDebugObjectInfo>(0,"");
+	}
+	SelectedComponent->ObjectId = SelectedTrack->GetObjectId();
+	SelectedComponent->ObjectName = SelectedTrack->GetDisplayName().ToString();
+	return SelectedComponent;
+}
+
+// build a component tree that's compatible with the public api from 5.0 for GetDebugComponents.
+void FRewindDebugger::RefreshDebugComponents(TArray<TSharedPtr<RewindDebugger::FRewindDebuggerTrack>>& InTracks, TArray<TSharedPtr<FDebugObjectInfo>>& OutComponents)
+{
+	OutComponents.SetNum(0);
+	for(auto& Track : InTracks)
+	{
+		int Index = OutComponents.Num();
+		OutComponents.Add(MakeShared<FDebugObjectInfo>(Track->GetObjectId(), Track->GetDisplayName().ToString()));
+		TArray<TSharedPtr<RewindDebugger::FRewindDebuggerTrack>> TrackChildren;
+		Track->IterateSubTracks([&TrackChildren](TSharedPtr<RewindDebugger::FRewindDebuggerTrack> Child) { TrackChildren.Add(Child); });
+		RefreshDebugComponents(TrackChildren, OutComponents[Index]->Children);
+	}
+}
+
+TArray<TSharedPtr<FDebugObjectInfo>>& FRewindDebugger::GetDebugComponents()
+{
+	RefreshDebugComponents(DebugTracks, DebugComponents);
+	return DebugComponents;
+}
diff --git a/Engine/Plugins/Animation/GameplayInsights/Source/RewindDebugger/Private/RewindDebugger.h b/Engine/Plugins/Animation/GameplayInsights/Source/RewindDebugger/Private/RewindDebugger.h
--- a/Engine/Plugins/Animation/GameplayInsights/Source/RewindDebugger/Private/RewindDebugger.h
+++ b/Engine/Plugins/Animation/GameplayInsights/Source/RewindDebugger/Private/RewindDebugger.h
@@ -3,6 +3,7 @@
 
 #include "CoreMinimal.h"
 #include "IRewindDebugger.h"
+#include "RewindDebuggerTrack.h"
 #include "BindableProperty.h"
 #include "Containers/Ticker.h"
 #include "UObject/WeakObjectPtr.h"
@@ -31,6 +32,8 @@ public:
 
 	// IRewindDebugger interface
 	virtual double CurrentTraceTime() const override { return TraceTime.Get(); }
+	virtual const TRange<double>& GetCurrentTraceRange() const override { return CurrentTraceRange; }
+	virtual const TRange<double>& GetCurrentViewRange() const override { return CurrentViewRange; }
 	virtual const TraceServices::IAnalysisSession* GetAnalysisSession() const override;
 	virtual uint64 GetTargetActorId() const override;
 	virtual bool GetTargetActorPosition(FVector& OutPosition) const override;
@@ -38,7 +41,10 @@ public:
 	virtual bool IsRecording() const override { return bRecording; }
 	virtual bool IsPIESimulating() const override { return bPIESimulating; }
 	virtual double GetRecordingDuration() const override { return RecordingDuration.Get(); }
-	virtual TSharedPtr<FDebugObjectInfo> GetSelectedComponent() const override { return SelectedComponent; }
+	virtual TSharedPtr<FDebugObjectInfo> GetSelectedComponent() const override;
+	virtual TArray<TSharedPtr<FDebugObjectInfo>>& GetDebugComponents() override;
+	
+	TArray<TSharedPtr<RewindDebugger::FRewindDebuggerTrack>>& GetDebugTracks() { return DebugTracks; }
 
 	// create singleton instance
 	static void Initialize();
@@ -47,7 +53,7 @@ public:
 	static void Shutdown();
 
 	// get singleton instance
-	static FRewindDebugger* Instance() { return InternalInstance; }
+	static FRewindDebugger* Instance() { return static_cast<FRewindDebugger*>(InternalInstance); }
 
 	// Start a new Recording:  Start tracing Object + Animation data, increment the current recording index, and reset the recording elapsed time to 0
 	void StartRecording();
@@ -86,16 +92,16 @@ public:
 	// Tick function: While recording, update recording duration.  While paused, and we have recorded data, update skinned mesh poses for the current frame, and handle playback.
 	void Tick(float DeltaTime);
 
-	// update the list of components for the currently selected debug target
-	void RefreshDebugComponents();
-
-	TArray<TSharedPtr<FDebugObjectInfo>>& GetDebugComponents() override { return DebugComponents; };
+	// update the list of tracks for the currently selected debug target
+	void RefreshDebugTracks();
+	
+	void SetCurrentViewRange(const TRange<double>& Range);
 
 	DECLARE_DELEGATE(FOnComponentListChanged)
 	void OnComponentListChanged(const FOnComponentListChanged& ComponentListChangedCallback);
 	
-	void ComponentDoubleClicked(TSharedPtr<FDebugObjectInfo> SelectedObject);
-	void ComponentSelectionChanged(TSharedPtr<FDebugObjectInfo> SelectedObject);
+	void ComponentDoubleClicked(TSharedPtr<RewindDebugger::FRewindDebuggerTrack> SelectedObject);
+	void ComponentSelectionChanged(TSharedPtr<RewindDebugger::FRewindDebuggerTrack> SelectedObject);
 	TSharedPtr<SWidget> BuildComponentContextMenu();
 
 	DECLARE_DELEGATE_OneParam( FOnTrackCursor, bool)
@@ -106,6 +112,8 @@ public:
 	TBindableProperty<FString, BindingType_Out>* GetDebugTargetActorProperty() { return &DebugTargetActor; }
 
 private:
+	void RefreshDebugComponents(TArray<TSharedPtr<RewindDebugger::FRewindDebuggerTrack>>& InTracks, TArray<TSharedPtr<FDebugObjectInfo>>& OutComponents);
+	
 	void OnPIEStarted(bool bSimulating);
 	void OnPIEPaused(bool bSimulating);
 	void OnPIEResumed(bool bSimulating);
@@ -113,7 +121,6 @@ private:
 	void OnPIESingleStepped(bool bSimulating);
 
 
-	bool UpdateComponentList(uint64 ParentId, TArray<TSharedPtr<FDebugObjectInfo>>& NewComponentList, bool bAddController = false);
 	void SetCurrentScrubTime(double Time);
 	void UpdateTraceTime();
 
@@ -141,13 +148,16 @@ private:
 
 	float PlaybackRate;
 	double CurrentScrubTime;
+	TRange<double> CurrentViewRange;
+	TRange<double> CurrentTraceRange;
 	int64 ScrubFrameIndex;
 	uint16 RecordingIndex;
 
-	static FRewindDebugger* InternalInstance;
-
 	TArray<TSharedPtr<FDebugObjectInfo>> DebugComponents;
-	TSharedPtr<FDebugObjectInfo> SelectedComponent;
+	mutable TSharedPtr<FDebugObjectInfo> SelectedComponent;
+	
+	TArray<TSharedPtr<RewindDebugger::FRewindDebuggerTrack>> DebugTracks;
+	TSharedPtr<RewindDebugger::FRewindDebuggerTrack> SelectedTrack;
 
 	struct FMeshComponentResetData
 	{
diff --git a/Engine/Plugins/Animation/GameplayInsights/Source/RewindDebugger/Private/RewindDebuggerFallbackTrack.cpp b/Engine/Plugins/Animation/GameplayInsights/Source/RewindDebugger/Private/RewindDebuggerFallbackTrack.cpp
new file mode 100644
--- /dev/null
+++ b/Engine/Plugins/Animation/GameplayInsights/Source/RewindDebugger/Private/RewindDebuggerFallbackTrack.cpp
@@ -0,0 +1,40 @@
+// Copyright Epic Games, Inc. All Rights Reserved.
+
+#include "RewindDebuggerFallbackTrack.h"
+#include "IRewindDebugger.h"
+#include "IGameplayProvider.h"
+#include "IAnimationProvider.h"
+#include "SObjectTimelineView.h"
+#include "Styling/SlateIconFinder.h"
+#include "ObjectTrace.h"
+
+namespace RewindDebugger
+{
+
+FSlateIcon FRewindDebuggerFallbackTrack::GetIconInternal()
+{
+	return ViewCreator->GetIcon();
+}
+
+bool FRewindDebuggerFallbackTrack::UpdateInternal()
+{
+	IRewindDebugger* RewindDebugger = IRewindDebugger::Instance();
+	TSharedPtr<IRewindDebuggerView> PinnedView = View.Pin();
+	if (PinnedView.IsValid())
+	{
+		PinnedView->SetTimeMarker(RewindDebugger->CurrentTraceTime());
+	}
+
+	return false;
+}
+	
+
+TSharedPtr<SWidget> FRewindDebuggerFallbackTrack::GetDetailsViewInternal()
+{
+	IRewindDebugger* RewindDebugger = IRewindDebugger::Instance();
+	TSharedPtr<IRewindDebuggerView> RewindDebuggerView = ViewCreator->CreateDebugView(ObjectId, RewindDebugger->CurrentTraceTime(), *RewindDebugger->GetAnalysisSession());
+	View = RewindDebuggerView;
+	return RewindDebuggerView;
+}
+
+}
diff --git a/Engine/Plugins/Animation/GameplayInsights/Source/RewindDebugger/Private/RewindDebuggerFallbackTrack.h b/Engine/Plugins/Animation/GameplayInsights/Source/RewindDebugger/Private/RewindDebuggerFallbackTrack.h
new file mode 100644
--- /dev/null
+++ b/Engine/Plugins/Animation/GameplayInsights/Source/RewindDebugger/Private/RewindDebuggerFallbackTrack.h
@@ -0,0 +1,33 @@
+// Copyright Epic Games, Inc. All Rights Reserved.
+#pragma once
+
+#include "RewindDebuggerTrack.h"
+
+namespace RewindDebugger
+{
+
+class FRewindDebuggerFallbackTrack : public FRewindDebuggerTrack
+{
+public:
+
+	FRewindDebuggerFallbackTrack(uint64 InObjectId, const IRewindDebuggerViewCreator* InViewCreator)
+		: ViewCreator(InViewCreator), ObjectId(InObjectId)
+	{
+	}
+
+private:
+	virtual FSlateIcon GetIconInternal() override;
+	virtual TSharedPtr<SWidget> GetDetailsViewInternal() override;
+	virtual FName GetNameInternal() const override { return ViewCreator->GetName(); }
+	virtual FText GetDisplayNameInternal() const override { return ViewCreator->GetTitle(); }
+	virtual uint64 GetObjectIdInternal() const override { return ObjectId; }
+	virtual bool UpdateInternal() override;
+
+	const IRewindDebuggerViewCreator* ViewCreator;
+	TWeakPtr<IRewindDebuggerView> View;
+	FSlateIcon Icon;
+	TRange<double> ExistenceRange;
+	uint64 ObjectId;
+};
+
+}
\ No newline at end of file
diff --git a/Engine/Plugins/Animation/GameplayInsights/Source/RewindDebugger/Private/RewindDebuggerModule.cpp b/Engine/Plugins/Animation/GameplayInsights/Source/RewindDebugger/Private/RewindDebuggerModule.cpp
--- a/Engine/Plugins/Animation/GameplayInsights/Source/RewindDebugger/Private/RewindDebuggerModule.cpp
+++ b/Engine/Plugins/Animation/GameplayInsights/Source/RewindDebugger/Private/RewindDebuggerModule.cpp
@@ -12,6 +12,7 @@
 #include "RewindDebuggerStyle.h"
 #include "RewindDebuggerCommands.h"
 #include "SRewindDebugger.h"
+#include "SRewindDebuggerDetails.h"
 #include "Widgets/Docking/SDockTab.h"
 #include "WorkspaceMenuStructure.h"
 #include "WorkspaceMenuStructureModule.h"
@@ -22,6 +23,25 @@
 #define LOCTEXT_NAMESPACE "RewindDebuggerModule"
 
 static const FName RewindDebuggerTabName("RewindDebugger");
+static const FName RewindDebuggerDetailsTabName("RewindDebuggerDetails");
+
+
+TSharedRef<SDockTab> FRewindDebuggerModule::SpawnRewindDebuggerDetailsTab(const FSpawnTabArgs& SpawnTabArgs)
+{
+	if (FRewindDebugger::Instance() == nullptr)
+	{
+		FRewindDebugger::Initialize();
+	}
+	
+	const TSharedRef<SDockTab> MajorTab = SNew(SDockTab)
+		.TabRole(ETabRole::NomadTab);
+
+	RewindDebuggerDetailsWidget = SNew(SRewindDebuggerDetails, MajorTab, SpawnTabArgs.GetOwnerWindow());
+
+	MajorTab->SetContent(RewindDebuggerDetailsWidget.ToSharedRef());
+
+	return MajorTab;
+}
 
 TSharedRef<SDockTab> FRewindDebuggerModule::SpawnRewindDebuggerTab(const FSpawnTabArgs& SpawnTabArgs)
 {
@@ -39,8 +59,6 @@ TSharedRef<SDockTab> FRewindDebuggerModule::SpawnRewindDebuggerTab(const FSpawnT
 			RewindDebuggerWidget = nullptr;
 		});
 
-	TSharedPtr<SWidget> TabContent;
-
 	TSharedPtr<FUICommandList> CommandList = MakeShared<FUICommandList>();
 	const FRewindDebuggerCommands& Commands = FRewindDebuggerCommands::Get();
 
@@ -105,9 +123,10 @@ TSharedRef<SDockTab> FRewindDebuggerModule::SpawnRewindDebuggerTab(const FSpawnT
 	RewindDebuggerWidget = SNew(SRewindDebugger, CommandList.ToSharedRef(), MajorTab, SpawnTabArgs.GetOwnerWindow())
 								.DebugTargetActor(DebuggerInstance->GetDebugTargetActorProperty())
 								.RecordingDuration(DebuggerInstance->GetRecordingDurationProperty())
-								.DebugComponents(&DebuggerInstance->GetDebugComponents())
+								.DebugComponents(&DebuggerInstance->GetDebugTracks())
 								.TraceTime(DebuggerInstance->GetTraceTimeProperty())
 								.OnScrubPositionChanged_Raw(DebuggerInstance,&FRewindDebugger::ScrubToTime)
+								.OnViewRangeChanged_Raw(DebuggerInstance,&FRewindDebugger::SetCurrentViewRange)
 								.OnComponentDoubleClicked_Raw(DebuggerInstance, &FRewindDebugger::ComponentDoubleClicked)
 								.OnComponentSelectionChanged_Raw(DebuggerInstance, &FRewindDebugger::ComponentSelectionChanged)
 								.BuildComponentContextMenu_Raw(DebuggerInstance, &FRewindDebugger::BuildComponentContextMenu)
@@ -137,7 +156,13 @@ void FRewindDebuggerModule::StartupModule()
 		.SetDisplayName(LOCTEXT("TabTitle", "Rewind Debugger"))
 		.SetIcon(FSlateIcon("RewindDebuggerStyle", "RewindDebugger.RewindIcon"))
 		.SetTooltipText(LOCTEXT("TooltipText", "Opens Rewind Debugger."));
-
+	
+	FGlobalTabmanager::Get()->RegisterNomadTabSpawner(
+		RewindDebuggerDetailsTabName,
+    	FOnSpawnTab::CreateRaw(this, &FRewindDebuggerModule::SpawnRewindDebuggerDetailsTab))
+		.SetDisplayName(LOCTEXT("DetailsTabTitle", "Rewind Debugger Details"))
+		.SetIcon(FSlateIcon("RewindDebuggerStyle", "RewindDebugger.RewindIcon"))
+		.SetTooltipText(LOCTEXT("TooltipText", "Opens Rewind Debugger Details Window."));
 
 	RewindDebuggerCameraExtension.Initialize();
 	IModularFeatures::Get().RegisterModularFeature(IRewindDebuggerExtension::ModularFeatureName, &RewindDebuggerCameraExtension);
diff --git a/Engine/Plugins/Animation/GameplayInsights/Source/RewindDebugger/Private/RewindDebuggerModule.h b/Engine/Plugins/Animation/GameplayInsights/Source/RewindDebugger/Private/RewindDebuggerModule.h
--- a/Engine/Plugins/Animation/GameplayInsights/Source/RewindDebugger/Private/RewindDebuggerModule.h
+++ b/Engine/Plugins/Animation/GameplayInsights/Source/RewindDebugger/Private/RewindDebuggerModule.h
@@ -8,6 +8,7 @@
 #include "RewindDebuggerCamera.h"
 
 class SRewindDebugger;
+class SRewindDebuggerDetails;
 
 class FRewindDebuggerModule : public IModuleInterface
 {
@@ -17,9 +18,11 @@ public:
 	virtual void ShutdownModule() override;
 
 	TSharedRef<SDockTab> SpawnRewindDebuggerTab(const FSpawnTabArgs& SpawnTabArgs);
+	TSharedRef<SDockTab> SpawnRewindDebuggerDetailsTab(const FSpawnTabArgs& SpawnTabArgs);
 
 private:
 	TSharedPtr<SRewindDebugger> RewindDebuggerWidget;
+	TSharedPtr<SRewindDebuggerDetails> RewindDebuggerDetailsWidget;
 
 	FRewindDebuggerCamera RewindDebuggerCameraExtension;
 };
diff --git a/Engine/Plugins/Animation/GameplayInsights/Source/RewindDebugger/Private/RewindDebuggerObjectTrack.cpp b/Engine/Plugins/Animation/GameplayInsights/Source/RewindDebugger/Private/RewindDebuggerObjectTrack.cpp
new file mode 100644
--- /dev/null
+++ b/Engine/Plugins/Animation/GameplayInsights/Source/RewindDebugger/Private/RewindDebuggerObjectTrack.cpp
@@ -0,0 +1,198 @@
+// Copyright Epic Games, Inc. All Rights Reserved.
+
+#include "RewindDebuggerObjectTrack.h"
+#include "RewindDebuggerFallbackTrack.h"
+#include "IRewindDebugger.h"
+#include "IGameplayProvider.h"
+#include "IAnimationProvider.h"
+#include "RewindDebuggerViewCreators.h"
+#include "RewindDebuggerTrackCreators.h"
+#include "SObjectTimelineView.h"
+#include "Styling/SlateIconFinder.h"
+#include "ObjectTrace.h"
+
+namespace RewindDebugger
+{
+
+TSharedPtr<SWidget> FRewindDebuggerObjectTrack::GetTimelineViewInternal()
+{
+	return SNew(SObjectTimelineView)
+		.ViewRange_Lambda([]() { return IRewindDebugger::Instance()->GetCurrentViewRange(); })
+		.FillColor(FLinearColor(0.02f, 0.02f, 0.02f, 0.5f))
+		.ExistenceRange_Raw(this, &FRewindDebuggerObjectTrack::GetExistenceRange);
+
+}
+
+
+// check if an object is or is a subclass of a type by name, based on Insights traced type info
+static bool IsTargetType(uint64 ObjectId, FName TargetTypeName, const TraceServices::IAnalysisSession& Session)
+{
+	TraceServices::FAnalysisSessionReadScope SessionReadScope(Session);
+
+	const IGameplayProvider* GameplayProvider = Session.ReadProvider<IGameplayProvider>("GameplayProvider");
+	const FObjectInfo& ObjectInfo = GameplayProvider->GetObjectInfo(ObjectId);
+	uint64 ClassId = ObjectInfo.ClassId;
+
+	while (ClassId != 0)
+	{
+		const FClassInfo& ClassInfo = GameplayProvider->GetClassInfo(ClassId);
+		if (ClassInfo.Name == TargetTypeName)
+		{
+			return true;
+		}
+		ClassId = ClassInfo.SuperId;
+	}
+
+	return false;
+}
+
+void FRewindDebuggerObjectTrack::IterateSubTracksInternal(TFunction<void(TSharedPtr<FRewindDebuggerTrack> SubTrack)> IteratorFunction)
+{
+	for (TSharedPtr<FRewindDebuggerTrack>& Track : Children)
+	{
+		IteratorFunction(Track);
+	}
+};
+
+bool FRewindDebuggerObjectTrack::UpdateInternal()
+{
+	IRewindDebugger *RewindDebugger = IRewindDebugger::Instance();
+	const TraceServices::IAnalysisSession* Session = RewindDebugger->GetAnalysisSession();
+	TraceServices::FAnalysisSessionReadScope SessionReadScope(*Session);
+
+	const IGameplayProvider* GameplayProvider = Session->ReadProvider<IGameplayProvider>("GameplayProvider");
+
+	bool bChanged = false;
+
+	ExistenceRange = GameplayProvider->GetObjectRecordingLifetime(ObjectId);
+
+	if (Icon.IsSet())
+	{
+		if (const FObjectInfo* ObjectInfo = GameplayProvider->FindObjectInfo(ObjectId))
+		{
+			if (const UClass* FoundClass = GameplayProvider->FindClass(ObjectInfo->ClassId))
+			{
+				Icon = FSlateIconFinder::FindIconForClass(FoundClass);
+				bChanged = true;
+			}
+		}
+	}
+
+	TArray<uint64, TInlineAllocator<32>> FoundObjects;
+
+	FoundObjects.Add(ObjectId); // prevent debug views from being removed
+
+	// add debug views as children
+	FRewindDebuggerTrackCreators::EnumerateCreators([this, &Session, &bChanged, RewindDebugger](const IRewindDebuggerTrackCreator* Creator)
+	{
+		int32 FoundIndex = Children.FindLastByPredicate([&bChanged, this, Creator, RewindDebugger](const TSharedPtr<FRewindDebuggerTrack>& Track) { return Track->GetName() == Creator->GetName(); });
+
+		const bool bHasDebugInfo = Creator->HasDebugInfo(ObjectId) && IsTargetType(ObjectId, Creator->GetTargetTypeName(), *Session);
+
+		if (FoundIndex >= 0)
+		{
+			if (!bHasDebugInfo)
+			{
+				bChanged = true;
+				Children.RemoveAt(FoundIndex);
+			}
+		}
+		else
+		{
+			if (bHasDebugInfo)
+			{
+				bChanged = true;
+				if (TSharedPtr<FRewindDebuggerTrack> Track = Creator->CreateTrack(ObjectId))
+				{
+					Children.Add(Track);
+				}
+			}
+		}
+	});
+
+	// Fallback codepath to add views with no track implementation
+	FRewindDebuggerViewCreators::EnumerateCreators([this, &Session, &bChanged, RewindDebugger](const IRewindDebuggerViewCreator* Creator)
+		{
+			int32 FoundIndex = Children.FindLastByPredicate([&bChanged, this, Creator, RewindDebugger](const TSharedPtr<FRewindDebuggerTrack>& Track) { return Track->GetName() == Creator->GetName(); });
+
+			bool bHasDebugInfo = Creator->HasDebugInfo(ObjectId)
+				&& IsTargetType(ObjectId, Creator->GetTargetTypeName(), *Session);
+
+			if (FoundIndex >= 0)
+			{
+				if (!bHasDebugInfo)
+				{
+					bChanged = true;
+					Children.RemoveAt(FoundIndex);
+				}
+			}
+			else
+			{
+				if (bHasDebugInfo)
+				{
+					bChanged = true;
+					TSharedPtr<FRewindDebuggerTrack> Track = MakeShared<FRewindDebuggerFallbackTrack>(ObjectId, Creator);
+					Children.Add(Track);
+				}
+			}
+		}
+	);
+
+	// add child objects/components
+	GameplayProvider->EnumerateObjects(RewindDebugger->GetCurrentTraceRange().GetLowerBoundValue(), RewindDebugger->GetCurrentTraceRange().GetUpperBoundValue(),
+		[this, &bChanged, GameplayProvider, &FoundObjects, RewindDebugger](const FObjectInfo& InObjectInfo)
+		{
+			if (InObjectInfo.OuterId == ObjectId)
+			{
+				int32 FoundIndex = Children.FindLastByPredicate([&InObjectInfo](const TSharedPtr<FRewindDebuggerTrack>& Info) { return Info->GetObjectId() == InObjectInfo.Id; });
+				bool bExisted = FoundIndex >= 0;
+
+				if (!bExisted)
+				{
+					FoundIndex = Children.Num();
+					Children.Add(MakeShared<FRewindDebuggerObjectTrack>(InObjectInfo.Id, InObjectInfo.Name));
+				}
+
+				bChanged = bChanged | !bExisted;
+				FoundObjects.Add(InObjectInfo.Id);
+			}
+		});
+
+	// add controller and it's component hierarchy if one is attached
+	if (bAddController)
+	{
+		// Should probably update this to use a time range and return all posessing controllers from the visible time range.  For now just returns the one at the current time.
+		if (uint64 ControllerId = GameplayProvider->FindPossessingController(ObjectId, RewindDebugger->CurrentTraceTime()))
+		{
+			const FObjectInfo& ObjectInfo = GameplayProvider->GetObjectInfo(ControllerId);
+			const int32 FoundIndex = Children.FindLastByPredicate([ObjectInfo](const TSharedPtr<FRewindDebuggerTrack>& Info) { return Info->GetObjectId() == ObjectInfo.Id; });
+
+			if (FoundIndex < 0)
+			{
+				bChanged = true;
+				Children.Add(MakeShared<FRewindDebuggerObjectTrack>(ObjectInfo.Id,ObjectInfo.Name));
+			}
+			
+			FoundObjects.Add(ControllerId);
+		}
+	}
+
+	// remove any components previously in the list that were not found in this time range.
+	for (int Index = Children.Num() - 1; Index >= 0; Index--)
+	{
+		if (!FoundObjects.Contains(Children[Index]->GetObjectId()))
+		{
+			bChanged = true;
+			Children.RemoveAt(Index);
+		}
+	}
+
+	for (auto& Child : Children)
+	{
+		bChanged = bChanged || Child->Update();
+	}
+
+	return bChanged;
+}
+
+}
\ No newline at end of file
diff --git a/Engine/Plugins/Animation/GameplayInsights/Source/RewindDebugger/Private/RewindDebuggerObjectTrack.h b/Engine/Plugins/Animation/GameplayInsights/Source/RewindDebugger/Private/RewindDebuggerObjectTrack.h
new file mode 100644
--- /dev/null
+++ b/Engine/Plugins/Animation/GameplayInsights/Source/RewindDebugger/Private/RewindDebuggerObjectTrack.h
@@ -0,0 +1,41 @@
+// Copyright Epic Games, Inc. All Rights Reserved.
+#pragma once
+
+#include "RewindDebuggerTrack.h"
+
+namespace RewindDebugger
+{
+
+class FRewindDebuggerObjectTrack : public FRewindDebuggerTrack
+{
+public:
+
+	FRewindDebuggerObjectTrack(uint64 InObjectId, const FString& InObjectName, bool bInAddController = false)
+		: ObjectName(InObjectName)
+		, ObjectId(InObjectId)
+		, bAddController(bInAddController)
+	{
+	}
+
+	TRange<double> GetExistenceRange() const { return ExistenceRange; }
+
+private:
+	virtual TSharedPtr<SWidget> GetTimelineViewInternal() override;
+	virtual bool UpdateInternal() override;
+	virtual void IterateSubTracksInternal(TFunction<void(TSharedPtr<FRewindDebuggerTrack> SubTrack)> IteratorFunction) override;
+	
+	virtual FName GetNameInternal() const override { return ""; }
+	virtual FSlateIcon GetIconInternal() override { return Icon; }
+	virtual FText GetDisplayNameInternal() const override { return FText::FromString(ObjectName); }
+	virtual uint64 GetObjectIdInternal() const override { return ObjectId; }
+	
+	FString ObjectName;
+	FSlateIcon Icon;
+	TRange<double> ExistenceRange;
+	uint64 ObjectId;
+	TArray<TSharedPtr<FRewindDebuggerTrack>> Children;
+
+	bool bAddController;
+};
+
+}
\ No newline at end of file
diff --git a/Engine/Plugins/Animation/GameplayInsights/Source/RewindDebugger/Private/RewindDebuggerTrackCreators.cpp b/Engine/Plugins/Animation/GameplayInsights/Source/RewindDebugger/Private/RewindDebuggerTrackCreators.cpp
new file mode 100644
--- /dev/null
+++ b/Engine/Plugins/Animation/GameplayInsights/Source/RewindDebugger/Private/RewindDebuggerTrackCreators.cpp
@@ -0,0 +1,43 @@
+// Copyright Epic Games, Inc. All Rights Reserved.
+
+#include "RewindDebuggerTrackCreators.h"
+#include "Features/IModularFeatures.h"
+#include "TraceServices/Model/AnalysisSession.h"
+#include "IGameplayProvider.h"
+
+namespace RewindDebugger
+{
+
+void FRewindDebuggerTrackCreators::EnumerateCreators(TFunctionRef<void(const IRewindDebuggerTrackCreator*)> Callback)
+{
+	IModularFeatures& ModularFeatures = IModularFeatures::Get();
+	static const FName CreatorFeatureName = IRewindDebuggerTrackCreator::ModularFeatureName;
+
+	const int32 NumExtensions = ModularFeatures.GetModularFeatureImplementationCount(CreatorFeatureName);
+	for (int32 ExtensionIndex = 0; ExtensionIndex < NumExtensions; ++ExtensionIndex)
+	{
+		IRewindDebuggerTrackCreator* TrackCreator = static_cast<IRewindDebuggerTrackCreator*>(ModularFeatures.GetModularFeatureImplementation(CreatorFeatureName, ExtensionIndex));
+		Callback(TrackCreator);
+	}
+}
+
+const IRewindDebuggerTrackCreator* FRewindDebuggerTrackCreators::GetCreator(FName CreatorName)
+{
+	IModularFeatures& ModularFeatures = IModularFeatures::Get();
+	static const FName CreatorFeatureName = IRewindDebuggerTrackCreator::ModularFeatureName;
+
+	const int32 NumExtensions = ModularFeatures.GetModularFeatureImplementationCount(CreatorFeatureName);
+	for (int32 ExtensionIndex = 0; ExtensionIndex < NumExtensions; ++ExtensionIndex)
+	{
+		IRewindDebuggerTrackCreator* TrackCreator = static_cast<IRewindDebuggerTrackCreator*>(ModularFeatures.GetModularFeatureImplementation(CreatorFeatureName, ExtensionIndex));
+		if (TrackCreator->GetName() == CreatorName)
+		{
+			return TrackCreator;
+		}
+	}
+
+	return nullptr;
+
+}
+
+}
diff --git a/Engine/Plugins/Animation/GameplayInsights/Source/RewindDebugger/Private/RewindDebuggerTrackCreators.h b/Engine/Plugins/Animation/GameplayInsights/Source/RewindDebugger/Private/RewindDebuggerTrackCreators.h
new file mode 100644
--- /dev/null
+++ b/Engine/Plugins/Animation/GameplayInsights/Source/RewindDebugger/Private/RewindDebuggerTrackCreators.h
@@ -0,0 +1,17 @@
+// Copyright Epic Games, Inc. All Rights Reserved.
+
+#pragma once
+
+#include "CoreMinimal.h"
+#include "RewindDebuggerTrack.h"
+#include "IRewindDebuggerTrackCreator.h"
+
+namespace RewindDebugger
+{
+	class FRewindDebuggerTrackCreators
+	{
+	public:
+		static void EnumerateCreators(TFunctionRef<void(const IRewindDebuggerTrackCreator*)> Callback);
+		static const IRewindDebuggerTrackCreator* GetCreator(FName CreatorName);
+	};
+}
\ No newline at end of file
diff --git a/Engine/Plugins/Animation/GameplayInsights/Source/RewindDebugger/Private/RewindDebuggerViewCreators.cpp b/Engine/Plugins/Animation/GameplayInsights/Source/RewindDebugger/Private/RewindDebuggerViewCreators.cpp
--- a/Engine/Plugins/Animation/GameplayInsights/Source/RewindDebugger/Private/RewindDebuggerViewCreators.cpp
+++ b/Engine/Plugins/Animation/GameplayInsights/Source/RewindDebugger/Private/RewindDebuggerViewCreators.cpp
@@ -18,30 +18,6 @@ void FRewindDebuggerViewCreators::EnumerateCreators(TFunctionRef<void(const IRew
     }
 }
 
-void FRewindDebuggerViewCreators::CreateDebugViews(uint64 ObjectId, double CurrentTime, const TraceServices::IAnalysisSession& Session, TArray<TSharedPtr<IRewindDebuggerView>>& OutDebugViews)
-{
-	TraceServices::FAnalysisSessionReadScope SessionReadScope(Session);
-
-	const IGameplayProvider* GameplayProvider = Session.ReadProvider<IGameplayProvider>("GameplayProvider");
-	const FObjectInfo& ObjectInfo = GameplayProvider->GetObjectInfo(ObjectId);
-	uint64 ClassId = ObjectInfo.ClassId;
-	TArray<FName> TypeNameHierarchy;
-	while (ClassId != 0)
-	{
-		const FClassInfo& ClassInfo = GameplayProvider->GetClassInfo(ClassId);
-		TypeNameHierarchy.Add(ClassInfo.Name);
-		ClassId = ClassInfo.SuperId;
-	}
-
-	EnumerateCreators([&OutDebugViews, &TypeNameHierarchy, ObjectId, CurrentTime, &Session](const IRewindDebuggerViewCreator* ViewCreator)
-	{
-		if (TypeNameHierarchy.Contains(ViewCreator->GetTargetTypeName()))
-		{
-			OutDebugViews.Add(ViewCreator->CreateDebugView(ObjectId, CurrentTime, Session));
-		}
-	});
-}
-
 const IRewindDebuggerViewCreator* FRewindDebuggerViewCreators::GetCreator(FName CreatorName)
 {
 	IModularFeatures& ModularFeatures = IModularFeatures::Get();
diff --git a/Engine/Plugins/Animation/GameplayInsights/Source/RewindDebugger/Private/RewindDebuggerViewCreators.h b/Engine/Plugins/Animation/GameplayInsights/Source/RewindDebugger/Private/RewindDebuggerViewCreators.h
--- a/Engine/Plugins/Animation/GameplayInsights/Source/RewindDebugger/Private/RewindDebuggerViewCreators.h
+++ b/Engine/Plugins/Animation/GameplayInsights/Source/RewindDebugger/Private/RewindDebuggerViewCreators.h
@@ -16,5 +16,4 @@ class FRewindDebuggerViewCreators
 public:
 	static void EnumerateCreators(TFunctionRef<void(const IRewindDebuggerViewCreator*)> Callback);
 	static const IRewindDebuggerViewCreator* GetCreator(FName CreatorName);
-	static void CreateDebugViews(uint64 ObjectId, double CurrentTime, const TraceServices::IAnalysisSession& Session, TArray<TSharedPtr<IRewindDebuggerView>>& OutDebugViews);
 };
diff --git a/Engine/Plugins/Animation/GameplayInsights/Source/RewindDebugger/Private/SObjectTimelineView.cpp b/Engine/Plugins/Animation/GameplayInsights/Source/RewindDebugger/Private/SObjectTimelineView.cpp
new file mode 100644
--- /dev/null
+++ b/Engine/Plugins/Animation/GameplayInsights/Source/RewindDebugger/Private/SObjectTimelineView.cpp
@@ -0,0 +1,56 @@
+// Copyright Epic Games, Inc. All Rights Reserved.
+
+#include "SObjectTimelineView.h"
+#include "Rendering/DrawElements.h"
+#include "TraceServices/Model/AnalysisSession.h"
+#include "SSimpleTimeSlider.h"
+
+#define LOCTEXT_NAMESPACE "SObjectTimelineView"
+
+void SObjectTimelineView::Construct( const SObjectTimelineView::FArguments& InArgs )
+{
+    ViewRange = InArgs._ViewRange;
+	FillColor = InArgs._FillColor;
+	DesiredSize = InArgs._DesiredSize;
+	ExistenceRange = InArgs._ExistenceRange;
+}
+
+int32 SObjectTimelineView::OnPaint( const FPaintArgs& Args, const FGeometry& AllottedGeometry, const FSlateRect& MyCullingRect, FSlateWindowElementList& OutDrawElements, int32 LayerId, const FWidgetStyle& InWidgetStyle, bool bParentEnabled ) const
+{
+	int32 NewLayer = PaintBlock( AllottedGeometry, MyCullingRect, OutDrawElements, LayerId, InWidgetStyle, bParentEnabled );
+
+	return FMath::Max( NewLayer, SCompoundWidget::OnPaint(Args, AllottedGeometry, MyCullingRect, OutDrawElements, NewLayer, InWidgetStyle, ShouldBeEnabled( bParentEnabled ) ) );
+ }
+
+
+FVector2D SObjectTimelineView::ComputeDesiredSize( float ) const
+{
+	return DesiredSize.Get();
+}
+
+int32 SObjectTimelineView::PaintBlock(const FGeometry& AllottedGeometry, const FSlateRect& MyCullingRect, FSlateWindowElementList& OutDrawElements, int32 LayerId, const FWidgetStyle& InWidgetStyle, bool bParentEnabled ) const
+{
+	// convert time range to from rewind debugger times to profiler times
+	TRange<double> DebugTimeRange = ViewRange.Get();
+	FLinearColor Color = FillColor.Get();
+	
+	SSimpleTimeSlider::FScrubRangeToScreen RangeToScreen( DebugTimeRange, AllottedGeometry.GetLocalSize() );
+	FVector2D Size = AllottedGeometry.GetLocalSize();
+
+	TRange<double> ObjectExistenceRange = ExistenceRange.Get();
+
+	double LowerBound = ObjectExistenceRange.HasLowerBound() ? ObjectExistenceRange.GetLowerBoundValue() : 0;
+	double UpperBound = ObjectExistenceRange.HasUpperBound() ? ObjectExistenceRange.GetUpperBoundValue() : Size.X;
+	float BoxMin = FMath::Max(0,RangeToScreen.InputToLocalX(LowerBound));
+	float BoxMax = FMath::Min(Size.X, RangeToScreen.InputToLocalX(UpperBound));
+
+	FPaintGeometry BoxGeometry = AllottedGeometry.ToPaintGeometry(FVector2D(BoxMin,1),FVector2D(BoxMax-BoxMin, Size.Y - 2), 1);
+	
+	const FSlateBrush* Brush = FAppStyle::GetBrush("Sequencer.SectionArea.Background");
+
+	FSlateDrawElement::MakeBox(OutDrawElements, LayerId++, BoxGeometry, Brush, ESlateDrawEffect::None, Color);
+
+	return LayerId;
+}
+
+#undef LOCTEXT_NAMESPACE
diff --git a/Engine/Plugins/Animation/GameplayInsights/Source/RewindDebugger/Private/SObjectTimelineView.h b/Engine/Plugins/Animation/GameplayInsights/Source/RewindDebugger/Private/SObjectTimelineView.h
new file mode 100644
--- /dev/null
+++ b/Engine/Plugins/Animation/GameplayInsights/Source/RewindDebugger/Private/SObjectTimelineView.h
@@ -0,0 +1,58 @@
+// Copyright Epic Games, Inc. All Rights Reserved.
+
+#pragma once
+
+#include "CoreMinimal.h"
+#include "TraceServices/Model/AnalysisSession.h"
+#include "Widgets/DeclarativeSyntaxSupport.h"
+
+class FPaintArgs;
+class FSlateWindowElementList;
+class FTimeSliderController;
+
+
+// a timeline view which renders a block for in the timeline where the object exists
+class SObjectTimelineView : public SCompoundWidget
+{
+public:
+	SLATE_BEGIN_ARGS(SObjectTimelineView)
+		: _ViewRange(TRange<double>(0,10))
+		, _DesiredSize(FVector2D(100.f,20.f))
+	{}
+    	/** View time range */
+    	SLATE_ATTRIBUTE(TRange<double>, ViewRange);
+	
+		/** Existence Time Range */
+		SLATE_ATTRIBUTE(TRange<double>, ExistenceRange);
+
+		/** Desired widget size */
+		SLATE_ATTRIBUTE(FVector2D, DesiredSize);
+	
+		/** Fill Color */
+		SLATE_ATTRIBUTE(FLinearColor, FillColor);
+     
+	SLATE_END_ARGS()
+
+
+	/**
+	 * Construct the widget
+	 * 
+	 * @param InArgs   A declaration from which to construct the widget
+	 */
+	void Construct( const FArguments& InArgs );
+
+protected:
+	// SWidget interface
+	virtual FVector2D ComputeDesiredSize(float) const override;
+	virtual int32 OnPaint( const FPaintArgs& Args, const FGeometry& AllottedGeometry, const FSlateRect& MyCullingRect, FSlateWindowElementList& OutDrawElements, int32 LayerId, const FWidgetStyle& InWidgetStyle, bool bParentEnabled ) const override;
+	
+	int32 PaintBlock(const FGeometry& AllottedGeometry, const FSlateRect& MyCullingRect, FSlateWindowElementList& OutDrawElements, int32 LayerId, const FWidgetStyle& InWidgetStyle, bool bParentEnabled ) const;
+	
+	TAttribute<TRange<double>> ViewRange;
+	TAttribute<TRange<double>> ExistenceRange;
+	TAttribute<FVector2D> DesiredSize;
+	TAttribute<FLinearColor> FillColor;
+	
+	const FSlateBrush* WhiteBrush;
+};
+
diff --git a/Engine/Plugins/Animation/GameplayInsights/Source/RewindDebugger/Private/SRewindDebugger.cpp b/Engine/Plugins/Animation/GameplayInsights/Source/RewindDebugger/Private/SRewindDebugger.cpp
--- a/Engine/Plugins/Animation/GameplayInsights/Source/RewindDebugger/Private/SRewindDebugger.cpp
+++ b/Engine/Plugins/Animation/GameplayInsights/Source/RewindDebugger/Private/SRewindDebugger.cpp
@@ -5,7 +5,6 @@
 #include "Editor.h"
 #include "Editor/EditorEngine.h"
 #include "Framework/Application/SlateApplication.h"
-#include "Framework/Docking/LayoutService.h"
 #include "Framework/MultiBox/MultiBoxBuilder.h"
 #include "Insights/IUnrealInsightsModule.h"
 #include "Modules/ModuleManager.h"
@@ -19,16 +18,12 @@
 #include "Styling/SlateIconFinder.h"
 #include "ToolMenu.h"
 #include "ToolMenus.h"
-#include "Widgets/Docking/SDockTab.h"
 #include "Widgets/Images/SImage.h"
 #include "Widgets/Input/SComboButton.h"
 #include "Widgets/Input/SButton.h"
 #include "Widgets/Layout/SExpandableArea.h"
 #include "Widgets/Layout/SScrollBox.h"
-#include "IRewindDebuggerViewCreator.h"
-#include "RewindDebuggerViewCreators.h"
 #include "IRewindDebuggerDoubleClickHandler.h"
-#include "Interfaces/IMainFrameModule.h"
 
 #define LOCTEXT_NAMESPACE "SRewindDebugger"
 
@@ -81,7 +76,13 @@ void SRewindDebugger::TrackCursor(bool bReverse)
 		}
 	}
 
-	ViewRange = CurrentViewRange;
+	SetViewRange(CurrentViewRange);
+}
+
+void SRewindDebugger::SetViewRange(TRange<double> NewRange)
+{
+	ViewRange = NewRange;
+	OnViewRangeChanged.ExecuteIfBound(NewRange);
 }
 
 void SRewindDebugger::SetDebugTargetActor(AActor* Actor)
@@ -184,27 +185,10 @@ TSharedRef<SWidget> SRewindDebugger::MakeSelectActorMenu()
 	return MenuBuilder.MakeWidget();
 }
 
-void SRewindDebugger::CloseAllTabs()
-{
-	 for(FName& TabName : TabNames)
-	 {
-		CloseTab(TabName);
-	 }
-}
-
-void SRewindDebugger::MainFrameCreationFinished(TSharedPtr<SWindow> InRootWindow, bool bIsNewProjectWindow)
-{
-	IMainFrameModule::Get().OnMainFrameCreationFinished().RemoveAll(this);
-	
-	 // close all tabs that may be open from restoring the saved layout config
-	CloseAllTabs();
-}
-
 void SRewindDebugger::Construct(const FArguments& InArgs, TSharedRef<FUICommandList> CommandList, const TSharedRef<SDockTab>& ConstructUnderMajorTab, const TSharedPtr<SWindow>& ConstructUnderWindow)
 {
-	bInitializing = true;
-	
 	OnScrubPositionChanged = InArgs._OnScrubPositionChanged;
+	OnViewRangeChanged = InArgs._OnViewRangeChanged;
 	OnComponentSelectionChanged = InArgs._OnComponentSelectionChanged;
 	BuildComponentContextMenu = InArgs._BuildComponentContextMenu;
 	ScrubTimeAttribute = InArgs._ScrubTime;
@@ -232,75 +216,66 @@ void SRewindDebugger::Construct(const FArguments& InArgs, TSharedRef<FUICommandL
 	}
 	ToolBarBuilder.EndSection();
 
-	ComponentTreeView =	SNew(SRewindDebuggerComponentTree)
-				.DebugComponents(InArgs._DebugComponents)
-				.OnMouseButtonDoubleClick(InArgs._OnComponentDoubleClicked)
-				.OnContextMenuOpening(this, &SRewindDebugger::OnContextMenuOpening)
-				.OnSelectionChanged(this, &SRewindDebugger::ComponentSelectionChanged);
-
-
-	TraceTime.OnPropertyChanged = TraceTime.OnPropertyChanged.CreateRaw(this, &SRewindDebugger::TraceTimeChanged);
+	TSharedPtr<SScrollBar> ScrollBar = SNew(SScrollBar); 
 
-	// Tab Manager 
-	TabManager = FGlobalTabmanager::Get()->NewTabManager(ConstructUnderMajorTab);
-
-	TabManager->SetOnPersistLayout(
-		FTabManager::FOnPersistLayout::CreateStatic(
-			[](const TSharedRef<FTabManager::FLayout>& InLayout)
+	ComponentTreeView =	SNew(SRewindDebuggerComponentTree)
+		.ExternalScrollBar(ScrollBar)
+		.OnExpansionChanged_Lambda([this]() { TimelinesView->RestoreExpansion(); })
+		.OnScrolled_Lambda([this](double ScrollOffset){ TimelinesView->ScrollTo(ScrollOffset); })
+		.DebugComponents(InArgs._DebugComponents)
+		.OnMouseButtonDoubleClick(InArgs._OnComponentDoubleClicked)
+		.OnContextMenuOpening(this, &SRewindDebugger::OnContextMenuOpening)
+		.OnSelectionChanged_Lambda(
+			[this](TSharedPtr<RewindDebugger::FRewindDebuggerTrack> SelectedItem, ESelectInfo::Type SelectInfo)
 			{
-				if (InLayout->GetPrimaryArea().Pin().IsValid())
+				if (!bInSelectionChanged)
 				{
-					FLayoutSaveRestore::SaveToConfig(GEditorLayoutIni, InLayout);
+					bInSelectionChanged = true;
+					TimelinesView->SetSelection(SelectedItem);
+					ComponentSelectionChanged(SelectedItem, SelectInfo);
+					bInSelectionChanged = false;
 				}
-			}
-		)
-	);
-
-	// Default Layout: all tabs in one stack, inside the rewind debugger tab
-
-	TSharedPtr<FTabManager::FStack> MainTabStack = FTabManager::NewStack();
-
-	FRewindDebuggerViewCreators::EnumerateCreators([this, MainTabStack](const IRewindDebuggerViewCreator* Creator)
-		{
-			FName TabName = Creator->GetName();
-			TabNames.Add(TabName);
-			// Add a closed tabs to the main tab stack in the default layout, so that the first time they won't pop up in their own window
-			TabManager->RegisterTabSpawner(TabName, FOnSpawnTab::CreateRaw(this, &SRewindDebugger::SpawnTab, TabName),
-													FCanSpawnTab::CreateRaw(this, &SRewindDebugger::CanSpawnTab, TabName))
-				.SetDisplayName(Creator->GetTitle())
-				.SetIcon(Creator->GetIcon());
-
-			MainTabStack->AddTab(TabName, ETabState::ClosedTab);
-		}
-	);
-
-	TSharedRef<FTabManager::FLayout> Layout = FTabManager::NewLayout("RewindDebuggerLayout1.0") 
-		->AddArea
-		(
-			FTabManager::NewPrimaryArea()
-			->Split
-			(
-				MainTabStack.ToSharedRef()
-			)
-		);
-
-	// load saved layout if it exists
-	Layout = FLayoutSaveRestore::LoadFromConfig(GEditorLayoutIni, Layout);
+			});
+
+	 TimelinesView = SNew(SRewindDebuggerTimelines)
+		.ExternalScrollbar(ScrollBar)
+		.OnExpansionChanged_Lambda([this]() { ComponentTreeView->RestoreExpansion(); })
+		.OnScrolled_Lambda([this](double ScrollOffset){ ComponentTreeView->ScrollTo(ScrollOffset); })
+		.DebugComponents(InArgs._DebugComponents)
+		.ViewRange_Lambda([this](){return ViewRange;})
+		.ClampRange_Lambda(
+			 [this]()
+			 {
+				 return TRange<double>(0.0f, RecordingDuration.Get());
+			 })
+		.OnViewRangeChanged(this, &SRewindDebugger::SetViewRange)
+		.ScrubPosition(ScrubTimeAttribute)
+		.OnScrubPositionChanged_Lambda(
+			[this](double NewScrubTime, bool bIsScrubbing)
+			{
+				if (bIsScrubbing)
+				{
+					OnScrubPositionChanged.ExecuteIfBound( NewScrubTime, bIsScrubbing );
+				}
+			})
+		.OnSelectionChanged_Lambda(
+			[this](TSharedPtr<RewindDebugger::FRewindDebuggerTrack> SelectedItem, ESelectInfo::Type SelectInfo)
+			{
+				if (!bInSelectionChanged)
+				{
+					bInSelectionChanged = true;
+					ComponentTreeView->SetSelection(SelectedItem);
+					ComponentSelectionChanged(SelectedItem, SelectInfo);
+					bInSelectionChanged = false;
+				}
+			});
 
 	UToolMenu* Menu = UToolMenus::Get()->FindMenu("RewindDebugger.MainMenu");
 
-	FToolMenuSection& Section = Menu->AddSection("ViewsSection", LOCTEXT("Views", "Views"));
-
-	Section.AddDynamicEntry("ViewsSection", FNewToolMenuDelegateLegacy::CreateLambda([this](FMenuBuilder& InMenuBuilder, UToolMenu* InMenu)
-	{
-		MakeViewsMenu(InMenuBuilder);
-	}));
-
-
 	ChildSlot
 	[
-		SNew(SSplitter)
-		+SSplitter::Slot().MinSize(280).Value(0)
+		SNew(SVerticalBox)
+		 + SVerticalBox::Slot().AutoHeight()
 		[
 			SNew(SVerticalBox)
 			+SVerticalBox::Slot().AutoHeight()
@@ -322,327 +297,141 @@ void SRewindDebugger::Construct(const FArguments& InArgs, TSharedRef<FUICommandL
 					ToolBarBuilder.MakeWidget()
 				]
 			]
-			+SVerticalBox::Slot().AutoHeight()
+		]
+		 + SVerticalBox::Slot().FillHeight(1.0)
+		 [
+			SNew(SSplitter)
+			+SSplitter::Slot().MinSize(280).Value(0)
 			[
-				SNew(SHorizontalBox)
-				+SHorizontalBox::Slot() .FillWidth(1.0f)
+				SNew(SVerticalBox)
+				+SVerticalBox::Slot().AutoHeight()
 				[
-					SNew(SComboButton)
-					.OnGetMenuContent(this, &SRewindDebugger::MakeSelectActorMenu)
-					.ButtonContent()
+					SNew(SHorizontalBox)
+					+SHorizontalBox::Slot() .FillWidth(1.0f)
 					[
-						SNew(SHorizontalBox)
-						+SHorizontalBox::Slot().AutoWidth().Padding(3)
+						SNew(SComboButton)
+						.OnGetMenuContent(this, &SRewindDebugger::MakeSelectActorMenu)
+						.ButtonContent()
 						[
-							SNew(SImage)
-							.Image_Lambda([this]
-								{
-									FSlateIcon ActorIcon = FSlateIconFinder::FindIconForClass(AActor::StaticClass());
-									if (DebugComponents != nullptr && DebugComponents->Num()>0)
+							SNew(SHorizontalBox)
+							+SHorizontalBox::Slot().AutoWidth().Padding(3)
+							[
+								SNew(SImage)
+								.Image_Lambda([this]
 									{
-										if (UObject* Object = FObjectTrace::GetObjectFromId((*DebugComponents)[0]->ObjectId))
+										FSlateIcon ActorIcon = FSlateIconFinder::FindIconForClass(AActor::StaticClass());
+										if (DebugComponents != nullptr && DebugComponents->Num()>0)
 										{
-											ActorIcon = FSlateIconFinder::FindIconForClass(Object->GetClass());
+											if (UObject* Object = FObjectTrace::GetObjectFromId((*DebugComponents)[0]->GetObjectId()))
+											{
+												ActorIcon = FSlateIconFinder::FindIconForClass(Object->GetClass());
+											}
 										}
+
+										return ActorIcon.GetIcon();
+									}
+								)
+							]
+							+SHorizontalBox::Slot().Padding(3)
+							[
+								SNew(STextBlock)
+								.Text_Lambda([this](){
+									if (DebugComponents == nullptr || DebugComponents->Num()==0)
+									{
+										return LOCTEXT("Select Actor", "Debug Target Actor");
 									}
 
-									return ActorIcon.GetIcon();
-								}
-							)
-						]
-						+SHorizontalBox::Slot().Padding(3)
-						[
-							SNew(STextBlock)
-							.Text_Lambda([this](){
-								if (DebugComponents == nullptr || DebugComponents->Num()==0)
-								{
-									return LOCTEXT("Select Actor", "Debug Target Actor");
-								}
-
-								FString ReadableName = (*DebugComponents)[0]->ObjectName;
-
-								if (UObject* Object = FObjectTrace::GetObjectFromId((*DebugComponents)[0]->ObjectId))
-								{
-									if (AActor* Actor = Cast<AActor>(Object))
+									FText ReadableName = (*DebugComponents)[0]->GetDisplayName();
+
+									if (UObject* Object = FObjectTrace::GetObjectFromId((*DebugComponents)[0]->GetObjectId()))
 									{
-										ReadableName = Actor->GetActorLabel();
+										if (AActor* Actor = Cast<AActor>(Object))
+										{
+											ReadableName = FText::FromString(Actor->GetActorLabel());
+										}
 									}
-								}
 
-								return FText::FromString(ReadableName);
-							} )
+									return ReadableName;
+								} )
+							]
 						]
 					]
-				]
-				+SHorizontalBox::Slot().AutoWidth().HAlign(HAlign_Right)
-				[
-					SNew(SButton)
-						.ButtonStyle(FAppStyle::Get(), "SimpleButton")
-						.OnClicked(this, &SRewindDebugger::OnSelectActorClicked)
+					+SHorizontalBox::Slot().AutoWidth().HAlign(HAlign_Right)
 					[
-						SNew(SImage)
-						.Image(FRewindDebuggerStyle::Get().GetBrush("RewindDebugger.SelectActor"))
+						SNew(SButton)
+							.ButtonStyle(FAppStyle::Get(), "SimpleButton")
+							.OnClicked(this, &SRewindDebugger::OnSelectActorClicked)
+						[
+							SNew(SImage)
+							.Image(FRewindDebuggerStyle::Get().GetBrush("RewindDebugger.SelectActor"))
+						]
 					]
 				]
+				+SVerticalBox::Slot().FillHeight(1.0f)
+				[
+					ComponentTreeView.ToSharedRef()
+				]
 			]
-			+SVerticalBox::Slot().FillHeight(1.0f)
-			[
-				ComponentTreeView.ToSharedRef()
-			]
-		]
-		+SSplitter::Slot() 
-		[
-			SNew(SVerticalBox)
-			+SVerticalBox::Slot().AutoHeight()
+			+SSplitter::Slot() 
 			[
-				SNew(SSimpleTimeSlider)
-					.ClampRangeHighlightSize(0.15f)
-					.ClampRangeHighlightColor(FLinearColor::Red.CopyWithNewOpacity(0.5f))
-					.ScrubPosition(ScrubTimeAttribute)
-					.ViewRange_Lambda([this](){ return ViewRange; })
-					.OnViewRangeChanged_Lambda([this](TRange<double> NewRange) { ViewRange = NewRange; })
-					.ClampRange_Lambda(
-							[this]()
-							{ 
-								return TRange<double>(0.0f,RecordingDuration.Get());
-							})	
-					.OnScrubPositionChanged_Lambda(
-						[this](double NewScrubTime, bool bIsScrubbing)
-								{
-									if (bIsScrubbing)
+				SNew(SVerticalBox)
+				+SVerticalBox::Slot().AutoHeight()
+				[
+					SNew(SSimpleTimeSlider)
+						.DesiredSize({100,24})
+						.ClampRangeHighlightSize(0.15f)
+						.ClampRangeHighlightColor(FLinearColor::Red.CopyWithNewOpacity(0.5f))
+						.ScrubPosition(ScrubTimeAttribute)
+						.ViewRange_Lambda([this](){ return ViewRange; })
+						.OnViewRangeChanged(this, &SRewindDebugger::SetViewRange)
+						.ClampRange_Lambda(
+								[this]()
+								{ 
+									return TRange<double>(0.0f,RecordingDuration.Get());
+								})	
+						.OnScrubPositionChanged_Lambda(
+							[this](double NewScrubTime, bool bIsScrubbing)
 									{
-										OnScrubPositionChanged.ExecuteIfBound( NewScrubTime, bIsScrubbing );
-									}
-								})
-			]
-			+SVerticalBox::Slot().FillHeight(1.0f)
-			[
-				TabManager->RestoreFrom(Layout,TSharedPtr<SWindow>()).ToSharedRef()
+										if (bIsScrubbing)
+										{
+											OnScrubPositionChanged.ExecuteIfBound( NewScrubTime, bIsScrubbing );
+										}
+									})
+				]
+				+SVerticalBox::Slot().FillHeight(1.0f)
+				[
+					SNew(SOverlay)
+					+ SOverlay::Slot()
+					[				
+						TimelinesView.ToSharedRef()
+					]
+					+SOverlay::Slot().HAlign(EHorizontalAlignment::HAlign_Right)
+					[
+						ScrollBar.ToSharedRef()
+					]
+					
+				]
 			]
 		]
 	];
-
-	if (IMainFrameModule::Get().IsWindowInitialized())
-	{
-		// close all tabs that may be open from restoring the saved layout config
-		CloseAllTabs();
-	}
-	else
-	{
-		// close them later if we are initalizing the layout, to avoid issues with empty windows and crashes
-		IMainFrameModule::Get().OnMainFrameCreationFinished().AddRaw(this, &SRewindDebugger::MainFrameCreationFinished);
-	}
-	
-	bInitializing = false;
 }
 
 void SRewindDebugger::RefreshDebugComponents()
 {
 	ComponentTreeView->Refresh();
+	TimelinesView->Refresh();
 }
 
-void SRewindDebugger::TraceTimeChanged(double Time)
-{
-	for(TSharedPtr<IRewindDebuggerView>& DebugView : DebugViews)
-	{
-		DebugView->SetTimeMarker(Time);
-	}
-	for(TSharedPtr<IRewindDebuggerView>& DebugView : PinnedDebugViews)
-	{
-		DebugView->SetTimeMarker(Time);
-	}
-}
-
-TSharedRef<SDockTab> SRewindDebugger::SpawnTab(const FSpawnTabArgs& Args, FName ViewName)
-{
-	TSharedPtr<IRewindDebuggerView>* View = DebugViews.FindByPredicate([ViewName](TSharedPtr<IRewindDebuggerView>& View) { return View->GetName() == ViewName; } );
-	if (View)
-	{
-		HiddenTabs.Remove(ViewName);
-
-		return SNew(SDockTab)
-		[
-			View->ToSharedRef()
-		]
-		.OnExtendContextMenu(this, &SRewindDebugger::ExtendTabMenu, *View)
-		.OnTabClosed_Lambda([this,ViewName](TSharedRef<SDockTab>)
-		{
-			if(!bInternalClosingTab) // skip this if the tab is being closed by our own code
-			{
-				HiddenTabs.Add(ViewName);
-			}
-		});
-	}
-
-	return SNew(SDockTab);
-}
-
-// returns true if DebugViews contains a view for the ViewName, but there is no matching Pinned view already open
-bool SRewindDebugger::CanSpawnTab(const FSpawnTabArgs& Args, FName ViewName)
-{
-	if (bInitializing)
-	{
-		return true;
-	}
-	
-	TSharedPtr<IRewindDebuggerView>* View = DebugViews.FindByPredicate([ViewName](TSharedPtr<IRewindDebuggerView>& View) { return View->GetName() == ViewName; } );
-	
-	if (View!=nullptr)
-	{
-		bool bPinned = nullptr != PinnedDebugViews.FindByPredicate(
-			[View](TSharedPtr<IRewindDebuggerView>& PinnedView)
-			{
-				return PinnedView->GetName() == (*View)->GetName() && PinnedView->GetObjectId() == (*View)->GetObjectId();
-			}
-		);
-	
-		return !bPinned;
-	}
-	return false;
-}
-
-void SRewindDebugger::OnPinnedTabClosed(TSharedRef<SDockTab> Tab)
-{
-	// remove view from list of pinned views
-	TSharedRef<IRewindDebuggerView> View = StaticCastSharedRef<IRewindDebuggerView>(Tab->GetContent());
-	PinnedDebugViews.Remove(View);
-
-	// recreate non-pinned tabs, so when closing a pinned tab for the currently selected component, the non-pinned one will appear
-	CreateDebugTabs();
-}
-
- void SRewindDebugger::PinTab(TSharedPtr<IRewindDebuggerView> View)
- {
-	if (PinnedDebugViews.Find(View) != INDEX_NONE)
-	{
-		return;
-	}
-
-	FName TabName = View->GetName();
-
-	CloseTab(TabName);
-
-	FSlateIcon TabIcon;
-	FText TabLabel;
-
-	if (const IRewindDebuggerViewCreator* Creator = FRewindDebuggerViewCreators::GetCreator(TabName))
-	{
-		TabIcon = Creator->GetIcon();
-		TabLabel = Creator->GetTitle();
-	}
-
-	TSharedPtr<SDockTab> NewTab = SNew(SDockTab)
-	[
-		// add a wrapper widget here, that says the name of the object/component for pinned tabs 
-		View.ToSharedRef()
-	]
-	.OnExtendContextMenu(this, &SRewindDebugger::ExtendTabMenu, View)
-	.OnTabClosed(this, &SRewindDebugger::OnPinnedTabClosed)
-	.Label(TabLabel)
-	.LabelSuffix(LOCTEXT(" (Locked)", " \xD83D\xDD12")); // unicode lock image
-
-	NewTab->SetTabIcon(TabIcon.GetIcon());
-
-
-	static const FName RewindDebuggerPinnedTab = "RewindDebuggerPinnedTabName";
-	TabManager->InsertNewDocumentTab(TabName, RewindDebuggerPinnedTab, FTabManager::FRequireClosedTab(), NewTab.ToSharedRef());
-
-	PinnedDebugViews.Add(View);
- }
-
 TSharedRef<SWidget> SRewindDebugger::MakeMainMenu()
 {
 	return UToolMenus::Get()->GenerateWidget("RewindDebugger.MainMenu", FToolMenuContext());
 }
 
-void SRewindDebugger::MakeViewsMenu(FMenuBuilder& MenuBuilder)
-{
-
-	TabManager->PopulateLocalTabSpawnerMenu(MenuBuilder);
-
-	MenuBuilder.AddMenuEntry(LOCTEXT("Show All Views", "Show All Views"),
-							 LOCTEXT("Show All tooltip", "Show all debug views that are relevant to the selected object type"),
-							 FSlateIcon(),
-							 FExecuteAction::CreateLambda([this]() { ShowAllViews(); }));
-}
-
-void SRewindDebugger::ExtendTabMenu(FMenuBuilder& MenuBuilder, TSharedPtr<IRewindDebuggerView> View)
-{
-	MenuBuilder.BeginSection("RewindDebugger", LOCTEXT("Rewind Debugger", "Rewind Debugger"));
-
-	MenuBuilder.AddMenuEntry(LOCTEXT("Keep View Open", "Keep View Open"),
-							 LOCTEXT("Keep View Open tooltip", "Keep this debug view open even while selected component/actor changes"),
-							 FSlateIcon(),
-							 FUIAction(
-								FExecuteAction::CreateRaw(this, &SRewindDebugger::PinTab, View),
-								FCanExecuteAction::CreateLambda([this,View](){ return PinnedDebugViews.Find(View) == INDEX_NONE; })
-							 )); 
-
-
-	MenuBuilder.EndSection();
-}
-
-void SRewindDebugger::ShowAllViews()
-{
-	HiddenTabs.Empty();
-	CreateDebugTabs();
-}
-
-void SRewindDebugger::CreateDebugViews()
-{
-	DebugViews.Empty();
-
-	if (SelectedComponent.IsValid())
-	{
-		IUnrealInsightsModule &UnrealInsightsModule = FModuleManager::LoadModuleChecked<IUnrealInsightsModule>("TraceInsights");
-		TSharedPtr<const TraceServices::IAnalysisSession> Session = UnrealInsightsModule.GetAnalysisSession();
-
-		FRewindDebuggerViewCreators::CreateDebugViews(SelectedComponent->ObjectId, TraceTime.Get(), *Session, DebugViews);
-	}
-}
-
-void SRewindDebugger::CloseTab(FName TabName)
-{
-	// using bInternalClosingTab to distinguish between procedural, and user initiated tab closing in the OnTabClosed callback
-	bInternalClosingTab = true;
-	TSharedPtr<SDockTab> Tab = TabManager->FindExistingLiveTab(TabName);
-	if (Tab.IsValid())
-	{
-		Tab->RequestCloseTab();
-	}
-	bInternalClosingTab = false;
-}
-
-void SRewindDebugger::CreateDebugTabs()
-{
-	for(FName& TabName : TabNames)
-	{
-		CloseTab(TabName);
-	}
-
-	for(TSharedPtr<IRewindDebuggerView> DebugView : DebugViews)
-	{
-		FName ViewName = DebugView->GetName();
-		uint64 ObjectId = DebugView->GetObjectId();
-
-	    bool bPinned = nullptr != PinnedDebugViews.FindByPredicate([ViewName, ObjectId](TSharedPtr<IRewindDebuggerView>& View) { return View->GetName() == ViewName && View->GetObjectId() == ObjectId; } );
-	    bool bHidden = HiddenTabs.Find(ViewName) != INDEX_NONE;
-
-		if(!bPinned && !bHidden)
-		{
-			TabManager->TryInvokeTab(ViewName);
-		}
-	}
-}
-
-void SRewindDebugger::ComponentSelectionChanged(TSharedPtr<FDebugObjectInfo> SelectedItem, ESelectInfo::Type SelectInfo)
+void SRewindDebugger::ComponentSelectionChanged(TSharedPtr<RewindDebugger::FRewindDebuggerTrack> SelectedItem, ESelectInfo::Type SelectInfo)
 {
 	SelectedComponent = SelectedItem;
 
 	OnComponentSelectionChanged.ExecuteIfBound(SelectedItem);
-
-	CreateDebugViews();
-	CreateDebugTabs();
 }
 
 TSharedPtr<SWidget> SRewindDebugger::OnContextMenuOpening()
@@ -652,14 +441,14 @@ TSharedPtr<SWidget> SRewindDebugger::OnContextMenuOpening()
 
 FReply SRewindDebugger::OnSelectActorClicked()
 {
-		FActorPickerModeModule& ActorPickerMode = FModuleManager::Get().GetModuleChecked<FActorPickerModeModule>("ActorPickerMode");
-		
-		// todo: force eject (from within BeginActorPickingMode?)
-
-		ActorPickerMode.BeginActorPickingMode(
-			FOnGetAllowedClasses(), 
-			FOnShouldFilterActor(),
-			FOnActorSelected::CreateRaw(this, &SRewindDebugger::SetDebugTargetActor));
+	FActorPickerModeModule& ActorPickerMode = FModuleManager::Get().GetModuleChecked<FActorPickerModeModule>("ActorPickerMode");
+	
+	// todo: force eject (from within BeginActorPickingMode?)
+
+	ActorPickerMode.BeginActorPickingMode(
+		FOnGetAllowedClasses(), 
+		FOnShouldFilterActor(),
+		FOnActorSelected::CreateRaw(this, &SRewindDebugger::SetDebugTargetActor));
 
 
 
diff --git a/Engine/Plugins/Animation/GameplayInsights/Source/RewindDebugger/Private/SRewindDebugger.h b/Engine/Plugins/Animation/GameplayInsights/Source/RewindDebugger/Private/SRewindDebugger.h
--- a/Engine/Plugins/Animation/GameplayInsights/Source/RewindDebugger/Private/SRewindDebugger.h
+++ b/Engine/Plugins/Animation/GameplayInsights/Source/RewindDebugger/Private/SRewindDebugger.h
@@ -4,10 +4,9 @@
 #include "BindableProperty.h"
 #include "CoreMinimal.h"
 #include "Framework/Commands/UICommandList.h"
-#include "Framework/Docking/TabManager.h"
 #include "IRewindDebuggerView.h"
-#include "RewindDebuggerModule.h"
 #include "SRewindDebuggerComponentTree.h"
+#include "SRewindDebuggerTimelines.h"
 #include "Widgets/DeclarativeSyntaxSupport.h"
 #include "Widgets/Input/SSearchBox.h"
 #include "Widgets/SCompoundWidget.h"
@@ -20,18 +19,20 @@ class SRewindDebugger : public SCompoundWidget
 
 public:
 	DECLARE_DELEGATE_TwoParams( FOnScrubPositionChanged, double, bool )
+	DECLARE_DELEGATE_OneParam( FOnViewRangeChanged, const TRange<double>& )
 	DECLARE_DELEGATE_OneParam( FOnDebugTargetChanged, TSharedPtr<FString> )
-	DECLARE_DELEGATE_OneParam( FOnComponentDoubleClicked, TSharedPtr<FDebugObjectInfo> )
-	DECLARE_DELEGATE_OneParam( FOnComponentSelectionChanged, TSharedPtr<FDebugObjectInfo> )
+	DECLARE_DELEGATE_OneParam( FOnComponentDoubleClicked, TSharedPtr<RewindDebugger::FRewindDebuggerTrack> )
+	DECLARE_DELEGATE_OneParam( FOnComponentSelectionChanged, TSharedPtr<RewindDebugger::FRewindDebuggerTrack> )
 	DECLARE_DELEGATE_RetVal( TSharedPtr<SWidget>, FBuildComponentContextMenu )
 
 	SLATE_BEGIN_ARGS(SRewindDebugger) { }
-		SLATE_ARGUMENT( TArray< TSharedPtr< FDebugObjectInfo > >*, DebugComponents );
+		SLATE_ARGUMENT( TArray< TSharedPtr< RewindDebugger::FRewindDebuggerTrack > >*, DebugComponents );
 		SLATE_ARGUMENT(DebugTargetInitializer, DebugTargetActor);
 		SLATE_ARGUMENT(TBindablePropertyInitializer<double>, TraceTime);
 		SLATE_ARGUMENT(TBindablePropertyInitializer<float>, RecordingDuration);
 		SLATE_ATTRIBUTE(double, ScrubTime);
 		SLATE_EVENT( FOnScrubPositionChanged, OnScrubPositionChanged);
+		SLATE_EVENT( FOnViewRangeChanged, OnViewRangeChanged);
 		SLATE_EVENT( FBuildComponentContextMenu, BuildComponentContextMenu );
 		SLATE_EVENT( FOnComponentDoubleClicked, OnComponentDoubleClicked );
 		SLATE_EVENT( FOnComponentSelectionChanged, OnComponentSelectionChanged );
@@ -57,14 +58,17 @@ public:
 
 	void TrackCursor(bool bReverse);
 	void RefreshDebugComponents();
-	void CloseAllTabs();
 private:
+	void SetViewRange(TRange<double> NewRange);
+
+	
 	void MainFrameCreationFinished(TSharedPtr<SWindow> InRootWindow, bool bIsNewProjectWindow);
 	
 	// Time Slider
 	TAttribute<double> ScrubTimeAttribute;
 	TAttribute<bool> TrackScrubbingAttribute;
 	FOnScrubPositionChanged OnScrubPositionChanged;
+	FOnViewRangeChanged OnViewRangeChanged;
 	TRange<double> ViewRange;
 	TBindableProperty<double> TraceTime;
 	TBindableProperty<float> RecordingDuration;
@@ -75,7 +79,6 @@ private:
 	FReply OnSelectActorClicked();
 
 	TBindableProperty<FString, BindingType_Out> DebugTargetActor;
-	TBindableProperty<float> DebugTargetAnimInstanceId;
 
 	void TraceTimeChanged(double Time);
 
@@ -83,31 +86,15 @@ private:
 	void MakeViewsMenu(FMenuBuilder& MenuBuilder);
 	
 	// component tree view
-	TArray<TSharedPtr<FDebugObjectInfo>>* DebugComponents;
-	TSharedPtr<FDebugObjectInfo> SelectedComponent;
-    void ComponentSelectionChanged(TSharedPtr<FDebugObjectInfo> SelectedItem, ESelectInfo::Type SelectInfo);
+	TArray<TSharedPtr<RewindDebugger::FRewindDebuggerTrack>>* DebugComponents;
+	TSharedPtr<RewindDebugger::FRewindDebuggerTrack> SelectedComponent;
+    void ComponentSelectionChanged(TSharedPtr<RewindDebugger::FRewindDebuggerTrack> SelectedItem, ESelectInfo::Type SelectInfo);
 	FBuildComponentContextMenu BuildComponentContextMenu;
 	FOnComponentSelectionChanged OnComponentSelectionChanged;
 	
 	TSharedPtr<SRewindDebuggerComponentTree> ComponentTreeView;
+	TSharedPtr<SRewindDebuggerTimelines> TimelinesView;
 	TSharedPtr<SWidget> OnContextMenuOpening();
 
-	// Debug View Tabs
-	TSharedRef<SDockTab> SpawnTab(const FSpawnTabArgs& Args, FName ViewName);
-	bool CanSpawnTab(const FSpawnTabArgs& Args, FName ViewName);
-	void CloseTab(FName TabName);
-	void OnPinnedTabClosed(TSharedRef<SDockTab> Tab);
-	void ExtendTabMenu(FMenuBuilder& MenuBuilder, TSharedPtr<IRewindDebuggerView> View);
-	void PinTab(TSharedPtr<IRewindDebuggerView> View);
-	void ShowAllViews();
-	void CreateDebugViews();
-	void CreateDebugTabs();
-	TArray<TSharedPtr<IRewindDebuggerView>> DebugViews;
-	TArray<TSharedPtr<IRewindDebuggerView>> PinnedDebugViews;
-	TArray<FName> TabNames;
-	TArray<FName> HiddenTabs;  // keep track of tabs that have been closed so we don't automatically reopen them when switching components
-	bool bInternalClosingTab = false;
-	bool bInitializing = false;
-
-	TSharedPtr<FTabManager> TabManager;
+	bool bInSelectionChanged = false;
 };
diff --git a/Engine/Plugins/Animation/GameplayInsights/Source/RewindDebugger/Private/SRewindDebuggerComponentTree.cpp b/Engine/Plugins/Animation/GameplayInsights/Source/RewindDebugger/Private/SRewindDebuggerComponentTree.cpp
--- a/Engine/Plugins/Animation/GameplayInsights/Source/RewindDebugger/Private/SRewindDebuggerComponentTree.cpp
+++ b/Engine/Plugins/Animation/GameplayInsights/Source/RewindDebugger/Private/SRewindDebuggerComponentTree.cpp
@@ -1,10 +1,14 @@
 // Copyright Epic Games, Inc. All Rights Reserved.
 
 #include "SRewindDebuggerComponentTree.h"
+
+#include "ISequencerWidgetsModule.h"
 #include "ObjectTrace.h"
 #include "Styling/SlateIconFinder.h"
 #include "Widgets/Images/SImage.h"
 
+#include "RewindDebugger.h"
+
 #define LOCTEXT_NAMESPACE "SAnimationInsights"
 
 SRewindDebuggerComponentTree::SRewindDebuggerComponentTree() 
@@ -15,84 +19,95 @@ SRewindDebuggerComponentTree::SRewindDebuggerComponentTree()
 
 SRewindDebuggerComponentTree::~SRewindDebuggerComponentTree() 
 {
+
 }
 
-TSharedRef<ITableRow> ComponentTreeViewGenerateRow(TSharedPtr<FDebugObjectInfo> InItem, const TSharedRef<STableViewBase>& OwnerTable)
+TSharedRef<ITableRow> SRewindDebuggerComponentTree::ComponentTreeViewGenerateRow(TSharedPtr<RewindDebugger::FRewindDebuggerTrack> InItem, const TSharedRef<STableViewBase>& OwnerTable)
 {
-	FString ReadableName = InItem->ObjectName;
-
-	FSlateIcon ObjectIcon = FSlateIconFinder::FindIconForClass(UObject::StaticClass());
-
-	if (UObject* Object = FObjectTrace::GetObjectFromId(InItem->ObjectId))
-	{
-		if (AActor* Actor = Cast<AActor>(Object))
-		{
-			ReadableName = Actor->GetActorLabel();
-		}
-		else if (UActorComponent* Component = Cast<UActorComponent>(Object))
-		{
-			ReadableName = Component->GetName();
-		}
-
-		ObjectIcon = FSlateIconFinder::FindIconForClass(Object->GetClass());
-	}
-
+	FSlateIcon ObjectIcon = InItem->GetIcon();
+	
 	return 
-		SNew(STableRow<TSharedPtr<FDebugObjectInfo>>, OwnerTable)
+		SNew(STableRow<TSharedPtr<RewindDebugger::FRewindDebuggerTrack>>, OwnerTable)
 		[
 			SNew(SHorizontalBox)
-			+SHorizontalBox::Slot().AutoWidth().Padding(3)
+			+SHorizontalBox::Slot().AutoWidth().Padding(2)
 			[
-				SNew(SImage)
+				SNew(SImage).DesiredSizeOverride(FVector2D(16,16))
 				.Image(ObjectIcon.GetIcon())
 			]
-			+SHorizontalBox::Slot().Padding(3)
+			+SHorizontalBox::Slot()
 			[
 				SNew(STextBlock)
-				.Text(FText::FromString(ReadableName))
+				.Text(InItem->GetDisplayName())
 			]
 		];
 }
 
-void ComponentTreeViewGetChildren(TSharedPtr<FDebugObjectInfo> InItem, TArray<TSharedPtr<FDebugObjectInfo>>& OutChildren)
+void ComponentTreeViewGetChildren(TSharedPtr<RewindDebugger::FRewindDebuggerTrack> InItem, TArray<TSharedPtr<RewindDebugger::FRewindDebuggerTrack>>& OutChildren)
 {
-	OutChildren.Append(InItem->Children);
+	InItem->IterateSubTracks([&OutChildren](TSharedPtr<RewindDebugger::FRewindDebuggerTrack> Track)
+	{
+		OutChildren.Add(Track);
+	});
 }
 
-void ComponentTreeViewExpansionChanged(TSharedPtr<FDebugObjectInfo> InItem, bool bShouldBeExpanded)
+void SRewindDebuggerComponentTree::ComponentTreeViewExpansionChanged(TSharedPtr<RewindDebugger::FRewindDebuggerTrack> InItem, bool bShouldBeExpanded)
 {
-	InItem->bExpanded = bShouldBeExpanded;
+	InItem->SetIsExpanded(bShouldBeExpanded);
+	OnExpansionChanged.ExecuteIfBound();
 }
 
+void SRewindDebuggerComponentTree::SetSelection(TSharedPtr<RewindDebugger::FRewindDebuggerTrack> SelectedItem)
+{
+	ComponentTreeView->SetSelection(SelectedItem);
+}
 
+void SRewindDebuggerComponentTree::ScrollTo(double ScrollOffset)
+{
+	ComponentTreeView->SetScrollOffset(ScrollOffset);
+}
 
 void SRewindDebuggerComponentTree::Construct(const FArguments& InArgs)
 {
 	DebugComponents = InArgs._DebugComponents;
 
-	ComponentTreeView = SNew(STreeView<TSharedPtr<FDebugObjectInfo>>)
-									.ItemHeight(16.0f)
+	ComponentTreeView = SNew(STreeView<TSharedPtr<RewindDebugger::FRewindDebuggerTrack>>)
+									.ItemHeight(20.0f)
 									.TreeItemsSource(DebugComponents)
-									.OnGenerateRow_Static(&ComponentTreeViewGenerateRow)
+									.OnGenerateRow(this, &SRewindDebuggerComponentTree::ComponentTreeViewGenerateRow)
 									.OnGetChildren_Static(&ComponentTreeViewGetChildren)
-									.OnExpansionChanged_Static(&ComponentTreeViewExpansionChanged)
+									.OnExpansionChanged(this, &SRewindDebuggerComponentTree::ComponentTreeViewExpansionChanged)
 									.SelectionMode(ESelectionMode::Single)
 									.OnSelectionChanged(InArgs._OnSelectionChanged)
 									.OnMouseButtonDoubleClick(InArgs._OnMouseButtonDoubleClick)
+									.ExternalScrollbar(InArgs._ExternalScrollBar)
+									.AllowOverscroll(EAllowOverscroll::No)
+									.OnTreeViewScrolled(InArgs._OnScrolled)
+									.ScrollbarDragFocusCause(EFocusCause::SetDirectly)
 									.OnContextMenuOpening(InArgs._OnContextMenuOpening);
 
 	ChildSlot
 	[
 		ComponentTreeView.ToSharedRef()
 	];
+
+	OnExpansionChanged = InArgs._OnExpansionChanged;
+}
+
+static void RestoreExpansion(TSharedPtr<RewindDebugger::FRewindDebuggerTrack> Track, TSharedPtr<STreeView<TSharedPtr<RewindDebugger::FRewindDebuggerTrack>>>& TreeView)
+{
+	TreeView->SetItemExpansion(Track, Track->GetIsExpanded());
+	Track->IterateSubTracks([&TreeView](TSharedPtr<RewindDebugger::FRewindDebuggerTrack> SubTrack)
+	{
+		RestoreExpansion(SubTrack, TreeView);
+	});
 }
 
-void RestoreExpansion(TArray<TSharedPtr<FDebugObjectInfo>>& Components, TSharedPtr<STreeView<TSharedPtr<FDebugObjectInfo>>>& TreeView)
+void SRewindDebuggerComponentTree::RestoreExpansion()
 {
-	for(TSharedPtr<FDebugObjectInfo> &Component : Components)
+	for (auto& Track : *DebugComponents)
 	{
-		TreeView->SetItemExpansion(Component, Component->bExpanded);
-		RestoreExpansion(Component->Children, TreeView);
+		::RestoreExpansion(Track, ComponentTreeView);
 	}
 }
 
@@ -103,7 +118,7 @@ void SRewindDebuggerComponentTree::Refresh()
 	if (DebugComponents)
 	{
 		// make sure any newly added TreeView nodes are created expanded
-		RestoreExpansion(*DebugComponents, ComponentTreeView);
+		RestoreExpansion();
 	}
 }
 
diff --git a/Engine/Plugins/Animation/GameplayInsights/Source/RewindDebugger/Private/SRewindDebuggerComponentTree.h b/Engine/Plugins/Animation/GameplayInsights/Source/RewindDebugger/Private/SRewindDebuggerComponentTree.h
--- a/Engine/Plugins/Animation/GameplayInsights/Source/RewindDebugger/Private/SRewindDebuggerComponentTree.h
+++ b/Engine/Plugins/Animation/GameplayInsights/Source/RewindDebugger/Private/SRewindDebuggerComponentTree.h
@@ -3,6 +3,7 @@
 
 #include "CoreMinimal.h"
 #include "IRewindDebugger.h"
+#include "RewindDebuggerTrack.h"
 #include "Widgets/DeclarativeSyntaxSupport.h"
 #include "Widgets/SCompoundWidget.h"
 #include "Widgets/Views/STreeView.h"
@@ -10,15 +11,18 @@
 
 class SRewindDebuggerComponentTree : public SCompoundWidget
 {
-	using FOnSelectionChanged = typename STreeView<TSharedPtr<FDebugObjectInfo>>::FOnSelectionChanged;
-	using FOnMouseButtonDoubleClick = typename STreeView<TSharedPtr<FDebugObjectInfo>>::FOnMouseButtonDoubleClick;
+	using FOnSelectionChanged = typename STreeView<TSharedPtr<RewindDebugger::FRewindDebuggerTrack>>::FOnSelectionChanged;
+	using FOnMouseButtonDoubleClick = typename STreeView<TSharedPtr<RewindDebugger::FRewindDebuggerTrack>>::FOnMouseButtonDoubleClick;
 	
 public:
 	SLATE_BEGIN_ARGS(SRewindDebuggerComponentTree) { }
-		SLATE_ARGUMENT( TArray< TSharedPtr< FDebugObjectInfo > >*, DebugComponents );
+		SLATE_ARGUMENT( TArray< TSharedPtr< RewindDebugger::FRewindDebuggerTrack > >*, DebugComponents );
+		SLATE_ARGUMENT( TSharedPtr< SScrollBar >, ExternalScrollBar );
 		SLATE_EVENT( FOnSelectionChanged, OnSelectionChanged )
 		SLATE_EVENT( FOnMouseButtonDoubleClick, OnMouseButtonDoubleClick )
 		SLATE_EVENT( FOnContextMenuOpening, OnContextMenuOpening )
+		SLATE_EVENT( FSimpleDelegate, OnExpansionChanged )
+		SLATE_EVENT( FOnTableViewScrolled, OnScrolled )
 	SLATE_END_ARGS()
 
 	/**
@@ -35,11 +39,19 @@ public:
 	void Construct(const FArguments& InArgs);
 
 	void Refresh();
+	void RestoreExpansion();
+	
+	void SetSelection(TSharedPtr<RewindDebugger::FRewindDebuggerTrack> SelectedItem);
+	void ScrollTo(double ScrollOffset);
 
 private:
-	TArray<TSharedPtr<FDebugObjectInfo>>* DebugComponents;
-    void OnComponentSelectionChanged(TSharedPtr<FDebugObjectInfo> SelectedItem, ESelectInfo::Type SelectInfo);
-	TSharedPtr<STreeView<TSharedPtr<FDebugObjectInfo>>> ComponentTreeView;
+	TSharedRef<ITableRow> ComponentTreeViewGenerateRow(TSharedPtr<RewindDebugger::FRewindDebuggerTrack> InItem, const TSharedRef<STableViewBase>& OwnerTable);
+	void ComponentTreeViewExpansionChanged(TSharedPtr<RewindDebugger::FRewindDebuggerTrack> InItem, bool bShouldBeExpanded);
+	
+	TArray<TSharedPtr<RewindDebugger::FRewindDebuggerTrack>>* DebugComponents;
+    void OnComponentSelectionChanged(TSharedPtr<RewindDebugger::FRewindDebuggerTrack> SelectedItem, ESelectInfo::Type SelectInfo);
+	TSharedPtr<STreeView<TSharedPtr<RewindDebugger::FRewindDebuggerTrack>>> ComponentTreeView;
 
-	FOnSelectionChanged OnSelectionChanged;
+	FSimpleDelegate OnExpansionChanged;
+	
 };
diff --git a/Engine/Plugins/Animation/GameplayInsights/Source/RewindDebugger/Private/SRewindDebuggerDetails.cpp b/Engine/Plugins/Animation/GameplayInsights/Source/RewindDebugger/Private/SRewindDebuggerDetails.cpp
new file mode 100644
--- /dev/null
+++ b/Engine/Plugins/Animation/GameplayInsights/Source/RewindDebugger/Private/SRewindDebuggerDetails.cpp
@@ -0,0 +1,48 @@
+// Copyright Epic Games, Inc. All Rights Reserved.
+
+#include "SRewindDebuggerDetails.h"
+#include "ActorPickerMode.h"
+#include "Editor.h"
+#include "Editor/EditorEngine.h"
+#include "Framework/Application/SlateApplication.h"
+#include "Framework/Docking/LayoutService.h"
+#include "Framework/MultiBox/MultiBoxBuilder.h"
+#include "Insights/IUnrealInsightsModule.h"
+#include "Modules/ModuleManager.h"
+#include "ObjectTrace.h"
+#include "RewindDebuggerStyle.h"
+#include "RewindDebuggerCommands.h"
+#include "SSimpleTimeSlider.h"
+#include "SceneOutlinerPublicTypes.h"
+#include "SceneOutlinerModule.h"
+#include "Selection.h"
+#include "Styling/SlateIconFinder.h"
+#include "ToolMenu.h"
+#include "ToolMenus.h"
+#include "Widgets/Docking/SDockTab.h"
+#include "Widgets/Images/SImage.h"
+#include "Widgets/Input/SComboButton.h"
+#include "Widgets/Input/SButton.h"
+#include "Widgets/Layout/SExpandableArea.h"
+#include "Widgets/Layout/SScrollBox.h"
+#include "IRewindDebuggerViewCreator.h"
+#include "RewindDebuggerViewCreators.h"
+#include "IRewindDebuggerDoubleClickHandler.h"
+
+#define LOCTEXT_NAMESPACE "SRewindDebuggerDetails"
+
+SRewindDebuggerDetails::SRewindDebuggerDetails() 
+	: SCompoundWidget()
+{ 
+}
+
+SRewindDebuggerDetails::~SRewindDebuggerDetails() 
+{
+}
+
+void SRewindDebuggerDetails::Construct(const FArguments& InArgs, const TSharedRef<SDockTab>& ConstructUnderMajorTab, const TSharedPtr<SWindow>& ConstructUnderWindow)
+{
+}
+
+
+#undef LOCTEXT_NAMESPACE
diff --git a/Engine/Plugins/Animation/GameplayInsights/Source/RewindDebugger/Private/SRewindDebuggerDetails.h b/Engine/Plugins/Animation/GameplayInsights/Source/RewindDebugger/Private/SRewindDebuggerDetails.h
new file mode 100644
--- /dev/null
+++ b/Engine/Plugins/Animation/GameplayInsights/Source/RewindDebugger/Private/SRewindDebuggerDetails.h
@@ -0,0 +1,44 @@
+// Copyright Epic Games, Inc. All Rights Reserved.
+#pragma once
+
+#include "BindableProperty.h"
+#include "CoreMinimal.h"
+#include "Framework/Commands/UICommandList.h"
+#include "Framework/Docking/TabManager.h"
+#include "IRewindDebuggerView.h"
+#include "RewindDebuggerModule.h"
+#include "SRewindDebuggerComponentTree.h"
+#include "Widgets/DeclarativeSyntaxSupport.h"
+#include "Widgets/Input/SSearchBox.h"
+#include "Widgets/SCompoundWidget.h"
+
+class SDockTab;
+
+class SRewindDebuggerDetails : public SCompoundWidget
+{
+	typedef TBindablePropertyInitializer<FString, BindingType_Out> DebugTargetInitializer;
+
+public:
+	SLATE_BEGIN_ARGS(SRewindDebuggerDetails) { }
+	SLATE_END_ARGS()
+	
+public:
+
+	/**
+	* Default constructor.
+	*/
+	SRewindDebuggerDetails();
+	virtual ~SRewindDebuggerDetails();
+
+	/**
+	* Constructs the application.
+	*
+	* @param InArgs The Slate argument list.
+	* @param ConstructUnderMajorTab The major tab which will contain the session front-end.
+	* @param ConstructUnderWindow The window in which this widget is being constructed.
+	* @param InStyleSet The style set to use.
+	*/
+	void Construct(const FArguments& InArgs, const TSharedRef<SDockTab>& ConstructUnderMajorTab, const TSharedPtr<SWindow>& ConstructUnderWindow);
+
+private:
+};
diff --git a/Engine/Plugins/Animation/GameplayInsights/Source/RewindDebugger/Private/SRewindDebuggerTimelines.cpp b/Engine/Plugins/Animation/GameplayInsights/Source/RewindDebugger/Private/SRewindDebuggerTimelines.cpp
new file mode 100644
--- /dev/null
+++ b/Engine/Plugins/Animation/GameplayInsights/Source/RewindDebugger/Private/SRewindDebuggerTimelines.cpp
@@ -0,0 +1,186 @@
+// Copyright Epic Games, Inc. All Rights Reserved.
+
+#include "SRewindDebuggerTimelines.h"
+
+#include "ISequencerWidgetsModule.h"
+#include "ObjectTrace.h"
+#include "Widgets/Layout/SSpacer.h"
+
+#include "RewindDebugger.h"
+#include "SSimpleTimeSlider.h"
+
+#define LOCTEXT_NAMESPACE "SAnimationInsights"
+
+
+class SRewindDebuggerTimelineTableRow : public STableRow<TSharedPtr<RewindDebugger::FRewindDebuggerTrack>>
+{
+	public:
+	
+	SLATE_BEGIN_ARGS(SRewindDebuggerTimelineTableRow) {}
+	SLATE_END_ARGS()
+	
+	void Construct(const SRewindDebuggerTimelineTableRow::FArguments& InArgs, const TSharedRef<STableViewBase>& InOwnerTableView)
+	{
+		STableRow<TSharedPtr<RewindDebugger::FRewindDebuggerTrack>>::Construct( STableRow<TSharedPtr<RewindDebugger::FRewindDebuggerTrack>>::FArguments(), InOwnerTableView);
+
+		SetExpanderArrowVisibility(EVisibility::Collapsed);
+	}
+
+	virtual int32 GetIndentLevel() const override
+	{
+		// don't indent timeline tracks
+		return 0;
+	}
+};
+
+
+SRewindDebuggerTimelines::SRewindDebuggerTimelines() 
+	: SSimpleTimeSlider()
+	, DebugComponents(nullptr)
+{ 
+}
+
+SRewindDebuggerTimelines::~SRewindDebuggerTimelines() 
+{
+
+}
+
+void SRewindDebuggerTimelines::SetSelection(TSharedPtr<RewindDebugger::FRewindDebuggerTrack> SelectedItem)
+{
+	ComponentTreeView->SetSelection(SelectedItem);
+}
+
+void SRewindDebuggerTimelines::ScrollTo(double ScrollOffset)
+{
+	ComponentTreeView->SetScrollOffset(ScrollOffset);
+}
+
+
+TSharedRef<ITableRow> SRewindDebuggerTimelines::ComponentTreeViewGenerateRow(TSharedPtr<RewindDebugger::FRewindDebuggerTrack> InItem, const TSharedRef<STableViewBase>& OwnerTable)
+{
+	TSharedPtr<SRewindDebuggerTimelineTableRow> Track = SNew(SRewindDebuggerTimelineTableRow, OwnerTable);
+	TSharedPtr<SWidget> TimelineView = InItem->GetTimelineView();
+
+	if (!TimelineView.IsValid())
+	{
+		TimelineView = SNew(SSpacer).Size(FVector2D(100, 20));
+	}
+
+	Track->SetContent(TimelineView.ToSharedRef());
+
+	return Track.ToSharedRef();
+}
+
+void TimelineViewGetChildren(TSharedPtr<RewindDebugger::FRewindDebuggerTrack> InItem, TArray<TSharedPtr<RewindDebugger::FRewindDebuggerTrack>>& OutChildren)
+{
+	InItem->IterateSubTracks([&OutChildren](TSharedPtr<RewindDebugger::FRewindDebuggerTrack> Track)
+	{
+		OutChildren.Add(Track);
+	});
+}
+
+void SRewindDebuggerTimelines::TimelineViewExpansionChanged(TSharedPtr<RewindDebugger::FRewindDebuggerTrack> InItem, bool bShouldBeExpanded)
+{
+	InItem->SetIsExpanded(bShouldBeExpanded);
+	OnExpansionChanged.ExecuteIfBound();
+}
+
+void SRewindDebuggerTimelines::Construct(const FArguments& InArgs)
+{
+	SSimpleTimeSlider::Construct(SSimpleTimeSlider::FArguments()
+		.ViewRange(InArgs._ViewRange)
+		.OnViewRangeChanged(InArgs._OnViewRangeChanged)
+		.ClampRange(InArgs._ClampRange)
+		.ScrubPosition(InArgs._ScrubPosition)
+		.OnScrubPositionChanged(InArgs._OnScrubPositionChanged));
+
+	DebugComponents = InArgs._DebugComponents;
+
+	OnExpansionChanged = InArgs._OnExpansionChanged;
+
+	ComponentTreeView = SNew(STreeView<TSharedPtr<RewindDebugger::FRewindDebuggerTrack>>)
+									.ItemHeight(20.0f)
+									.OnSelectionChanged(InArgs._OnSelectionChanged)
+									.TreeItemsSource(DebugComponents)
+									.OnGenerateRow(this, &SRewindDebuggerTimelines::ComponentTreeViewGenerateRow)
+									.OnGetChildren_Static(&TimelineViewGetChildren)
+									.OnExpansionChanged(this, &SRewindDebuggerTimelines::TimelineViewExpansionChanged)
+									.SelectionMode(ESelectionMode::Single)
+									.AllowOverscroll(EAllowOverscroll::No)
+									.ExternalScrollbar(InArgs._ExternalScrollbar)
+									.OnTreeViewScrolled(InArgs._OnScrolled);
+
+	ChildSlot
+	[
+		ComponentTreeView.ToSharedRef()
+	];
+}
+
+int32 SRewindDebuggerTimelines::OnPaint( const FPaintArgs& Args, const FGeometry& AllottedGeometry, const FSlateRect& MyCullingRect, FSlateWindowElementList& OutDrawElements, int32 LayerId, const FWidgetStyle& InWidgetStyle, bool bParentEnabled ) const
+{
+	TRange<double> LocalViewRange = ViewRange.Get();
+	FScrubRangeToScreen RangeToScreen( LocalViewRange, AllottedGeometry.GetLocalSize() );
+
+	// Paint Child Widgets (Tracks) 
+	LayerId = SCompoundWidget::OnPaint(Args, AllottedGeometry, MyCullingRect, OutDrawElements, LayerId, InWidgetStyle, ShouldBeEnabled(bParentEnabled));
+
+	// Paint Major Time Marks
+	FDrawTickArgs TickArgs;
+	TickArgs.AllottedGeometry = AllottedGeometry;
+	TickArgs.bOnlyDrawMajorTicks = true;
+	TickArgs.TickColor = FLinearColor(0.2f,0.2f,0.2f,0.2f);
+	TickArgs.ClippingRect = MyCullingRect;
+	TickArgs.DrawEffects = ESlateDrawEffect::None;
+	TickArgs.StartLayer = ++LayerId;
+	TickArgs.TickOffset = 0;
+	TickArgs.MajorTickHeight = AllottedGeometry.Size.Y;
+	DrawTicks(OutDrawElements, RangeToScreen, TickArgs);
+    	
+	// Paint Cursor
+	const float XPos = RangeToScreen.InputToLocalX( ScrubPosition.Get() );
+
+	TArray<FVector2D> Points {{XPos, 0}, { XPos, AllottedGeometry.Size.Y}};
+	
+
+	FSlateDrawElement::MakeLines(
+				OutDrawElements,
+						++LayerId,
+						AllottedGeometry.ToPaintGeometry(),
+						Points,
+						ESlateDrawEffect::None,
+						FLinearColor::White,
+						false
+						);
+
+	return LayerId;
+}
+
+static void RestoreTimelineExpansion(TSharedPtr<RewindDebugger::FRewindDebuggerTrack> Track, TSharedPtr<STreeView<TSharedPtr<RewindDebugger::FRewindDebuggerTrack>>>& TreeView)
+{
+	TreeView->SetItemExpansion(Track, Track->GetIsExpanded());
+	Track->IterateSubTracks([&TreeView](TSharedPtr<RewindDebugger::FRewindDebuggerTrack> SubTrack)
+		{
+			RestoreTimelineExpansion(SubTrack, TreeView);
+		});
+}
+
+void SRewindDebuggerTimelines::RestoreExpansion()
+{
+	for (auto& Track : *DebugComponents)
+	{
+		::RestoreTimelineExpansion(Track, ComponentTreeView);
+	}
+}
+
+void SRewindDebuggerTimelines::Refresh()
+{
+	ComponentTreeView->RebuildList();
+
+	if (DebugComponents)
+	{
+		// make sure any newly added TreeView nodes are created expanded
+		RestoreExpansion();
+	}
+}
+
+#undef LOCTEXT_NAMESPACE
diff --git a/Engine/Plugins/Animation/GameplayInsights/Source/RewindDebugger/Private/SRewindDebuggerTimelines.h b/Engine/Plugins/Animation/GameplayInsights/Source/RewindDebugger/Private/SRewindDebuggerTimelines.h
new file mode 100644
--- /dev/null
+++ b/Engine/Plugins/Animation/GameplayInsights/Source/RewindDebugger/Private/SRewindDebuggerTimelines.h
@@ -0,0 +1,78 @@
+// Copyright Epic Games, Inc. All Rights Reserved.
+#pragma once
+
+#include "CoreMinimal.h"
+#include "IRewindDebugger.h"
+#include "RewindDebuggerTrack.h"
+#include "Widgets/DeclarativeSyntaxSupport.h"
+#include "Widgets/SCompoundWidget.h"
+#include "Widgets/Views/STreeView.h"
+#include "SSimpleTimeSlider.h"
+
+
+class SRewindDebuggerTimelines : public SSimpleTimeSlider
+{
+	using FOnSelectionChanged = typename STreeView<TSharedPtr<RewindDebugger::FRewindDebuggerTrack>>::FOnSelectionChanged;
+public:
+	SLATE_BEGIN_ARGS(SRewindDebuggerTimelines) {}
+	 
+		SLATE_ARGUMENT( TArray< TSharedPtr< RewindDebugger::FRewindDebuggerTrack > >*, DebugComponents );
+		SLATE_ARGUMENT( TSharedPtr<SScrollBar>, ExternalScrollbar)
+		SLATE_EVENT( FOnSelectionChanged, OnSelectionChanged )
+		SLATE_EVENT( FOnTableViewScrolled, OnScrolled)
+		SLATE_EVENT( FSimpleDelegate, OnExpansionChanged )
+
+	
+		/** The scrub position */
+    	SLATE_ATTRIBUTE(double, ScrubPosition);
+    
+    	/** View time range */
+    	SLATE_ATTRIBUTE(TRange<double>, ViewRange);
+    
+    	/** Clamp time range */
+    	SLATE_ATTRIBUTE(TRange<double>, ClampRange);
+    
+		/** Called when the scrub position changes */
+		SLATE_EVENT(FOnScrubPositionChanged, OnScrubPositionChanged);
+	
+		/** Called right before the scrubber begins to move */
+		SLATE_EVENT(FSimpleDelegate, OnBeginScrubberMovement);
+	
+		/** Called right after the scrubber handle is released by the user */
+		SLATE_EVENT(FSimpleDelegate, OnEndScrubberMovement);
+	
+    	/** Called when the view range changes */
+    	SLATE_EVENT(FOnRangeChanged, OnViewRangeChanged);
+    	
+	SLATE_END_ARGS()
+
+	/**
+	* Default constructor.
+	*/
+	SRewindDebuggerTimelines();
+	virtual ~SRewindDebuggerTimelines();
+
+	/**
+	* Constructs the application.
+	*
+	* @param InArgs The Slate argument list.
+	*/
+	void Construct(const FArguments& InArgs);
+
+	void Refresh();
+	void RestoreExpansion();
+
+	void SetSelection(TSharedPtr<RewindDebugger::FRewindDebuggerTrack> SelectedItem);
+	void ScrollTo(double ScrollOffset);
+	
+	virtual int32 OnPaint( const FPaintArgs& Args, const FGeometry& AllottedGeometry, const FSlateRect& MyCullingRect, FSlateWindowElementList& OutDrawElements, int32 LayerId, const FWidgetStyle& InWidgetStyle, bool bParentEnabled ) const override;
+
+private:
+	TSharedRef<ITableRow> ComponentTreeViewGenerateRow(TSharedPtr<RewindDebugger::FRewindDebuggerTrack> InItem, const TSharedRef<STableViewBase>& OwnerTable);
+	void TimelineViewExpansionChanged(TSharedPtr<RewindDebugger::FRewindDebuggerTrack> InItem, bool bShouldBeExpanded);
+	
+	TArray<TSharedPtr<RewindDebugger::FRewindDebuggerTrack>>* DebugComponents;
+	TSharedPtr<STreeView<TSharedPtr<RewindDebugger::FRewindDebuggerTrack>>> ComponentTreeView;
+	
+	FSimpleDelegate OnExpansionChanged;
+};
diff --git a/Engine/Plugins/Experimental/Animation/PoseSearch/Source/Editor/Private/PoseSearchDebugger.cpp b/Engine/Plugins/Experimental/Animation/PoseSearch/Source/Editor/Private/PoseSearchDebugger.cpp
--- a/Engine/Plugins/Experimental/Animation/PoseSearch/Source/Editor/Private/PoseSearchDebugger.cpp
+++ b/Engine/Plugins/Experimental/Animation/PoseSearch/Source/Editor/Private/PoseSearchDebugger.cpp
@@ -2827,6 +2827,29 @@ TSharedPtr<IRewindDebuggerView> FDebuggerViewCreator::CreateDebugView(uint64 Obj
 	return FDebugger::Get()->GenerateInstance(ObjectId);
 }
 
+bool FDebuggerViewCreator::HasDebugInfo(uint64 AnimInstanceId, const TraceServices::IAnalysisSession& Session) const
+{
+	// Get provider and validate
+	TraceServices::FAnalysisSessionReadScope SessionReadScope(Session);
+
+	const FTraceProvider* PoseSearchProvider = Session.ReadProvider<FTraceProvider>(FTraceProvider::ProviderName);
+	const IAnimationProvider* AnimationProvider = Session.ReadProvider<IAnimationProvider>("AnimationProvider");
+	const IGameplayProvider* GameplayProvider = Session.ReadProvider<IGameplayProvider>("GameplayProvider");
+	if (!(PoseSearchProvider && AnimationProvider && GameplayProvider))
+	{
+		return false;
+	}
+	
+	bool bHasData = false;
+	
+	PoseSearchProvider->EnumerateMotionMatchingStateTimelines(AnimInstanceId, [&bHasData](const FTraceProvider::FMotionMatchingStateTimeline& InTimeline)
+	{
+		bHasData = true;
+	});
+	
+	return bHasData;
+}
+
 FName FDebuggerViewCreator::GetName() const
 {
 	static const FName Name("PoseSearchDebugger");
diff --git a/Engine/Plugins/Experimental/Animation/PoseSearch/Source/Editor/Private/PoseSearchDebugger.h b/Engine/Plugins/Experimental/Animation/PoseSearch/Source/Editor/Private/PoseSearchDebugger.h
--- a/Engine/Plugins/Experimental/Animation/PoseSearch/Source/Editor/Private/PoseSearchDebugger.h
+++ b/Engine/Plugins/Experimental/Animation/PoseSearch/Source/Editor/Private/PoseSearchDebugger.h
@@ -758,6 +758,7 @@ public:
 	
 	/** Creates the PoseSearch Slate view for the provided AnimInstance */
 	virtual TSharedPtr<IRewindDebuggerView> CreateDebugView(uint64 ObjectId, double CurrentTime, const TraceServices::IAnalysisSession& InAnalysisSession) const override;
+	virtual bool HasDebugInfo(uint64 ObjectId, const TraceServices::IAnalysisSession& InAnalysisSession) const override;
 };
 
 
diff --git a/Engine/Source/Developer/ToolWidgets/Private/SSimpleTimeSlider.cpp b/Engine/Source/Developer/ToolWidgets/Private/SSimpleTimeSlider.cpp
--- a/Engine/Source/Developer/ToolWidgets/Private/SSimpleTimeSlider.cpp
+++ b/Engine/Source/Developer/ToolWidgets/Private/SSimpleTimeSlider.cpp
@@ -22,6 +22,7 @@ void SSimpleTimeSlider::Construct( const SSimpleTimeSlider::FArguments& InArgs )
     OnViewRangeChanged = InArgs._OnViewRangeChanged;
 	ClampRangeHighlightColor = InArgs._ClampRangeHighlightColor;
 	ClampRangeHighlightSize = InArgs._ClampRangeHighlightSize;
+	DesiredSize = InArgs._DesiredSize;
 	
 	DistanceDragged = 0.0f;
 	bDraggingScrubber = false;
@@ -44,7 +45,7 @@ int32 SSimpleTimeSlider::OnPaint( const FPaintArgs& Args, const FGeometry& Allot
 
 FVector2D SSimpleTimeSlider::ComputeDesiredSize( float ) const
 {
-	return FVector2D(100, 22);
+	return DesiredSize;
 }
 
 namespace ScrubConstants
@@ -56,40 +57,6 @@ namespace ScrubConstants
 	const float MinDisplayTickSpacing = 0.001f;
 }
 
-
-/** Utility struct for converting between scrub range space and local/absolute screen space */
-struct SSimpleTimeSlider::FScrubRangeToScreen
-{
-	FVector2D WidgetSize;
-
-	TRange<double> ViewInput;
-	float ViewInputRange;
-	float PixelsPerInput;
-
-	FScrubRangeToScreen(TRange<double> InViewInput, const FVector2D& InWidgetSize )
-	{
-		WidgetSize = InWidgetSize;
-
-		ViewInput = InViewInput;
-		ViewInputRange = ViewInput.Size<float>();
-		PixelsPerInput = ViewInputRange > 0 ? ( WidgetSize.X / ViewInputRange ) : 0;
-	}
-
-	/** Local Widget Space -> Curve Input domain. */
-	float LocalXToInput(float ScreenX) const
-	{
-		float LocalX = ScreenX;
-		return (LocalX/PixelsPerInput) + ViewInput.GetLowerBoundValue();
-	}
-
-	/** Curve Input domain -> local Widget Space */
-	float InputToLocalX(float Input) const
-	{
-		return (Input - ViewInput.GetLowerBoundValue()) * PixelsPerInput;
-	}
-};
-
-
 /**
  * Gets the the next spacing value in the series 
  * to determine a good spacing value
@@ -138,28 +105,7 @@ static double DetermineOptimalSpacing(float InPixelsPerInput, uint32 InMinTick,
 	return Spacing;
 }
 
-struct SSimpleTimeSlider::FDrawTickArgs
-{
-	/** Geometry of the area */
-	FGeometry AllottedGeometry;
-	/** Clipping rect of the area */
-	FSlateRect ClippingRect;
-	/** Color of each tick */
-	FLinearColor TickColor;
-	/** Offset in Y where to start the tick */
-	float TickOffset;
-	/** Height in of major ticks */
-	float MajorTickHeight;
-	/** Start layer for elements */
-	int32 StartLayer;
-	/** Draw effects to apply */
-	ESlateDrawEffect DrawEffects;
-	/** Whether or not to only draw major ticks */
-	bool bOnlyDrawMajorTicks;
-	/** Whether or not to mirror labels */
-	bool bMirrorLabels;
-	
-};
+
 
 void SSimpleTimeSlider::DrawTicks( FSlateWindowElementList& OutDrawElements, const struct FScrubRangeToScreen& RangeToScreen, FDrawTickArgs& InArgs ) const
 {
@@ -361,7 +307,7 @@ int32 SSimpleTimeSlider::OnPaintTimeSlider( bool bMirrorLabels, const FGeometry&
 	return LayerId;
 }
 
-FReply SSimpleTimeSlider::OnMouseButtonDown( const FGeometry& MyGeometry, const FPointerEvent& MouseEvent )
+FReply SSimpleTimeSlider::OnPreviewMouseButtonDown(const FGeometry& MyGeometry, const FPointerEvent& MouseEvent )
 {
 	bool bHandleLeftMouseButton = MouseEvent.GetEffectingButton() == EKeys::LeftMouseButton;
 	bool bHandleRightMouseButton = MouseEvent.GetEffectingButton() == EKeys::RightMouseButton && AllowPan.Get();
@@ -553,7 +499,7 @@ FReply SSimpleTimeSlider::OnMouseWheel( const FGeometry& MyGeometry, const FPoin
 {
 	FReply ReturnValue = FReply::Unhandled();
 
-	if ( AllowZoom.Get() )
+	if ( AllowZoom.Get() && MouseEvent.GetModifierKeys().IsControlDown())
 	{
 		const float ZoomDelta = -0.1f * MouseEvent.GetWheelDelta();
 
diff --git a/Engine/Source/Developer/ToolWidgets/Public/SSimpleTimeSlider.h b/Engine/Source/Developer/ToolWidgets/Public/SSimpleTimeSlider.h
--- a/Engine/Source/Developer/ToolWidgets/Public/SSimpleTimeSlider.h
+++ b/Engine/Source/Developer/ToolWidgets/Public/SSimpleTimeSlider.h
@@ -28,6 +28,7 @@ public:
 		, _CursorSize(0)
 		, _ClampRangeHighlightColor(FLinearColor(0.05,0.05,0.05,1.0))
 		, _ClampRangeHighlightSize(1.0)
+		, _DesiredSize(100,22)
 	{}
 		/* If we should mirror the labels on the timeline */
 		SLATE_ATTRIBUTE( bool, MirrorLabels )
@@ -55,6 +56,9 @@ public:
 
 		/** Size of clamp range overlay 0.0 for none 1.0 for the height of the slider */
 		SLATE_ATTRIBUTE(float, ClampRangeHighlightSize);
+
+		/* Widget desired size */
+		SLATE_ARGUMENT(FVector2d, DesiredSize);
     
 	
 		/** Called when the scrub position changes */
@@ -84,15 +88,66 @@ public:
 	void SetClampRange(double MinValue, double MaxValue);
 	bool IsPanning() { return bPanning; }
 
+	/** Utility struct for converting between scrub range space and local/absolute screen space */
+	struct FScrubRangeToScreen
+	{
+		FVector2D WidgetSize;
+	
+		TRange<double> ViewInput;
+		float ViewInputRange;
+		float PixelsPerInput;
+	
+		FScrubRangeToScreen(TRange<double> InViewInput, const FVector2D& InWidgetSize )
+		{
+			WidgetSize = InWidgetSize;
+	
+			ViewInput = InViewInput;
+			ViewInputRange = ViewInput.Size<float>();
+			PixelsPerInput = ViewInputRange > 0 ? ( WidgetSize.X / ViewInputRange ) : 0;
+		}
+	
+		/** Local Widget Space -> Curve Input domain. */
+		float LocalXToInput(float ScreenX) const
+		{
+			float LocalX = ScreenX;
+			return (LocalX/PixelsPerInput) + ViewInput.GetLowerBoundValue();
+		}
+	
+		/** Curve Input domain -> local Widget Space */
+		float InputToLocalX(float Input) const
+		{
+			return (Input - ViewInput.GetLowerBoundValue()) * PixelsPerInput;
+		}
+	};
+
 protected:
-	// forward declared as class members to prevent name collision with similar types defined in other units
-	struct FScrubRangeToScreen;
-	struct FDrawTickArgs;
+	
+	struct FDrawTickArgs
+    {
+    	/** Geometry of the area */
+    	FGeometry AllottedGeometry;
+    	/** Clipping rect of the area */
+    	FSlateRect ClippingRect;
+    	/** Color of each tick */
+    	FLinearColor TickColor;
+    	/** Offset in Y where to start the tick */
+    	float TickOffset;
+    	/** Height in of major ticks */
+    	float MajorTickHeight;
+    	/** Start layer for elements */
+    	int32 StartLayer;
+    	/** Draw effects to apply */
+    	ESlateDrawEffect DrawEffects;
+    	/** Whether or not to only draw major ticks */
+    	bool bOnlyDrawMajorTicks;
+    	/** Whether or not to mirror labels */
+    	bool bMirrorLabels;
+    };
 
 	// SWidget interface
 	virtual FVector2D ComputeDesiredSize(float) const override;
 	virtual int32 OnPaint( const FPaintArgs& Args, const FGeometry& AllottedGeometry, const FSlateRect& MyCullingRect, FSlateWindowElementList& OutDrawElements, int32 LayerId, const FWidgetStyle& InWidgetStyle, bool bParentEnabled ) const override;
-	virtual FReply OnMouseButtonDown( const FGeometry& MyGeometry, const FPointerEvent& MouseEvent ) override;
+	virtual FReply OnPreviewMouseButtonDown( const FGeometry& MyGeometry, const FPointerEvent& MouseEvent ) override;
 	virtual FReply OnMouseButtonUp( const FGeometry& MyGeometry, const FPointerEvent& MouseEvent ) override;
 	virtual FReply OnMouseMove( const FGeometry& MyGeometry, const FPointerEvent& MouseEvent ) override;
 	virtual FReply OnMouseWheel( const FGeometry& MyGeometry, const FPointerEvent& MouseEvent ) override;
@@ -147,6 +202,7 @@ protected:
 	/***/
 	TSharedPtr<SScrollBar> Scrollbar;
 	FVector2D SoftwareCursorPosition;
+	FVector2D DesiredSize;
 
 	
 };
diff --git a/Engine/Source/Editor/RewindDebuggerInterface/Private/RewindDebuggerInterface.cpp b/Engine/Source/Editor/RewindDebuggerInterface/Private/RewindDebuggerInterface.cpp
--- a/Engine/Source/Editor/RewindDebuggerInterface/Private/RewindDebuggerInterface.cpp
+++ b/Engine/Source/Editor/RewindDebuggerInterface/Private/RewindDebuggerInterface.cpp
@@ -2,6 +2,7 @@
 
 #include "IRewindDebuggerExtension.h"
 #include "IRewindDebuggerViewCreator.h"
+#include "IRewindDebuggerTrackCreator.h"
 #include "IRewindDebuggerDoubleClickHandler.h"
 #include "IRewindDebugger.h"
 
@@ -9,10 +10,22 @@ const FName IRewindDebuggerExtension::ModularFeatureName = "RewindDebuggerExtens
 const FName IRewindDebuggerViewCreator::ModularFeatureName = "RewindDebuggerViewCreator";
 const FName IRewindDebuggerDoubleClickHandler::ModularFeatureName = "RewindDebuggerDoubleClickHandler";
 
+namespace RewindDebugger
+{
+	const FName IRewindDebuggerTrackCreator::ModularFeatureName = "RewindDebuggerTrackCreator";
+}
+
+IRewindDebugger* IRewindDebugger::InternalInstance = nullptr;
+
 IRewindDebugger::~IRewindDebugger()
 {
 }
 
 IRewindDebugger::IRewindDebugger()
 {
+}
+
+IRewindDebugger* IRewindDebugger::Instance()
+{
+	return InternalInstance;
 }
\ No newline at end of file
diff --git a/Engine/Source/Editor/RewindDebuggerInterface/Public/IRewindDebugger.h b/Engine/Source/Editor/RewindDebuggerInterface/Public/IRewindDebugger.h
--- a/Engine/Source/Editor/RewindDebuggerInterface/Public/IRewindDebugger.h
+++ b/Engine/Source/Editor/RewindDebuggerInterface/Public/IRewindDebugger.h
@@ -45,6 +45,12 @@ public:
 	// get the time the debugger is scrubbed to, in seconds since the capture started (or the recording duration while the game is running)
 	virtual double CurrentTraceTime() const = 0;
 
+	// current visible range in trace/profiler units (same units as CurrentTraceTime)
+	virtual const TRange<double>& GetCurrentTraceRange() const = 0;
+	
+	// current visible range in Rewind Debugger recording time units
+	virtual const TRange<double>& GetCurrentViewRange() const = 0;
+
 	// get the current analysis session
 	virtual const TraceServices::IAnalysisSession* GetAnalysisSession() const = 0;
 
@@ -71,4 +77,10 @@ public:
 
 	// returns the length of the current recording
 	virtual double GetRecordingDuration() const = 0;
+
+	// get the current IRewindDebugger instance
+	static IRewindDebugger* Instance();
+	
+protected:
+	static IRewindDebugger* InternalInstance;
 };
diff --git a/Engine/Source/Editor/RewindDebuggerInterface/Public/IRewindDebuggerTrackCreator.h b/Engine/Source/Editor/RewindDebuggerInterface/Public/IRewindDebuggerTrackCreator.h
new file mode 100644
--- /dev/null
+++ b/Engine/Source/Editor/RewindDebuggerInterface/Public/IRewindDebuggerTrackCreator.h
@@ -0,0 +1,68 @@
+// Copyright Epic Games, Inc. All Rights Reserved.
+
+#pragma once
+
+#include "CoreMinimal.h"
+#include "IRewindDebuggerView.h"
+#include "Textures/SlateIcon.h"
+
+namespace TraceServices
+{
+	class IAnalysisSession;
+}
+
+class IRewindDebugger;
+
+namespace RewindDebugger
+{
+	
+class FRewindDebuggerTrack;
+
+// Interface class which creates tracks
+class REWINDDEBUGGERINTERFACE_API IRewindDebuggerTrackCreator : public IModularFeature
+{
+	public: 
+	static const FName ModularFeatureName;
+
+	// returns the type of UObject this track creator can create tracks for
+	FName GetTargetTypeName() const
+	{
+		return GetTargetTypeNameInternal();
+	}
+
+	// returns an identifying Name for this type of track
+	FName GetName() const
+	{
+		return GetNameInternal();
+	}
+
+	// optional additional filter, to prevent debug views from being listed if they have no data
+	bool HasDebugInfo(uint64 ObjectId) const
+	{
+		return HasDebugInfoInternal(ObjectId);
+	};
+	
+	// Create a track which will be shown in the timeline view and tree view, as a child track of the Object
+	TSharedPtr<RewindDebugger::FRewindDebuggerTrack> CreateTrack(uint64 ObjectId) const
+	{
+		return CreateTrackInternal(ObjectId);
+	}
+
+	private:
+
+	virtual FName GetTargetTypeNameInternal() const { return "Object"; }
+	
+	virtual FName GetNameInternal() const { return FName(); }
+	
+	virtual bool HasDebugInfoInternal(uint64 ObjectId) const
+	{
+		return true;
+	};
+	
+	virtual TSharedPtr<FRewindDebuggerTrack> CreateTrackInternal(uint64 ObjectId) const
+	{
+		return TSharedPtr<FRewindDebuggerTrack>();
+	}
+};
+	
+}
\ No newline at end of file
diff --git a/Engine/Source/Editor/RewindDebuggerInterface/Public/IRewindDebuggerView.h b/Engine/Source/Editor/RewindDebuggerInterface/Public/IRewindDebuggerView.h
--- a/Engine/Source/Editor/RewindDebuggerInterface/Public/IRewindDebuggerView.h
+++ b/Engine/Source/Editor/RewindDebuggerInterface/Public/IRewindDebuggerView.h
@@ -6,8 +6,6 @@
 #include "Widgets/SCompoundWidget.h"
 
 // Interface class for views of recorded debug data
-// a delegate which creates implementations of this interface can be registered with a UObject type name to IGameplayInsightsDebugViewCreator
-// and they will be shown in the Rewind Debugger when that type of object is selected
 class REWINDDEBUGGERINTERFACE_API IRewindDebuggerView : public SCompoundWidget
 {
 	public:
diff --git a/Engine/Source/Editor/RewindDebuggerInterface/Public/IRewindDebuggerViewCreator.h b/Engine/Source/Editor/RewindDebuggerInterface/Public/IRewindDebuggerViewCreator.h
--- a/Engine/Source/Editor/RewindDebuggerInterface/Public/IRewindDebuggerViewCreator.h
+++ b/Engine/Source/Editor/RewindDebuggerInterface/Public/IRewindDebuggerViewCreator.h
@@ -19,6 +19,12 @@ public:
 
 	// returns the name of a type of UObject for which this debug view will be created
 	virtual FName GetTargetTypeName() const = 0;
+	
+	// optional additional filter, to prevent debug views from being listed if they have no data
+	virtual bool HasDebugInfo(uint64 ObjectId) const
+	{
+		return true;
+	};
 
 	// returns a unique name for identifying this type of widget (same value returned by IRewindDebuggerView::GetName)
 	virtual FName GetName() const = 0;
diff --git a/Engine/Source/Editor/RewindDebuggerInterface/Public/RewindDebuggerTrack.h b/Engine/Source/Editor/RewindDebuggerInterface/Public/RewindDebuggerTrack.h
new file mode 100644
--- /dev/null
+++ b/Engine/Source/Editor/RewindDebuggerInterface/Public/RewindDebuggerTrack.h
@@ -0,0 +1,98 @@
+// Copyright Epic Games, Inc. All Rights Reserved.
+#pragma once
+
+#include "CoreMinimal.h"
+#include "IRewindDebuggerViewCreator.h"
+
+namespace TraceServices
+{
+	class IAnalysisSession;
+}
+
+namespace RewindDebugger
+{
+	
+class FRewindDebuggerTrack
+{
+public:
+	FRewindDebuggerTrack(): bExpanded(true)
+	{
+	}
+
+	virtual ~FRewindDebuggerTrack()
+	{
+	}
+
+	bool GetIsExpanded()
+	{
+		return bExpanded;
+	}
+
+	void SetIsExpanded(bool bIsExpanded)
+	{
+		bExpanded = bIsExpanded;
+	}
+
+	// Update should do work to compute children etc for the current time range.  Return true if children have changed.
+	bool Update()
+	{
+		return UpdateInternal();
+	}
+
+	// Get a widget to show in the timeline view for this track
+	TSharedPtr<SWidget> GetTimelineView()
+	{
+		return GetTimelineViewInternal();
+	}
+
+	// Get a widget to show in the details tab, when this track is selected
+	TSharedPtr<SWidget> GetDetailsView()
+	{
+		return GetDetailsViewInternal();
+	}
+	
+	// unique name for track (must match creator name if track is created by an IRewindDebuggerViewCreator) 
+	FName GetName() const
+	{
+		return GetNameInternal();
+	}
+
+	// icon to display in the tree view
+	FSlateIcon GetIcon()
+	{
+		return GetIconInternal();
+	}
+
+	// display name for track in Tree View
+	FText GetDisplayName() const
+	{
+		return GetDisplayNameInternal();
+	}
+
+	// insights object id for an object associated with this track
+	uint64 GetObjectId() const
+	{
+		return GetObjectIdInternal();
+	}
+
+	// iterate over all sub-tracks of this track and call Iterator function
+	void IterateSubTracks(TFunction<void(TSharedPtr<FRewindDebuggerTrack> SubTrack)> IteratorFunction)
+	{
+		IterateSubTracksInternal(IteratorFunction);
+	}
+
+private:
+
+	virtual bool UpdateInternal() { return false; }
+	virtual TSharedPtr<SWidget> GetTimelineViewInternal() { return TSharedPtr<SWidget>(); }
+	virtual TSharedPtr<SWidget> GetDetailsViewInternal() { return TSharedPtr<SWidget>(); }
+	virtual FSlateIcon GetIconInternal() { return FSlateIcon(); }
+	virtual FName GetNameInternal() const { return ""; }
+	virtual FText GetDisplayNameInternal() const { return FText(); }
+	virtual uint64 GetObjectIdInternal() const { return 0; }
+	virtual void IterateSubTracksInternal(TFunction<void(TSharedPtr<FRewindDebuggerTrack> SubTrack)> IteratorFunction) { } 
+
+	bool bExpanded;
+};
+	
+}
\ No newline at end of file
diff --git a/Engine/Source/Runtime/Engine/Private/Animation/AnimTrace.cpp b/Engine/Source/Runtime/Engine/Private/Animation/AnimTrace.cpp
--- a/Engine/Source/Runtime/Engine/Private/Animation/AnimTrace.cpp
+++ b/Engine/Source/Runtime/Engine/Private/Animation/AnimTrace.cpp
@@ -26,8 +26,9 @@
 
 UE_TRACE_CHANNEL_DEFINE(AnimationChannel);
 
-UE_TRACE_EVENT_BEGIN(Animation, TickRecord)
+UE_TRACE_EVENT_BEGIN(Animation, TickRecord2)
 	UE_TRACE_EVENT_FIELD(uint64, Cycle)
+	UE_TRACE_EVENT_FIELD(double, RecordingTime)
 	UE_TRACE_EVENT_FIELD(uint64, AnimInstanceId)
 	UE_TRACE_EVENT_FIELD(uint64, AssetId)
 	UE_TRACE_EVENT_FIELD(int32, NodeId)
@@ -49,8 +50,9 @@ UE_TRACE_EVENT_BEGIN(Animation, SkeletalMesh2, NoSync|Important)
 	UE_TRACE_EVENT_FIELD(int32[], ParentIndices)
 UE_TRACE_EVENT_END()
 
-UE_TRACE_EVENT_BEGIN(Animation, SkeletalMeshComponent2)
+UE_TRACE_EVENT_BEGIN(Animation, SkeletalMeshComponent3)
 	UE_TRACE_EVENT_FIELD(uint64, Cycle)
+	UE_TRACE_EVENT_FIELD(double, RecordingTime)
 	UE_TRACE_EVENT_FIELD(uint64, ComponentId)
 	UE_TRACE_EVENT_FIELD(uint64, MeshId)
 	UE_TRACE_EVENT_FIELD(float[], ComponentToWorld)
@@ -227,8 +229,9 @@ UE_TRACE_EVENT_BEGIN(Animation, SyncMarker)
 	UE_TRACE_EVENT_FIELD(uint32, NameId)
 UE_TRACE_EVENT_END()
 
-UE_TRACE_EVENT_BEGIN(Animation, Montage)
+UE_TRACE_EVENT_BEGIN(Animation, Montage2)
 	UE_TRACE_EVENT_FIELD(uint64, Cycle)
+	UE_TRACE_EVENT_FIELD(double, RecordingTime)
 	UE_TRACE_EVENT_FIELD(uint64, AnimInstanceId)
 	UE_TRACE_EVENT_FIELD(uint64, MontageId)
 	UE_TRACE_EVENT_FIELD(uint32, CurrentSectionNameId)
@@ -423,22 +426,23 @@ void FAnimTrace::OutputAnimTickRecord(const FAnimationBaseContext& InContext, co
 			BlendSpaceFilteredPositionY = InTickRecord.BlendSpace.BlendFilter->GetFilterLastOutput().Y;
 		}
 
-		UE_TRACE_LOG(Animation, TickRecord, AnimationChannel)
-			<< TickRecord.Cycle(FPlatformTime::Cycles64())
-			<< TickRecord.AnimInstanceId(FObjectTrace::GetObjectId(AnimInstance))
-			<< TickRecord.AssetId(FObjectTrace::GetObjectId(InTickRecord.SourceAsset))
-			<< TickRecord.NodeId(InContext.GetCurrentNodeId())
-			<< TickRecord.BlendWeight(InTickRecord.EffectiveBlendWeight)
-			<< TickRecord.PlaybackTime(PlaybackTime)
-			<< TickRecord.RootMotionWeight(InTickRecord.RootMotionWeightModifier)
-			<< TickRecord.PlayRate(InTickRecord.PlayRateMultiplier)
-			<< TickRecord.BlendSpacePositionX(BlendSpacePositionX)
-			<< TickRecord.BlendSpacePositionY(BlendSpacePositionY)
-			<< TickRecord.BlendSpaceFilteredPositionX(BlendSpaceFilteredPositionX)
-			<< TickRecord.BlendSpaceFilteredPositionY(BlendSpaceFilteredPositionY)
-			<< TickRecord.FrameCounter(FObjectTrace::GetObjectWorldTickCounter(AnimInstance))
-			<< TickRecord.Looping(InTickRecord.bLooping)
-			<< TickRecord.IsBlendSpace(bIsBlendSpace);
+		UE_TRACE_LOG(Animation, TickRecord2, AnimationChannel)
+			<< TickRecord2.Cycle(FPlatformTime::Cycles64())
+			<< TickRecord2.RecordingTime(FObjectTrace::GetWorldElapsedTime(AnimInstance->GetWorld()))
+			<< TickRecord2.AnimInstanceId(FObjectTrace::GetObjectId(AnimInstance))
+			<< TickRecord2.AssetId(FObjectTrace::GetObjectId(InTickRecord.SourceAsset))
+			<< TickRecord2.NodeId(InContext.GetCurrentNodeId())
+			<< TickRecord2.BlendWeight(InTickRecord.EffectiveBlendWeight)
+			<< TickRecord2.PlaybackTime(PlaybackTime)
+			<< TickRecord2.RootMotionWeight(InTickRecord.RootMotionWeightModifier)
+			<< TickRecord2.PlayRate(InTickRecord.PlayRateMultiplier)
+			<< TickRecord2.BlendSpacePositionX(BlendSpacePositionX)
+			<< TickRecord2.BlendSpacePositionY(BlendSpacePositionY)
+			<< TickRecord2.BlendSpaceFilteredPositionX(BlendSpaceFilteredPositionX)
+			<< TickRecord2.BlendSpaceFilteredPositionY(BlendSpaceFilteredPositionY)
+			<< TickRecord2.FrameCounter(FObjectTrace::GetObjectWorldTickCounter(AnimInstance))
+			<< TickRecord2.Looping(InTickRecord.bLooping)
+			<< TickRecord2.IsBlendSpace(bIsBlendSpace);
 	}
 }
 
@@ -561,16 +565,17 @@ void FAnimTrace::OutputSkeletalMeshComponent(const USkeletalMeshComponent* InCom
 			}
 		}
 
-		UE_TRACE_LOG(Animation, SkeletalMeshComponent2, AnimationChannel)
-			<< SkeletalMeshComponent2.Cycle(FPlatformTime::Cycles64())
-			<< SkeletalMeshComponent2.ComponentId(FObjectTrace::GetObjectId(InComponent))
-			<< SkeletalMeshComponent2.MeshId(FObjectTrace::GetObjectId(InComponent->SkeletalMesh))
-			<< SkeletalMeshComponent2.ComponentToWorld(reinterpret_cast<const float*>(&InComponent->GetComponentToWorld()), sizeof(FTransform) / sizeof(float))
-			<< SkeletalMeshComponent2.Pose(reinterpret_cast<const float*>(InComponent->GetComponentSpaceTransforms().GetData()), BoneCount * (sizeof(FTransform) / sizeof(float)))
-			<< SkeletalMeshComponent2.CurveIds(CurveIds.GetData(), CurveIds.Num())
-			<< SkeletalMeshComponent2.CurveValues(CurveValues.GetData(), CurveValues.Num())
-			<< SkeletalMeshComponent2.LodIndex((uint16)InComponent->GetPredictedLODLevel())
-			<< SkeletalMeshComponent2.FrameCounter(FObjectTrace::GetObjectWorldTickCounter(InComponent));
+		UE_TRACE_LOG(Animation, SkeletalMeshComponent3, AnimationChannel)
+			<< SkeletalMeshComponent3.Cycle(FPlatformTime::Cycles64())
+			<< SkeletalMeshComponent3.RecordingTime(FObjectTrace::GetWorldElapsedTime(InComponent->GetWorld()))
+			<< SkeletalMeshComponent3.ComponentId(FObjectTrace::GetObjectId(InComponent))
+			<< SkeletalMeshComponent3.MeshId(FObjectTrace::GetObjectId(InComponent->SkeletalMesh))
+			<< SkeletalMeshComponent3.ComponentToWorld(reinterpret_cast<const float*>(&InComponent->GetComponentToWorld()), sizeof(FTransform) / sizeof(float))
+			<< SkeletalMeshComponent3.Pose(reinterpret_cast<const float*>(InComponent->GetComponentSpaceTransforms().GetData()), BoneCount * (sizeof(FTransform) / sizeof(float)))
+			<< SkeletalMeshComponent3.CurveIds(CurveIds.GetData(), CurveIds.Num())
+			<< SkeletalMeshComponent3.CurveValues(CurveValues.GetData(), CurveValues.Num())
+			<< SkeletalMeshComponent3.LodIndex((uint16)InComponent->GetPredictedLODLevel())
+			<< SkeletalMeshComponent3.FrameCounter(FObjectTrace::GetObjectWorldTickCounter(InComponent));
 	}
 }
 
@@ -1131,16 +1136,17 @@ void FAnimTrace::OutputMontage(UAnimInstance* InAnimInstance, const FAnimMontage
 		const uint32 CurrentSectionNameId = OutputName(InMontageInstance.GetCurrentSection());
 		const uint32 NextSectionNameId = OutputName(InMontageInstance.GetNextSection());
 
-		UE_TRACE_LOG(Animation, Montage, AnimationChannel)
-			<< Montage.Cycle(FPlatformTime::Cycles64())
-			<< Montage.AnimInstanceId(FObjectTrace::GetObjectId(InAnimInstance))
-			<< Montage.MontageId(FObjectTrace::GetObjectId(InMontageInstance.Montage))
-			<< Montage.CurrentSectionNameId(CurrentSectionNameId)
-			<< Montage.NextSectionNameId(NextSectionNameId)
-			<< Montage.Weight(InMontageInstance.GetWeight())
-			<< Montage.DesiredWeight(InMontageInstance.GetDesiredWeight())
-			<< Montage.Position(InMontageInstance.GetPosition())
-			<< Montage.FrameCounter(FObjectTrace::GetObjectWorldTickCounter(InAnimInstance));
+		UE_TRACE_LOG(Animation, Montage2, AnimationChannel)
+			<< Montage2.Cycle(FPlatformTime::Cycles64())
+			<< Montage2.RecordingTime(FObjectTrace::GetWorldElapsedTime(InAnimInstance->GetWorld()))
+			<< Montage2.AnimInstanceId(FObjectTrace::GetObjectId(InAnimInstance))
+			<< Montage2.MontageId(FObjectTrace::GetObjectId(InMontageInstance.Montage))
+			<< Montage2.CurrentSectionNameId(CurrentSectionNameId)
+			<< Montage2.NextSectionNameId(NextSectionNameId)
+			<< Montage2.Weight(InMontageInstance.GetWeight())
+			<< Montage2.DesiredWeight(InMontageInstance.GetDesiredWeight())
+			<< Montage2.Position(InMontageInstance.GetPosition())
+			<< Montage2.FrameCounter(FObjectTrace::GetObjectWorldTickCounter(InAnimInstance));
 	}
 }
 
diff --git a/Engine/Source/Runtime/Engine/Private/ObjectTrace.cpp b/Engine/Source/Runtime/Engine/Private/ObjectTrace.cpp
--- a/Engine/Source/Runtime/Engine/Private/ObjectTrace.cpp
+++ b/Engine/Source/Runtime/Engine/Private/ObjectTrace.cpp
@@ -38,13 +38,15 @@ UE_TRACE_EVENT_BEGIN(Object, Object)
 	UE_TRACE_EVENT_FIELD(UE::Trace::WideString, Path)
 UE_TRACE_EVENT_END()
 
-UE_TRACE_EVENT_BEGIN(Object, ObjectLifetimeBegin)
+UE_TRACE_EVENT_BEGIN(Object, ObjectLifetimeBegin2)
 	UE_TRACE_EVENT_FIELD(uint64, Cycle)
+	UE_TRACE_EVENT_FIELD(double, RecordingTime)
 	UE_TRACE_EVENT_FIELD(uint64, Id)
 UE_TRACE_EVENT_END()
 
-UE_TRACE_EVENT_BEGIN(Object, ObjectLifetimeEnd)
+UE_TRACE_EVENT_BEGIN(Object, ObjectLifetimeEnd2)
 	UE_TRACE_EVENT_FIELD(uint64, Cycle)
+	UE_TRACE_EVENT_FIELD(double, RecordingTime)
 	UE_TRACE_EVENT_FIELD(uint64, Id)
 UE_TRACE_EVENT_END()
 
@@ -450,9 +452,10 @@ void FObjectTrace::OutputObjectLifetimeBegin(const UObject* InObject)
 
 	TRACE_OBJECT(InObject);
 
-	UE_TRACE_LOG(Object, ObjectLifetimeBegin, ObjectChannel)
-		<< ObjectLifetimeBegin.Cycle(FPlatformTime::Cycles64())
-		<< ObjectLifetimeBegin.Id(GetObjectId(InObject));
+	UE_TRACE_LOG(Object, ObjectLifetimeBegin2, ObjectChannel)
+		<< ObjectLifetimeBegin2.Cycle(FPlatformTime::Cycles64())
+		<< ObjectLifetimeBegin2.RecordingTime(FObjectTrace::GetWorldElapsedTime(InObject->GetWorld()))
+		<< ObjectLifetimeBegin2.Id(GetObjectId(InObject));
 }
 
 void FObjectTrace::OutputObjectLifetimeEnd(const UObject* InObject)
@@ -475,9 +478,10 @@ void FObjectTrace::OutputObjectLifetimeEnd(const UObject* InObject)
 
 	TRACE_OBJECT(InObject);
 
-	UE_TRACE_LOG(Object, ObjectLifetimeEnd, ObjectChannel)
-		<< ObjectLifetimeEnd.Cycle(FPlatformTime::Cycles64())
-		<< ObjectLifetimeEnd.Id(GetObjectId(InObject));
+	UE_TRACE_LOG(Object, ObjectLifetimeEnd2, ObjectChannel)
+		<< ObjectLifetimeEnd2.Cycle(FPlatformTime::Cycles64())
+		<< ObjectLifetimeEnd2.RecordingTime(FObjectTrace::GetWorldElapsedTime(InObject->GetWorld()))
+		<< ObjectLifetimeEnd2.Id(GetObjectId(InObject));
 }
 
 void FObjectTrace::OutputPawnPossess(const UObject* InController, const UObject* InPawn)
