Fix Rewind debugger compile issues

diff --git a/Engine/Plugins/Animation/GameplayInsights/Source/RewindDebugger/Private/RewindDebuggerModule.cpp b/Engine/Plugins/Animation/GameplayInsights/Source/RewindDebugger/Private/RewindDebuggerModule.cpp
--- a/Engine/Plugins/Animation/GameplayInsights/Source/RewindDebugger/Private/RewindDebuggerModule.cpp
+++ b/Engine/Plugins/Animation/GameplayInsights/Source/RewindDebugger/Private/RewindDebuggerModule.cpp
@@ -55,7 +55,6 @@ TSharedRef<SDockTab> FRewindDebuggerModule::SpawnRewindDebuggerTab(const FSpawnT
 		.OnTabClosed_Lambda([this](TSharedRef<SDockTab>)
 		{
 			// clear reference to widget so it will be destroyed
-			RewindDebuggerWidget->CloseAllTabs();
 			RewindDebuggerWidget = nullptr;
 		});
 
