fix non-unitiy build

diff --git a/Engine/Plugins/Animation/GameplayInsights/Source/GameplayInsights/Private/GameplayProvider.cpp b/Engine/Plugins/Animation/GameplayInsights/Source/GameplayInsights/Private/GameplayProvider.cpp
--- a/Engine/Plugins/Animation/GameplayInsights/Source/GameplayInsights/Private/GameplayProvider.cpp
+++ b/Engine/Plugins/Animation/GameplayInsights/Source/GameplayInsights/Private/GameplayProvider.cpp
@@ -10,6 +10,8 @@ FName FGameplayProvider::ProviderName("GameplayProvider");
 #define LOCTEXT_NAMESPACE "GameplayProvider"
 
 #if WITH_EDITOR
+#include "Engine/BlueprintCore.h"
+
 // This is copied from EditorUtilitySubsystem (where it is not public)
 // Should probably be somewhere shared
 static UClass* FindBlueprintClass(const FString& TargetNameRaw)
