more non-unity build fixes for Rewind Debugger

diff --git a/Engine/Plugins/Animation/GameplayInsights/Source/GameplayInsights/Private/GameplayProvider.cpp b/Engine/Plugins/Animation/GameplayInsights/Source/GameplayInsights/Private/GameplayProvider.cpp
--- a/Engine/Plugins/Animation/GameplayInsights/Source/GameplayInsights/Private/GameplayProvider.cpp
+++ b/Engine/Plugins/Animation/GameplayInsights/Source/GameplayInsights/Private/GameplayProvider.cpp
@@ -10,6 +10,7 @@ FName FGameplayProvider::ProviderName("GameplayProvider");
 #define LOCTEXT_NAMESPACE "GameplayProvider"
 
 #if WITH_EDITOR
+#include "Engine/Blueprint.h"
 #include "Engine/BlueprintCore.h"
 
 // This is copied from EditorUtilitySubsystem (where it is not public)
diff --git a/Engine/Plugins/Animation/GameplayInsights/Source/GameplayInsights/Private/SCurveTimelineView.h b/Engine/Plugins/Animation/GameplayInsights/Source/GameplayInsights/Private/SCurveTimelineView.h
--- a/Engine/Plugins/Animation/GameplayInsights/Source/GameplayInsights/Private/SCurveTimelineView.h
+++ b/Engine/Plugins/Animation/GameplayInsights/Source/GameplayInsights/Private/SCurveTimelineView.h
@@ -3,6 +3,7 @@
 #pragma once
 
 #include "CoreMinimal.h"
+#include "Widgets/SCompoundWidget.h"
 #include "Widgets/DeclarativeSyntaxSupport.h"
 
 class FPaintArgs;
diff --git a/Engine/Plugins/Animation/GameplayInsights/Source/RewindDebugger/Private/RewindDebugger.cpp b/Engine/Plugins/Animation/GameplayInsights/Source/RewindDebugger/Private/RewindDebugger.cpp
--- a/Engine/Plugins/Animation/GameplayInsights/Source/RewindDebugger/Private/RewindDebugger.cpp
+++ b/Engine/Plugins/Animation/GameplayInsights/Source/RewindDebugger/Private/RewindDebugger.cpp
@@ -11,11 +11,11 @@
 #include "ObjectTrace.h"
 #include "TraceServices/Model/Frames.h"
 #include "SLevelViewport.h"
+#include "Widgets/Docking/SDockTab.h"
+#include "Widgets/Layout/SSpacer.h"
 #include "IRewindDebuggerExtension.h"
 #include "IRewindDebuggerDoubleClickHandler.h"
 #include "RewindDebuggerObjectTrack.h"
-#include "RewindDebuggerViewCreators.h"
-#include "RewindDebuggerTrackCreators.h"
 #include "Animation/AnimBlueprint.h"
 #include "Animation/AnimBlueprintGeneratedClass.h"
 #include "ToolMenus.h"
diff --git a/Engine/Plugins/Animation/GameplayInsights/Source/RewindDebugger/Private/SObjectTimelineView.h b/Engine/Plugins/Animation/GameplayInsights/Source/RewindDebugger/Private/SObjectTimelineView.h
--- a/Engine/Plugins/Animation/GameplayInsights/Source/RewindDebugger/Private/SObjectTimelineView.h
+++ b/Engine/Plugins/Animation/GameplayInsights/Source/RewindDebugger/Private/SObjectTimelineView.h
@@ -4,6 +4,7 @@
 
 #include "CoreMinimal.h"
 #include "TraceServices/Model/AnalysisSession.h"
+#include "Widgets/SCompoundWidget.h"
 #include "Widgets/DeclarativeSyntaxSupport.h"
 
 class FPaintArgs;
