Rewind Debugger Timelines: Adding a null check caught by static analysis
#reb Thomas.Sarkanen

diff --git a/Engine/Plugins/Animation/GameplayInsights/Source/GameplayInsights/Private/AnimCurvesTrack.cpp b/Engine/Plugins/Animation/GameplayInsights/Source/GameplayInsights/Private/AnimCurvesTrack.cpp
--- a/Engine/Plugins/Animation/GameplayInsights/Source/GameplayInsights/Private/AnimCurvesTrack.cpp
+++ b/Engine/Plugins/Animation/GameplayInsights/Source/GameplayInsights/Private/AnimCurvesTrack.cpp
@@ -164,7 +164,7 @@ bool FAnimCurveTrack::UpdateInternal()
 		
 	bool bChanged = false;
 
-	if (CurveName.IsEmpty())
+	if (CurveName.IsEmpty() && AnimationProvider)
 	{
 		CurveName = FText::FromString(AnimationProvider->GetName(CurveId));
 		bChanged = true;
diff --git a/Engine/Plugins/Animation/GameplayInsights/Source/GameplayInsights/Private/BlendWeightsTrack.cpp b/Engine/Plugins/Animation/GameplayInsights/Source/GameplayInsights/Private/BlendWeightsTrack.cpp
--- a/Engine/Plugins/Animation/GameplayInsights/Source/GameplayInsights/Private/BlendWeightsTrack.cpp
+++ b/Engine/Plugins/Animation/GameplayInsights/Source/GameplayInsights/Private/BlendWeightsTrack.cpp
@@ -202,7 +202,7 @@ bool FBlendWeightTrack::UpdateInternal()
 	if (CurveType == ECurveType::BlendWeight)
 	{
 		// Blend Weight track gets name/icon of animation
-		if (CurveName.IsEmpty())
+		if (CurveName.IsEmpty() && GameplayProvider)
 		{
 			if (const FObjectInfo* ObjectInfo = GameplayProvider->FindObjectInfo(AssetId))
 			{
diff --git a/Engine/Plugins/Animation/GameplayInsights/Source/GameplayInsights/Private/MontagesTrack.cpp b/Engine/Plugins/Animation/GameplayInsights/Source/GameplayInsights/Private/MontagesTrack.cpp
--- a/Engine/Plugins/Animation/GameplayInsights/Source/GameplayInsights/Private/MontagesTrack.cpp
+++ b/Engine/Plugins/Animation/GameplayInsights/Source/GameplayInsights/Private/MontagesTrack.cpp
@@ -182,7 +182,7 @@ bool FMontageTrack::UpdateInternal()
 	if (CurveType == ECurveType::BlendWeight)
 	{
 		// Blend Weight track gets name/icon of animation
-		if (CurveName.IsEmpty())
+		if (CurveName.IsEmpty() && GameplayProvider)
 		{
 			if (const FObjectInfo* ObjectInfo = GameplayProvider->FindObjectInfo(AssetId))
 			{
