Optimize FindProfileTime in RewindDebugger to use a binary search

diff --git a/Engine/Plugins/Animation/GameplayInsights/Source/RewindDebugger/Private/RewindDebugger.cpp b/Engine/Plugins/Animation/GameplayInsights/Source/RewindDebugger/Private/RewindDebugger.cpp
--- a/Engine/Plugins/Animation/GameplayInsights/Source/RewindDebugger/Private/RewindDebugger.cpp
+++ b/Engine/Plugins/Animation/GameplayInsights/Source/RewindDebugger/Private/RewindDebugger.cpp
@@ -70,46 +70,34 @@ static double FindProfileTime(double InDebugTime, const TraceServices::IAnalysis
 				}
 				else
 				{
-					// find the two keys surrounding the CurrentScrubTime, and pick the nearest to update ProfileTime
-					if (Recording->GetEvent(ScrubFrameIndex).ElapsedTime > InDebugTime)
+					uint64 StartEventIndex = 0;
+					uint64 EndEventIndex = EventCount -1;
+
+					while (EndEventIndex - StartEventIndex > 1)
 					{
-						for (uint64 EventIndex = ScrubFrameIndex; EventIndex > 0; EventIndex--)
+						uint64 MiddleEventIndex = ((StartEventIndex + EndEventIndex) / 2);
+						const FRecordingInfoMessage& MiddleEvent = Recording->GetEvent(MiddleEventIndex);
+						if (InDebugTime < MiddleEvent.ElapsedTime)
 						{
-							const FRecordingInfoMessage& Event = Recording->GetEvent(EventIndex);
-							const FRecordingInfoMessage& NextEvent = Recording->GetEvent(EventIndex - 1);
-							if (Event.ElapsedTime >= InDebugTime && NextEvent.ElapsedTime <= InDebugTime)
-							{
-								if (Event.ElapsedTime - InDebugTime < InDebugTime - NextEvent.ElapsedTime)
-								{
-									ScrubFrameIndex = EventIndex;
-								}
-								else
-								{
-									ScrubFrameIndex = EventIndex - 1;
-								}
-								break;
-							}
+							EndEventIndex = MiddleEventIndex;
 						}
+						else
+						{
+							StartEventIndex = MiddleEventIndex;
+						}
+					}
+
+					check(EndEventIndex == StartEventIndex + 1)
+					const FRecordingInfoMessage& Event = Recording->GetEvent(StartEventIndex);
+					const FRecordingInfoMessage& NextEvent = Recording->GetEvent(EndEventIndex);
+					check (Event.ElapsedTime <= InDebugTime && NextEvent.ElapsedTime >= InDebugTime)
+					if (InDebugTime - Event.ElapsedTime < NextEvent.ElapsedTime - InDebugTime)
+					{
+						ScrubFrameIndex = StartEventIndex;
 					}
 					else
 					{
-						for (uint64 EventIndex = ScrubFrameIndex; EventIndex < EventCount - 1; EventIndex++)
-						{
-							const FRecordingInfoMessage& Event = Recording->GetEvent(EventIndex);
-							const FRecordingInfoMessage& NextEvent = Recording->GetEvent(EventIndex + 1);
-							if (Event.ElapsedTime <= InDebugTime && NextEvent.ElapsedTime >= InDebugTime)
-							{
-								if (InDebugTime - Event.ElapsedTime < NextEvent.ElapsedTime - InDebugTime)
-								{
-									ScrubFrameIndex = EventIndex;
-								}
-								else
-								{
-									ScrubFrameIndex = EventIndex + 1;
-								}
-								break;
-							}
-						}
+						ScrubFrameIndex = EndEventIndex;
 					}
 				}
 
