Fix Rewind Debugger Details view showing in Window menu
#prefilght 6283cdc8114e08fa2f2354f6

diff --git a/Engine/Plugins/Animation/GameplayInsights/Source/RewindDebugger/Private/RewindDebuggerModule.cpp b/Engine/Plugins/Animation/GameplayInsights/Source/RewindDebugger/Private/RewindDebuggerModule.cpp
--- a/Engine/Plugins/Animation/GameplayInsights/Source/RewindDebugger/Private/RewindDebuggerModule.cpp
+++ b/Engine/Plugins/Animation/GameplayInsights/Source/RewindDebugger/Private/RewindDebuggerModule.cpp
@@ -159,8 +159,9 @@ void FRewindDebuggerModule::StartupModule()
 	FGlobalTabmanager::Get()->RegisterNomadTabSpawner(
 		RewindDebuggerDetailsTabName,
     	FOnSpawnTab::CreateRaw(this, &FRewindDebuggerModule::SpawnRewindDebuggerDetailsTab))
+		.SetGroup(WorkspaceMenu::GetMenuStructure().GetDeveloperToolsDebugCategory())
 		.SetDisplayName(LOCTEXT("DetailsTabTitle", "Rewind Debugger Details"))
-		.SetIcon(FSlateIcon("RewindDebuggerStyle", "RewindDebugger.RewindIcon"))
+		.SetIcon(FSlateIcon("RewindDebuggerStyle", "RewindDebugger.RewindDetailsIcon"))
 		.SetTooltipText(LOCTEXT("TooltipText", "Opens Rewind Debugger Details Window."));
 
 	RewindDebuggerCameraExtension.Initialize();
diff --git a/Engine/Plugins/Animation/GameplayInsights/Source/RewindDebugger/Private/RewindDebuggerStyle.cpp b/Engine/Plugins/Animation/GameplayInsights/Source/RewindDebugger/Private/RewindDebuggerStyle.cpp
--- a/Engine/Plugins/Animation/GameplayInsights/Source/RewindDebugger/Private/RewindDebuggerStyle.cpp
+++ b/Engine/Plugins/Animation/GameplayInsights/Source/RewindDebugger/Private/RewindDebuggerStyle.cpp
@@ -36,6 +36,7 @@ FRewindDebuggerStyle::FRewindDebuggerStyle() :
 
 	// tab icon
 	Set("RewindDebugger.RewindIcon", new IMAGE_BRUSH("Rewind_24x", Icon16x16));
+	Set("RewindDebugger.RewindDetailsIcon", new IMAGE_BRUSH("RewindDetails_24x", Icon16x16));
 
 	// menu icon
 	Set("RewindDebugger.MenuIcon", new CORE_IMAGE_BRUSH_SVG("Slate/Starship/Common/menu", Icon16x16));
