Trace: Defer FTraceAuxiliary::Start to make way for default trace channels to be enabled.

Right now GameplayInsights and NetworkPredictionInsights start a Trace connection directly in the default StartupModule pathway. This shouldn't happen by default because if Start() is called UnrealTraceServer.exe will start logging to a trace file. This in turn will consume disk space on every run of Editor.

GameplayInsights defers Start() until AnimationInsights or RewindDebugger is shown on screen.

NetworkPredictionInsights requires a command line param, -NPTrace, be specified.

#rnd

diff --git a/Engine/Plugins/Animation/GameplayInsights/Source/GameplayInsights/Private/GameplayInsightsModule.cpp b/Engine/Plugins/Animation/GameplayInsights/Source/GameplayInsights/Private/GameplayInsightsModule.cpp
--- a/Engine/Plugins/Animation/GameplayInsights/Source/GameplayInsights/Private/GameplayInsightsModule.cpp
+++ b/Engine/Plugins/Animation/GameplayInsights/Source/GameplayInsights/Private/GameplayInsightsModule.cpp
@@ -61,6 +61,10 @@ void FGameplayInsightsModule::StartupModule()
 	IUnrealInsightsModule& UnrealInsightsModule = FModuleManager::LoadModuleChecked<IUnrealInsightsModule>("TraceInsights");
 	UnrealInsightsModule.OnMajorTabCreated().AddLambda([this](const FName& InMajorTabId, TSharedPtr<FTabManager> InTabManager)
 	{
+#if WITH_EDITOR
+		StartTrace();
+#endif
+
 		if (InMajorTabId == FInsightsManagerTabs::TimingProfilerTabId)
 		{
 			WeakTimingProfilerTabManager = InTabManager;
@@ -200,14 +204,9 @@ void FGameplayInsightsModule::StartupModule()
 #else
 				UE_LOG(LogCore, Display, TEXT("GameplayInsights module auto-connecting to local trace server..."));
 				UnrealInsightsModule.ConnectToStore(TEXT("127.0.0.1"));
-				const bool bConnected = FTraceAuxiliary::Start(
-					FTraceAuxiliary::EConnectionType::Network,
-					TEXT("127.0.0.1"),
-					nullptr);
 #endif // WITH_TRACE_STORE
 
 				UnrealInsightsModule.CreateSessionViewer(false);
-				UnrealInsightsModule.StartAnalysisForLastLiveSession();
 			}
 		});
 
@@ -419,6 +418,17 @@ bool FGameplayInsightsModule::IsObjectPropertyTraceEnabled(UObject* Object)
 #endif
 }
 
+void FGameplayInsightsModule::StartTrace()
+{
+	const bool bConnected = FTraceAuxiliary::Start(
+		FTraceAuxiliary::EConnectionType::Network,
+		TEXT("127.0.0.1"),
+		nullptr);
+	
+	IUnrealInsightsModule& UnrealInsightsModule = FModuleManager::LoadModuleChecked<IUnrealInsightsModule>("TraceInsights");
+	UnrealInsightsModule.StartAnalysisForLastLiveSession();
+}
+
 #endif
 
 IMPLEMENT_MODULE(FGameplayInsightsModule, GameplayInsights);
diff --git a/Engine/Plugins/Animation/GameplayInsights/Source/GameplayInsights/Private/GameplayInsightsModule.h b/Engine/Plugins/Animation/GameplayInsights/Source/GameplayInsights/Private/GameplayInsightsModule.h
--- a/Engine/Plugins/Animation/GameplayInsights/Source/GameplayInsights/Private/GameplayInsightsModule.h
+++ b/Engine/Plugins/Animation/GameplayInsights/Source/GameplayInsights/Private/GameplayInsightsModule.h
@@ -33,6 +33,7 @@ public:
 #if WITH_EDITOR
 	virtual void EnableObjectPropertyTrace(UObject* Object, bool bEnable = true) override;
 	virtual bool IsObjectPropertyTraceEnabled(UObject* Object) override;
+	virtual void StartTrace() override;
 #endif
 
 	// Spawn a document tab
diff --git a/Engine/Plugins/Animation/GameplayInsights/Source/GameplayInsights/Public/IGameplayInsightsModule.h b/Engine/Plugins/Animation/GameplayInsights/Source/GameplayInsights/Public/IGameplayInsightsModule.h
--- a/Engine/Plugins/Animation/GameplayInsights/Source/GameplayInsights/Public/IGameplayInsightsModule.h
+++ b/Engine/Plugins/Animation/GameplayInsights/Source/GameplayInsights/Public/IGameplayInsightsModule.h
@@ -24,5 +24,7 @@ public:
 	virtual void EnableObjectPropertyTrace(UObject* Object, bool bEnable = true) = 0;
 	// check if an Object has it's properties tracing to Insights
 	virtual bool IsObjectPropertyTraceEnabled(UObject* Object) = 0;
+	// start tracing data
+	virtual void StartTrace() = 0;
 #endif
 };
diff --git a/Engine/Plugins/Animation/GameplayInsights/Source/RewindDebugger/Private/RewindDebugger.cpp b/Engine/Plugins/Animation/GameplayInsights/Source/RewindDebugger/Private/RewindDebugger.cpp
--- a/Engine/Plugins/Animation/GameplayInsights/Source/RewindDebugger/Private/RewindDebugger.cpp
+++ b/Engine/Plugins/Animation/GameplayInsights/Source/RewindDebugger/Private/RewindDebugger.cpp
@@ -7,6 +7,7 @@
 #include "IAnimationProvider.h"
 #include "IGameplayProvider.h"
 #include "Insights/IUnrealInsightsModule.h"
+#include "IGameplayInsightsModule.h"
 #include "Modules/ModuleManager.h"
 #include "ObjectTrace.h"
 #include "TraceServices/Model/Frames.h"
@@ -149,6 +150,9 @@ FRewindDebugger::FRewindDebugger()  :
 
 		return true;
 	});
+
+	IGameplayInsightsModule* GameplayInsightsModule = &FModuleManager::LoadModuleChecked<IGameplayInsightsModule>("GameplayInsights");
+	GameplayInsightsModule->StartTrace();
 }
 
 FRewindDebugger::~FRewindDebugger() 
diff --git a/Engine/Plugins/Runtime/NetworkPredictionInsights/Source/NetworkPredictionInsights/Private/NetworkPredictionInsightsModule.cpp b/Engine/Plugins/Runtime/NetworkPredictionInsights/Source/NetworkPredictionInsights/Private/NetworkPredictionInsightsModule.cpp
--- a/Engine/Plugins/Runtime/NetworkPredictionInsights/Source/NetworkPredictionInsights/Private/NetworkPredictionInsightsModule.cpp
+++ b/Engine/Plugins/Runtime/NetworkPredictionInsights/Source/NetworkPredictionInsights/Private/NetworkPredictionInsightsModule.cpp
@@ -105,7 +105,8 @@ void FNetworkPredictionInsightsModule::StartupModule()
 		{
 			LLM_SCOPE_BYNAME(TEXT("Insights/NetworkPredictionInsights"));
 			IUnrealInsightsModule& UnrealInsightsModule = FModuleManager::LoadModuleChecked<IUnrealInsightsModule>("TraceInsights");
-			if (!UnrealInsightsModule.GetStoreClient())
+			bool bNetworkPredictionTrace = FParse::Param(FCommandLine::Get(), TEXT("NPTrace"));
+			if (!UnrealInsightsModule.GetStoreClient() && bNetworkPredictionTrace)
 			{
 #if WITH_TRACE_STORE
 				UE_LOG(LogCore, Display, TEXT("NetworkPredictionInsights module auto-connecting to internal trace server..."));
