Fix for Rewind Debugger frame step button resetting user back to 0.0 on the timeline.

Scrub frame index is now being updated when scrub time is set.

Changed FindProfileTime() to GetScrubTimeInformation(). Also, ensured it makes used of cached index for improved performance by trying a linear search if time difference is mall enough before binary search is performed.

Removed unnecessary search from recording index from FindProfileTime().

diff --git a/Engine/Plugins/Animation/GameplayInsights/Source/RewindDebugger/Private/RewindDebugger.cpp b/Engine/Plugins/Animation/GameplayInsights/Source/RewindDebugger/Private/RewindDebugger.cpp
--- a/Engine/Plugins/Animation/GameplayInsights/Source/RewindDebugger/Private/RewindDebugger.cpp
+++ b/Engine/Plugins/Animation/GameplayInsights/Source/RewindDebugger/Private/RewindDebugger.cpp
@@ -34,83 +34,6 @@ static void IterateExtensions(TFunction<void(IRewindDebuggerExtension* Extension
 	}
 }
 
-static double FindProfileTime(double InDebugTime, const TraceServices::IAnalysisSession* AnalysisSession)
-{
-	const IGameplayProvider* GameplayProvider = AnalysisSession->ReadProvider<IGameplayProvider>("GameplayProvider");
-	const IAnimationProvider* AnimationProvider = AnalysisSession->ReadProvider<IAnimationProvider>("AnimationProvider");
-
-	if (GameplayProvider && AnimationProvider)
-	{
-		TraceServices::FAnalysisSessionReadScope SessionReadScope(*AnalysisSession);
-
-		int ScrubFrameIndex = 0;
-		int RecordingIndex = 0;
-		while (GameplayProvider->GetRecordingInfo(RecordingIndex))
-		{
-			// find latest recording. (hack!)
-			RecordingIndex++;
-		}
-		RecordingIndex--;
-
-		if (const IGameplayProvider::RecordingInfoTimeline* Recording = GameplayProvider->GetRecordingInfo(RecordingIndex))
-		{
-			uint64 EventCount = Recording->GetEventCount();
-
-			if (EventCount > 0)
-			{
-				// check if we are outside of the recorded range, and apply the first or last frame
-				const FRecordingInfoMessage& FirstEvent = Recording->GetEvent(0);
-				const FRecordingInfoMessage& LastEvent = Recording->GetEvent(EventCount - 1);
-				if (InDebugTime <= FirstEvent.ElapsedTime)
-				{
-					ScrubFrameIndex = FMath::Min<uint64>(1, EventCount - 1);
-				}
-				else if (InDebugTime >= LastEvent.ElapsedTime)
-				{
-					ScrubFrameIndex = EventCount - 1;
-				}
-				else
-				{
-					uint64 StartEventIndex = 0;
-					uint64 EndEventIndex = EventCount -1;
-
-					while (EndEventIndex - StartEventIndex > 1)
-					{
-						uint64 MiddleEventIndex = ((StartEventIndex + EndEventIndex) / 2);
-						const FRecordingInfoMessage& MiddleEvent = Recording->GetEvent(MiddleEventIndex);
-						if (InDebugTime < MiddleEvent.ElapsedTime)
-						{
-							EndEventIndex = MiddleEventIndex;
-						}
-						else
-						{
-							StartEventIndex = MiddleEventIndex;
-						}
-					}
-
-					check(EndEventIndex == StartEventIndex + 1)
-					const FRecordingInfoMessage& Event = Recording->GetEvent(StartEventIndex);
-					const FRecordingInfoMessage& NextEvent = Recording->GetEvent(EndEventIndex);
-					check (Event.ElapsedTime <= InDebugTime && NextEvent.ElapsedTime >= InDebugTime)
-					if (InDebugTime - Event.ElapsedTime < NextEvent.ElapsedTime - InDebugTime)
-					{
-						ScrubFrameIndex = StartEventIndex;
-					}
-					else
-					{
-						ScrubFrameIndex = EndEventIndex;
-					}
-				}
-
-				const FRecordingInfoMessage& Event = Recording->GetEvent(ScrubFrameIndex);
-				return Event.ProfileTime;
-			}
-		}
-	}
-
-	return 0;
-}
-
 FRewindDebugger::FRewindDebugger()  :
 	ControlState(FRewindDebugger::EControlState::Pause),
 	bPIEStarted(false),
@@ -121,7 +44,6 @@ FRewindDebugger::FRewindDebugger()  :
 	CurrentScrubTime(0),
 	CurrentViewRange(0,0),
 	CurrentTraceRange(0,0),
-	ScrubFrameIndex(0),
 	RecordingIndex(0),
 	bTargetActorPositionValid(false)
 {
@@ -497,11 +419,11 @@ void FRewindDebugger::Step(int frames)
 
 					if (EventCount > 0)
 					{
-						ScrubFrameIndex = FMath::Clamp<int64>(ScrubFrameIndex + frames, 0, (int64)EventCount - 1);
+						ScrubTimeInformation.FrameIndex = FMath::Clamp<int64>(ScrubTimeInformation.FrameIndex + frames, 0, (int64)EventCount - 1);
+						const FRecordingInfoMessage& Event = Recording->GetEvent(ScrubTimeInformation.FrameIndex);
 
-						const FRecordingInfoMessage& Event = Recording->GetEvent(ScrubFrameIndex);
 						SetCurrentScrubTime(Event.ElapsedTime);
-						TraceTime.Set(Event.ProfileTime);
+						
 						TrackCursorDelegate.ExecuteIfBound(false);
 					}
 				}
@@ -558,23 +480,150 @@ void FRewindDebugger::SetCurrentViewRange(const TRange<double>& Range)
 	CurrentViewRange = Range;
 	if (const TraceServices::IAnalysisSession* Session = GetAnalysisSession())
 	{
-		CurrentTraceRange.SetLowerBoundValue(FindProfileTime(CurrentViewRange.GetLowerBoundValue(), Session));
-		CurrentTraceRange.SetUpperBoundValue(FindProfileTime(CurrentViewRange.GetUpperBoundValue(), Session));
+		GetScrubTimeInformation(CurrentViewRange.GetLowerBoundValue(), LowerBoundViewTimeInformation, RecordingIndex, Session);
+		GetScrubTimeInformation(CurrentViewRange.GetUpperBoundValue(), UpperBoundViewTimeInformation, RecordingIndex, Session);
+		
+		CurrentTraceRange.SetLowerBoundValue(LowerBoundViewTimeInformation.ProfileTime);
+		CurrentTraceRange.SetUpperBoundValue(UpperBoundViewTimeInformation.ProfileTime);
 	}
 }
 
 void FRewindDebugger::SetCurrentScrubTime(double Time)
 {
 	CurrentScrubTime = Time;
-	UpdateTraceTime();
-}
 
+	if (const TraceServices::IAnalysisSession* Session = GetAnalysisSession())
+	{
+		GetScrubTimeInformation(CurrentScrubTime, ScrubTimeInformation, RecordingIndex, Session);
+		
+		TraceTime.Set(ScrubTimeInformation.ProfileTime);
+	}
+}
 
-void FRewindDebugger::UpdateTraceTime()
+void FRewindDebugger::GetScrubTimeInformation(double InDebugTime, FScrubTimeInformation & InOutTimeInformation, uint16 InRecordingIndex, const TraceServices::IAnalysisSession* AnalysisSession)
 {
-	if (const TraceServices::IAnalysisSession* Session = GetAnalysisSession())
+	const IGameplayProvider* GameplayProvider = AnalysisSession->ReadProvider<IGameplayProvider>("GameplayProvider");
+	const IAnimationProvider* AnimationProvider = AnalysisSession->ReadProvider<IAnimationProvider>("AnimationProvider");
+	
+	if (GameplayProvider && AnimationProvider)
 	{
-		TraceTime.Set(FindProfileTime(CurrentScrubTime, Session));
+		TraceServices::FAnalysisSessionReadScope SessionReadScope(*AnalysisSession);
+		
+		if (const IGameplayProvider::RecordingInfoTimeline* Recording = GameplayProvider->GetRecordingInfo(InRecordingIndex))
+		{
+			const uint64 EventCount = Recording->GetEventCount();
+
+			if (EventCount > 0)
+			{
+				int ScrubFrameIndex = InOutTimeInformation.FrameIndex;
+				const FRecordingInfoMessage& FirstEvent = Recording->GetEvent(0);
+				const FRecordingInfoMessage& LastEvent = Recording->GetEvent(EventCount - 1);
+
+				// Check if we are outside of the recorded range, and apply the first or last frame
+				if (InDebugTime <= FirstEvent.ElapsedTime)
+				{
+					ScrubFrameIndex = FMath::Min<uint64>(1, EventCount - 1);
+				}
+				else if (InDebugTime >= LastEvent.ElapsedTime)
+				{
+					ScrubFrameIndex = EventCount - 1;
+				}
+				// Find the two keys surrounding the InDebugTime, and pick the nearest to update InOutTimeInformation
+				else
+				{
+					const FRecordingInfoMessage& ScrubEvent = Recording->GetEvent(ScrubFrameIndex);
+					constexpr float MaxTimeDifferenceInSeconds = 15.0f / 60.0f;
+					
+					// Use linear search on smaller time differences
+					if (FMath::Abs(InDebugTime - ScrubEvent.ElapsedTime) <= MaxTimeDifferenceInSeconds)
+					{
+						if (Recording->GetEvent(ScrubFrameIndex).ElapsedTime > InDebugTime)
+						{
+							for (uint64 EventIndex = ScrubFrameIndex; EventIndex > 0; EventIndex--)
+							{
+								const FRecordingInfoMessage& Event = Recording->GetEvent(EventIndex);
+								const FRecordingInfoMessage& NextEvent = Recording->GetEvent(EventIndex - 1);
+								if (Event.ElapsedTime >= InDebugTime && NextEvent.ElapsedTime <= InDebugTime)
+								{
+									if (Event.ElapsedTime - InDebugTime < InDebugTime - NextEvent.ElapsedTime)
+									{
+										ScrubFrameIndex = EventIndex;
+									}
+									else
+									{
+										ScrubFrameIndex = EventIndex - 1;
+									}
+									break;
+								}
+							}
+						}
+						else
+						{
+							for (uint64 EventIndex = ScrubFrameIndex; EventIndex < EventCount - 1; EventIndex++)
+							{
+								const FRecordingInfoMessage& Event = Recording->GetEvent(EventIndex);
+								const FRecordingInfoMessage& NextEvent = Recording->GetEvent(EventIndex + 1);
+								if (Event.ElapsedTime <= InDebugTime && NextEvent.ElapsedTime >= InDebugTime)
+								{
+									if (InDebugTime - Event.ElapsedTime < NextEvent.ElapsedTime - InDebugTime)
+									{
+										ScrubFrameIndex = EventIndex;
+									}
+									else
+									{
+										ScrubFrameIndex = EventIndex + 1;
+									}
+									break;
+								}
+							}
+						}
+					}
+					// Binary search for surrounding keys on big time differences
+					else
+					{
+						uint64 StartEventIndex = 0;
+						uint64 EndEventIndex = EventCount -1;
+						
+						while (EndEventIndex - StartEventIndex > 1)
+						{
+							const uint64 MiddleEventIndex = ((StartEventIndex + EndEventIndex) / 2);
+							const FRecordingInfoMessage& MiddleEvent = Recording->GetEvent(MiddleEventIndex);
+							if (InDebugTime < MiddleEvent.ElapsedTime)
+							{
+								EndEventIndex = MiddleEventIndex;
+							}
+							else
+							{
+								StartEventIndex = MiddleEventIndex;
+							}
+						}
+						
+						// Ensure there is not frames between start and end index
+						check(EndEventIndex == StartEventIndex + 1)
+				
+						const FRecordingInfoMessage& Event = Recording->GetEvent(StartEventIndex);
+						const FRecordingInfoMessage& NextEvent = Recording->GetEvent(EndEventIndex);
+
+						// Ensure debug time is between both frames time range
+						check (Event.ElapsedTime <= InDebugTime && NextEvent.ElapsedTime >= InDebugTime)
+
+						// Choose frame that is nearest to the debug time
+						if (InDebugTime - Event.ElapsedTime < NextEvent.ElapsedTime - InDebugTime)
+						{
+							ScrubFrameIndex = StartEventIndex;
+						}
+						else
+						{
+							ScrubFrameIndex = EndEventIndex;
+						}
+					}
+				}
+				
+				const FRecordingInfoMessage& Event = Recording->GetEvent(ScrubFrameIndex);
+				InOutTimeInformation.FrameIndex = ScrubFrameIndex;
+				InOutTimeInformation.ProfileTime = Event.ProfileTime;
+			}
+		}
 	}
 }
 
@@ -619,8 +668,6 @@ void FRewindDebugger::Tick(float DeltaTime)
 			{
 				if (RecordingDuration.Get() > 0)
 				{
-					UpdateTraceTime();
-
 					if (ControlState == EControlState::Play || ControlState == EControlState::PlayReverse)
 					{
 						float Rate = PlaybackRate * (ControlState == EControlState::Play ? 1 : -1);
diff --git a/Engine/Plugins/Animation/GameplayInsights/Source/RewindDebugger/Private/RewindDebugger.h b/Engine/Plugins/Animation/GameplayInsights/Source/RewindDebugger/Private/RewindDebugger.h
--- a/Engine/Plugins/Animation/GameplayInsights/Source/RewindDebugger/Private/RewindDebugger.h
+++ b/Engine/Plugins/Animation/GameplayInsights/Source/RewindDebugger/Private/RewindDebugger.h
@@ -120,9 +120,7 @@ private:
 	void OnPIEStopped(bool bSimulating);
 	void OnPIESingleStepped(bool bSimulating);
 
-
 	void SetCurrentScrubTime(double Time);
-	void UpdateTraceTime();
 
 	TBindableProperty<double> TraceTime;
 	TBindableProperty<float> RecordingDuration;
@@ -150,9 +148,20 @@ private:
 	double CurrentScrubTime;
 	TRange<double> CurrentViewRange;
 	TRange<double> CurrentTraceRange;
-	int64 ScrubFrameIndex;
 	uint16 RecordingIndex;
 
+	struct FScrubTimeInformation
+	{
+		double ProfileTime = 0;  // Profiling/Tracing time
+		int64 FrameIndex = 0;    // Scrub Frame Index
+	};
+	
+	FScrubTimeInformation ScrubTimeInformation;
+	FScrubTimeInformation LowerBoundViewTimeInformation;
+	FScrubTimeInformation UpperBoundViewTimeInformation;
+
+	static void GetScrubTimeInformation(double InDebugTime, FScrubTimeInformation & InOutTimeInformation, uint16 InRecordingIndex, const TraceServices::IAnalysisSession* AnalysisSession);
+	
 	TArray<TSharedPtr<FDebugObjectInfo>> DebugComponents;
 	mutable TSharedPtr<FDebugObjectInfo> SelectedComponent;
 	
