Rewind Debugger track filtering

diff --git a/Engine/Plugins/Animation/GameplayInsights/Source/RewindDebugger/Private/RewindDebuggerObjectTrack.cpp b/Engine/Plugins/Animation/GameplayInsights/Source/RewindDebugger/Private/RewindDebuggerObjectTrack.cpp
--- a/Engine/Plugins/Animation/GameplayInsights/Source/RewindDebugger/Private/RewindDebuggerObjectTrack.cpp
+++ b/Engine/Plugins/Animation/GameplayInsights/Source/RewindDebugger/Private/RewindDebuggerObjectTrack.cpp
@@ -66,7 +66,7 @@ bool FRewindDebuggerObjectTrack::UpdateInternal()
 
 	ExistenceRange = GameplayProvider->GetObjectRecordingLifetime(ObjectId);
 
-	if (Icon.IsSet())
+	if (!Icon.IsSet())
 	{
 		if (const FObjectInfo* ObjectInfo = GameplayProvider->FindObjectInfo(ObjectId))
 		{
diff --git a/Engine/Plugins/Animation/GameplayInsights/Source/RewindDebugger/Private/RewindDebuggerObjectTrack.h b/Engine/Plugins/Animation/GameplayInsights/Source/RewindDebugger/Private/RewindDebuggerObjectTrack.h
--- a/Engine/Plugins/Animation/GameplayInsights/Source/RewindDebugger/Private/RewindDebuggerObjectTrack.h
+++ b/Engine/Plugins/Animation/GameplayInsights/Source/RewindDebugger/Private/RewindDebuggerObjectTrack.h
@@ -28,6 +28,7 @@ private:
 	virtual FSlateIcon GetIconInternal() override { return Icon; }
 	virtual FText GetDisplayNameInternal() const override { return FText::FromString(ObjectName); }
 	virtual uint64 GetObjectIdInternal() const override { return ObjectId; }
+	virtual bool HasDebugDataInternal() const override { return false; }
 	
 	FString ObjectName;
 	FSlateIcon Icon;
diff --git a/Engine/Plugins/Animation/GameplayInsights/Source/RewindDebugger/Private/SRewindDebugger.cpp b/Engine/Plugins/Animation/GameplayInsights/Source/RewindDebugger/Private/SRewindDebugger.cpp
--- a/Engine/Plugins/Animation/GameplayInsights/Source/RewindDebugger/Private/SRewindDebugger.cpp
+++ b/Engine/Plugins/Animation/GameplayInsights/Source/RewindDebugger/Private/SRewindDebugger.cpp
@@ -1,6 +1,7 @@
 // Copyright Epic Games, Inc. All Rights Reserved.
 
 #include "SRewindDebugger.h"
+
 #include "ActorPickerMode.h"
 #include "Editor.h"
 #include "Editor/EditorEngine.h"
@@ -19,11 +20,11 @@
 #include "ToolMenu.h"
 #include "ToolMenus.h"
 #include "Widgets/Images/SImage.h"
-#include "Widgets/Input/SComboButton.h"
 #include "Widgets/Input/SButton.h"
+#include "Widgets/Input/SComboButton.h"
+#include "Widgets/Input/SSearchBox.h"
 #include "Widgets/Layout/SExpandableArea.h"
 #include "Widgets/Layout/SScrollBox.h"
-#include "IRewindDebuggerDoubleClickHandler.h"
 
 #define LOCTEXT_NAMESPACE "SRewindDebugger"
 
@@ -196,6 +197,12 @@ void SRewindDebugger::Construct(const FArguments& InArgs, TSharedRef<FUICommandL
 	TraceTime.Initialize(InArgs._TraceTime);
 	RecordingDuration.Initialize(InArgs._RecordingDuration);
 	DebugTargetActor.Initialize(InArgs._DebugTargetActor);
+
+	
+	TrackFilterBox = SNew(SSearchBox).HintText(LOCTEXT("Filter Tracks","Filter Tracks")).OnTextChanged_Lambda([this](const FText&)
+	{
+		RefreshDebugComponents();
+	});
 	
 	FSlimHorizontalToolBarBuilder ToolBarBuilder(CommandList, FMultiBoxCustomization::None, nullptr, true);
 
@@ -366,6 +373,17 @@ void SRewindDebugger::Construct(const FArguments& InArgs, TSharedRef<FUICommandL
 							.Image(FRewindDebuggerStyle::Get().GetBrush("RewindDebugger.SelectActor"))
 						]
 					]
+					+SHorizontalBox::Slot().AutoWidth()
+					[
+						SNew(SComboButton)
+							.ComboButtonStyle(FAppStyle::Get(), "SimpleComboButton")
+							.OnGetMenuContent(this, &SRewindDebugger::MakeFilterMenu)
+							.ButtonContent()
+						[
+							SNew(SImage)
+							.Image(FAppStyle::Get().GetBrush("Icons.Filter"))
+						]
+					]
 				]
 				+SVerticalBox::Slot().FillHeight(1.0f)
 				[
@@ -416,8 +434,33 @@ void SRewindDebugger::Construct(const FArguments& InArgs, TSharedRef<FUICommandL
 	];
 }
 
+bool FilterTrack(TSharedPtr<RewindDebugger::FRewindDebuggerTrack>& Track, const FString& FilterString, bool bRemoveNoData, bool bParentFilterPassed = false)
+{
+	const bool bStringFilterEmpty =  FilterString.IsEmpty();
+	const bool bStringFilterPassed = bParentFilterPassed || bStringFilterEmpty || Track->GetDisplayName().ToString().Contains(FilterString);
+
+	const bool bThisFilterPassed = (!bStringFilterEmpty && bStringFilterPassed);
+
+	bool bAnyChildVisible = false;
+	Track->IterateSubTracks([&bAnyChildVisible, bThisFilterPassed, FilterString, bRemoveNoData](TSharedPtr<RewindDebugger::FRewindDebuggerTrack> ChildTrack)
+	{
+		const bool bChildIsVisible = FilterTrack(ChildTrack, FilterString, bRemoveNoData, bThisFilterPassed);
+		bAnyChildVisible |= bChildIsVisible;
+	});
+
+	bool bVisible = bAnyChildVisible || ((!bRemoveNoData || Track->HasDebugData()) && bStringFilterPassed);
+
+	Track->SetIsVisible(bVisible);
+	return bVisible;
+}
+
 void SRewindDebugger::RefreshDebugComponents()
 {
+	if (DebugComponents && DebugComponents->Num() > 0)
+	{
+		FilterTrack((*DebugComponents)[0], TrackFilterBox->GetText().ToString(), !ShouldDisplayEmptyTracks());
+	}
+	
 	ComponentTreeView->Refresh();
 	TimelinesView->Refresh();
 }
@@ -427,6 +470,27 @@ TSharedRef<SWidget> SRewindDebugger::MakeMainMenu()
 	return UToolMenus::Get()->GenerateWidget("RewindDebugger.MainMenu", FToolMenuContext());
 }
 
+TSharedRef<SWidget> SRewindDebugger::MakeFilterMenu()
+{
+	FMenuBuilder Builder(true, nullptr);
+	Builder.AddWidget(TrackFilterBox.ToSharedRef(), FText(), true, false);
+	Builder.AddSeparator();
+	
+	Builder.AddMenuEntry(
+		LOCTEXT("DisplayEmptyTracks", "Show Empty Object Tracks"),
+		LOCTEXT("DisplayEmptyTracks", "Show Object/Component tracks which have no sub tracks with any debug data"),
+		FSlateIcon(),
+		FUIAction(
+			FExecuteAction::CreateRaw(this, &SRewindDebugger::ToggleDisplayEmptyTracks),
+			FCanExecuteAction(),
+			FIsActionChecked::CreateRaw(this, &SRewindDebugger::ShouldDisplayEmptyTracks)),
+		NAME_None,
+		EUserInterfaceActionType::ToggleButton
+	);
+	
+	return Builder.MakeWidget();
+}
+
 void SRewindDebugger::ComponentSelectionChanged(TSharedPtr<RewindDebugger::FRewindDebuggerTrack> SelectedItem, ESelectInfo::Type SelectInfo)
 {
 	SelectedComponent = SelectedItem;
diff --git a/Engine/Plugins/Animation/GameplayInsights/Source/RewindDebugger/Private/SRewindDebugger.h b/Engine/Plugins/Animation/GameplayInsights/Source/RewindDebugger/Private/SRewindDebugger.h
--- a/Engine/Plugins/Animation/GameplayInsights/Source/RewindDebugger/Private/SRewindDebugger.h
+++ b/Engine/Plugins/Animation/GameplayInsights/Source/RewindDebugger/Private/SRewindDebugger.h
@@ -13,6 +13,9 @@
 
 class SDockTab;
 
+class SSearchBox;
+class SCheckBox;
+
 class SRewindDebugger : public SCompoundWidget
 {
 	typedef TBindablePropertyInitializer<FString, BindingType_Out> DebugTargetInitializer;
@@ -61,8 +64,10 @@ public:
 private:
 	void SetViewRange(TRange<double> NewRange);
 
-	
 	void MainFrameCreationFinished(TSharedPtr<SWindow> InRootWindow, bool bIsNewProjectWindow);
+
+	void ToggleDisplayEmptyTracks() { bDisplayEmptyTracks = !bDisplayEmptyTracks; RefreshDebugComponents(); }
+	bool ShouldDisplayEmptyTracks() { return bDisplayEmptyTracks; }
 	
 	// Time Slider
 	TAttribute<double> ScrubTimeAttribute;
@@ -83,6 +88,7 @@ private:
 	void TraceTimeChanged(double Time);
 
 	TSharedRef<SWidget> MakeMainMenu();
+	TSharedRef<SWidget> MakeFilterMenu();
 	void MakeViewsMenu(FMenuBuilder& MenuBuilder);
 	
 	// component tree view
@@ -95,6 +101,9 @@ private:
 	TSharedPtr<SRewindDebuggerComponentTree> ComponentTreeView;
 	TSharedPtr<SRewindDebuggerTimelines> TimelinesView;
 	TSharedPtr<SWidget> OnContextMenuOpening();
-
+	
+	TSharedPtr<SSearchBox> TrackFilterBox;
+	
 	bool bInSelectionChanged = false;
+	bool bDisplayEmptyTracks = false;
 };
diff --git a/Engine/Plugins/Animation/GameplayInsights/Source/RewindDebugger/Private/SRewindDebuggerComponentTree.cpp b/Engine/Plugins/Animation/GameplayInsights/Source/RewindDebugger/Private/SRewindDebuggerComponentTree.cpp
--- a/Engine/Plugins/Animation/GameplayInsights/Source/RewindDebugger/Private/SRewindDebuggerComponentTree.cpp
+++ b/Engine/Plugins/Animation/GameplayInsights/Source/RewindDebugger/Private/SRewindDebuggerComponentTree.cpp
@@ -47,7 +47,10 @@ void ComponentTreeViewGetChildren(TSharedPtr<RewindDebugger::FRewindDebuggerTrac
 {
 	InItem->IterateSubTracks([&OutChildren](TSharedPtr<RewindDebugger::FRewindDebuggerTrack> Track)
 	{
-		OutChildren.Add(Track);
+		if (Track->IsVisible())
+		{
+			OutChildren.Add(Track);
+		}
 	});
 }
 
diff --git a/Engine/Plugins/Animation/GameplayInsights/Source/RewindDebugger/Private/SRewindDebuggerTimelines.cpp b/Engine/Plugins/Animation/GameplayInsights/Source/RewindDebugger/Private/SRewindDebuggerTimelines.cpp
--- a/Engine/Plugins/Animation/GameplayInsights/Source/RewindDebugger/Private/SRewindDebuggerTimelines.cpp
+++ b/Engine/Plugins/Animation/GameplayInsights/Source/RewindDebugger/Private/SRewindDebuggerTimelines.cpp
@@ -75,7 +75,10 @@ void TimelineViewGetChildren(TSharedPtr<RewindDebugger::FRewindDebuggerTrack> In
 {
 	InItem->IterateSubTracks([&OutChildren](TSharedPtr<RewindDebugger::FRewindDebuggerTrack> Track)
 	{
-		OutChildren.Add(Track);
+		if (Track->IsVisible())
+		{
+			OutChildren.Add(Track);
+		}
 	});
 }
 
diff --git a/Engine/Source/Editor/RewindDebuggerInterface/Public/IRewindDebuggerTrackCreator.h b/Engine/Source/Editor/RewindDebuggerInterface/Public/IRewindDebuggerTrackCreator.h
--- a/Engine/Source/Editor/RewindDebuggerInterface/Public/IRewindDebuggerTrackCreator.h
+++ b/Engine/Source/Editor/RewindDebuggerInterface/Public/IRewindDebuggerTrackCreator.h
@@ -21,7 +21,9 @@ class FRewindDebuggerTrack;
 // Interface class which creates tracks
 class REWINDDEBUGGERINTERFACE_API IRewindDebuggerTrackCreator : public IModularFeature
 {
-	public: 
+	public:
+	virtual ~IRewindDebuggerTrackCreator() {};
+	
 	static const FName ModularFeatureName;
 
 	// returns the type of UObject this track creator can create tracks for
diff --git a/Engine/Source/Editor/RewindDebuggerInterface/Public/RewindDebuggerTrack.h b/Engine/Source/Editor/RewindDebuggerInterface/Public/RewindDebuggerTrack.h
--- a/Engine/Source/Editor/RewindDebuggerInterface/Public/RewindDebuggerTrack.h
+++ b/Engine/Source/Editor/RewindDebuggerInterface/Public/RewindDebuggerTrack.h
@@ -15,7 +15,7 @@ namespace RewindDebugger
 class FRewindDebuggerTrack
 {
 public:
-	FRewindDebuggerTrack(): bExpanded(true)
+	FRewindDebuggerTrack(): bExpanded(true), bVisible(true)
 	{
 	}
 
@@ -81,6 +81,15 @@ public:
 		IterateSubTracksInternal(IteratorFunction);
 	}
 
+	// returns true for tracks that contain debug data (used for filtering out parts of the hierarchy with no useful information in them)
+	bool HasDebugData() const
+	{
+		return HasDebugDataInternal();
+	}
+
+	bool IsVisible() const { return bVisible; }
+	void SetIsVisible(bool bInIsVisible) { bVisible = bInIsVisible; }
+
 private:
 
 	virtual bool UpdateInternal() { return false; }
@@ -90,9 +99,11 @@ private:
 	virtual FName GetNameInternal() const { return ""; }
 	virtual FText GetDisplayNameInternal() const { return FText(); }
 	virtual uint64 GetObjectIdInternal() const { return 0; }
-	virtual void IterateSubTracksInternal(TFunction<void(TSharedPtr<FRewindDebuggerTrack> SubTrack)> IteratorFunction) { } 
+	virtual bool HasDebugDataInternal() const { return true; }
+	virtual void IterateSubTracksInternal(TFunction<void(TSharedPtr<FRewindDebuggerTrack> SubTrack)> IteratorFunction) { }
 
 	bool bExpanded;
+	bool bVisible;
 };
 	
 }
\ No newline at end of file
