fix duplicate LOCTEXT key in SRewindDebugger

diff --git a/Engine/Plugins/Animation/GameplayInsights/Source/RewindDebugger/Private/SRewindDebugger.cpp b/Engine/Plugins/Animation/GameplayInsights/Source/RewindDebugger/Private/SRewindDebugger.cpp
--- a/Engine/Plugins/Animation/GameplayInsights/Source/RewindDebugger/Private/SRewindDebugger.cpp
+++ b/Engine/Plugins/Animation/GameplayInsights/Source/RewindDebugger/Private/SRewindDebugger.cpp
@@ -478,7 +478,7 @@ TSharedRef<SWidget> SRewindDebugger::MakeFilterMenu()
 	
 	Builder.AddMenuEntry(
 		LOCTEXT("DisplayEmptyTracks", "Show Empty Object Tracks"),
-		LOCTEXT("DisplayEmptyTracks", "Show Object/Component tracks which have no sub tracks with any debug data"),
+		LOCTEXT("DisplayEmptyTracksToolTip", "Show Object/Component tracks which have no sub tracks with any debug data"),
 		FSlateIcon(),
 		FUIAction(
 			FExecuteAction::CreateRaw(this, &SRewindDebugger::ToggleDisplayEmptyTracks),
