Rewind Debugger Notifies Track

diff --git a/Engine/Plugins/Animation/GameplayInsights/Source/GameplayInsights/Private/AnimCurvesTrack.cpp b/Engine/Plugins/Animation/GameplayInsights/Source/GameplayInsights/Private/AnimCurvesTrack.cpp
--- a/Engine/Plugins/Animation/GameplayInsights/Source/GameplayInsights/Private/AnimCurvesTrack.cpp
+++ b/Engine/Plugins/Animation/GameplayInsights/Source/GameplayInsights/Private/AnimCurvesTrack.cpp
@@ -51,8 +51,6 @@ bool FAnimCurvesTrack::UpdateInternal()
 	const FGameplayProvider* GameplayProvider = AnalysisSession->ReadProvider<FGameplayProvider>(FGameplayProvider::ProviderName);
 	const FAnimationProvider* AnimationProvider = AnalysisSession->ReadProvider<FAnimationProvider>(FAnimationProvider::ProviderName);
 
-	int CurveCount = 0;
-
 	if(GameplayProvider && AnimationProvider)
 	{
 		TraceServices::FAnalysisSessionReadScope SessionReadScope(*AnalysisSession);
@@ -76,15 +74,15 @@ bool FAnimCurvesTrack::UpdateInternal()
 
 		UniqueTrackIds.StableSort();
 		
-		CurveCount = UniqueTrackIds.Num();
+		const int32 CurveCount = UniqueTrackIds.Num();
 
-		if (Children.Num()!=CurveCount)
+		if (Children.Num() != UniqueTrackIds.Num())
 			bChanged = true;
 		
 		Children.SetNum(CurveCount);
 		for(int i = 0; i < CurveCount; i++)
 		{
-			if (!Children[i].IsValid() || static_cast<FAnimCurveTrack*>(Children[i].Get())->GetCurveId() != UniqueTrackIds[i])
+			if (!Children[i].IsValid() || Children[i].Get()->GetCurveId() != UniqueTrackIds[i])
 			{
 				Children[i] = MakeShared<FAnimCurveTrack>(ObjectId, UniqueTrackIds[i]);
 				bChanged = true;
diff --git a/Engine/Plugins/Animation/GameplayInsights/Source/GameplayInsights/Private/AnimCurvesTrack.h b/Engine/Plugins/Animation/GameplayInsights/Source/GameplayInsights/Private/AnimCurvesTrack.h
--- a/Engine/Plugins/Animation/GameplayInsights/Source/GameplayInsights/Private/AnimCurvesTrack.h
+++ b/Engine/Plugins/Animation/GameplayInsights/Source/GameplayInsights/Private/AnimCurvesTrack.h
@@ -49,7 +49,7 @@ public:
 
 private:
 	virtual bool UpdateInternal() override;
-	TSharedPtr<SWidget> GetDetailsViewInternal() override;
+	virtual TSharedPtr<SWidget> GetDetailsViewInternal() override;
 	
 	virtual FSlateIcon GetIconInternal() override { return Icon; }
 	virtual FName GetNameInternal() const override { return "AnimationCurves"; }
@@ -67,7 +67,7 @@ private:
 class FAnimationCurvesTrackCreator : public IRewindDebuggerTrackCreator
 {
 private:
-	virtual FName GetTargetTypeNameInternal() const;
+	virtual FName GetTargetTypeNameInternal() const override;
 	virtual FName GetNameInternal() const override;
 	virtual TSharedPtr<FRewindDebuggerTrack> CreateTrackInternal(uint64 ObjectId) const override;
 	virtual bool HasDebugInfoInternal(uint64 ObjectId) const override;
diff --git a/Engine/Plugins/Animation/GameplayInsights/Source/GameplayInsights/Private/AnimationAnalyzer.cpp b/Engine/Plugins/Animation/GameplayInsights/Source/GameplayInsights/Private/AnimationAnalyzer.cpp
--- a/Engine/Plugins/Animation/GameplayInsights/Source/GameplayInsights/Private/AnimationAnalyzer.cpp
+++ b/Engine/Plugins/Animation/GameplayInsights/Source/GameplayInsights/Private/AnimationAnalyzer.cpp
@@ -42,7 +42,9 @@ void FAnimationAnalyzer::OnAnalysisBegin(const FOnAnalysisContext& Context)
 	Builder.RouteEvent(RouteId_StateMachineState, "Animation", "StateMachineState");
 	Builder.RouteEvent(RouteId_Name, "Animation", "Name");
 	Builder.RouteEvent(RouteId_Notify, "Animation", "Notify");
+	Builder.RouteEvent(RouteId_Notify2, "Animation", "Notify2");
 	Builder.RouteEvent(RouteId_SyncMarker, "Animation", "SyncMarker");
+	Builder.RouteEvent(RouteId_SyncMarker2, "Animation", "SyncMarker2");
 	Builder.RouteEvent(RouteId_Montage, "Animation", "Montage");
 	Builder.RouteEvent(RouteId_Montage2, "Animation", "Montage2");
 	Builder.RouteEvent(RouteId_Sync, "Animation", "Sync");
@@ -330,8 +332,22 @@ bool FAnimationAnalyzer::OnEvent(uint16 RouteId, EStyle Style, const FOnEventCon
 		break;
 	}
 	case RouteId_Notify:
+		{
+			uint64 Cycle = EventData.GetValue<uint64>("Cycle");
+			uint64 AnimInstanceId = EventData.GetValue<uint64>("AnimInstanceId");
+			uint64 AssetId = EventData.GetValue<uint64>("AssetId");
+			uint64 NotifyId = EventData.GetValue<uint64>("NotifyId");
+			uint32 NameId = EventData.GetValue<uint32>("NameId");
+			float Time = EventData.GetValue<float>("Time");
+			float Duration = EventData.GetValue<float>("Duration");
+			uint8 NotifyEventType = EventData.GetValue<uint8>("NotifyEventType");
+			AnimationProvider.AppendNotify(AnimInstanceId, Context.EventTime.AsSeconds(Cycle), 0.0, AssetId, NotifyId, NameId, Time, Duration, (EAnimNotifyMessageType)NotifyEventType);
+			break;
+		}
+	case RouteId_Notify2:
 	{
 		uint64 Cycle = EventData.GetValue<uint64>("Cycle");
+		double RecordingTime = EventData.GetValue<double>("RecordingTime");
 		uint64 AnimInstanceId = EventData.GetValue<uint64>("AnimInstanceId");
 		uint64 AssetId = EventData.GetValue<uint64>("AssetId");
 		uint64 NotifyId = EventData.GetValue<uint64>("NotifyId");
@@ -339,7 +355,7 @@ bool FAnimationAnalyzer::OnEvent(uint16 RouteId, EStyle Style, const FOnEventCon
 		float Time = EventData.GetValue<float>("Time");
 		float Duration = EventData.GetValue<float>("Duration");
 		uint8 NotifyEventType = EventData.GetValue<uint8>("NotifyEventType");
-		AnimationProvider.AppendNotify(AnimInstanceId, Context.EventTime.AsSeconds(Cycle), AssetId, NotifyId, NameId, Time, Duration, (EAnimNotifyMessageType)NotifyEventType);
+		AnimationProvider.AppendNotify(AnimInstanceId, Context.EventTime.AsSeconds(Cycle), RecordingTime, AssetId, NotifyId, NameId, Time, Duration, (EAnimNotifyMessageType)NotifyEventType);
 		break;
 	}
 	case RouteId_SyncMarker:
@@ -347,7 +363,16 @@ bool FAnimationAnalyzer::OnEvent(uint16 RouteId, EStyle Style, const FOnEventCon
 		uint64 Cycle = EventData.GetValue<uint64>("Cycle");
 		uint64 AnimInstanceId = EventData.GetValue<uint64>("AnimInstanceId");
 		uint32 NameId = EventData.GetValue<uint32>("NameId");
-		AnimationProvider.AppendNotify(AnimInstanceId, Context.EventTime.AsSeconds(Cycle), 0, 0, NameId, 0.0f, 0.0f, EAnimNotifyMessageType::SyncMarker);
+		AnimationProvider.AppendNotify(AnimInstanceId, Context.EventTime.AsSeconds(Cycle), 0.0, 0, 0, NameId, 0.0f, 0.0f, EAnimNotifyMessageType::SyncMarker);
+		break;
+	}
+	case RouteId_SyncMarker2:
+	{
+		uint64 Cycle = EventData.GetValue<uint64>("Cycle");
+		double RecordingTime = EventData.GetValue<double>("RecordingTime");
+		uint64 AnimInstanceId = EventData.GetValue<uint64>("AnimInstanceId");
+		uint32 NameId = EventData.GetValue<uint32>("NameId");
+		AnimationProvider.AppendNotify(AnimInstanceId, Context.EventTime.AsSeconds(Cycle), RecordingTime, 0, 0, NameId, 0.0f, 0.0f, EAnimNotifyMessageType::SyncMarker);
 		break;
 	}
 	case RouteId_Montage:
diff --git a/Engine/Plugins/Animation/GameplayInsights/Source/GameplayInsights/Private/AnimationAnalyzer.h b/Engine/Plugins/Animation/GameplayInsights/Source/GameplayInsights/Private/AnimationAnalyzer.h
--- a/Engine/Plugins/Animation/GameplayInsights/Source/GameplayInsights/Private/AnimationAnalyzer.h
+++ b/Engine/Plugins/Animation/GameplayInsights/Source/GameplayInsights/Private/AnimationAnalyzer.h
@@ -43,7 +43,9 @@ private:
 		RouteId_StateMachineState,
 		RouteId_Name,
 		RouteId_Notify,
+		RouteId_Notify2,
 		RouteId_SyncMarker,
+		RouteId_SyncMarker2,
 		RouteId_Montage,
 		RouteId_Montage2,
 		RouteId_Sync,
diff --git a/Engine/Plugins/Animation/GameplayInsights/Source/GameplayInsights/Private/AnimationProvider.cpp b/Engine/Plugins/Animation/GameplayInsights/Source/GameplayInsights/Private/AnimationProvider.cpp
--- a/Engine/Plugins/Animation/GameplayInsights/Source/GameplayInsights/Private/AnimationProvider.cpp
+++ b/Engine/Plugins/Animation/GameplayInsights/Source/GameplayInsights/Private/AnimationProvider.cpp
@@ -1100,7 +1100,7 @@ void FAnimationProvider::AppendStateMachineState(uint64 InAnimInstanceId, double
 	Session.UpdateDurationSeconds(InTime);
 }
 
-void FAnimationProvider::AppendNotify(uint64 InAnimInstanceId, double InTime, uint64 InAssetId, uint64 InNotifyId, uint32 InNameId, float InNotifyTime, float InNotifyDuration, EAnimNotifyMessageType InNotifyEventType)
+void FAnimationProvider::AppendNotify(uint64 InAnimInstanceId, double InTime, double InRecordingTime, uint64 InAssetId, uint64 InNotifyId, uint32 InNameId, float InNotifyTime, float InNotifyDuration, EAnimNotifyMessageType InNotifyEventType)
 {
 	Session.WriteAccessCheck();
 
@@ -1139,6 +1139,7 @@ void FAnimationProvider::AppendNotify(uint64 InAnimInstanceId, double InTime, ui
 		if(InNotifyEventType == EAnimNotifyMessageType::Begin)
 		{
 			FAnimNotifyMessage Message;
+			Message.RecordingTime = InRecordingTime;
 			Message.AnimInstanceId = InAnimInstanceId;
 			Message.AssetId = InAssetId;
 			Message.NotifyId = InNotifyId;
@@ -1172,6 +1173,7 @@ void FAnimationProvider::AppendNotify(uint64 InAnimInstanceId, double InTime, ui
 
 		FAnimNotifyMessage Message;
 		Message.AnimInstanceId = InAnimInstanceId;
+		Message.RecordingTime = InRecordingTime;
 		Message.AssetId = InAssetId;
 		Message.NotifyId = InNotifyId;
 		Message.NameId = InNameId;
diff --git a/Engine/Plugins/Animation/GameplayInsights/Source/GameplayInsights/Private/AnimationProvider.h b/Engine/Plugins/Animation/GameplayInsights/Source/GameplayInsights/Private/AnimationProvider.h
--- a/Engine/Plugins/Animation/GameplayInsights/Source/GameplayInsights/Private/AnimationProvider.h
+++ b/Engine/Plugins/Animation/GameplayInsights/Source/GameplayInsights/Private/AnimationProvider.h
@@ -99,7 +99,7 @@ public:
 	void AppendStateMachineState(uint64 InAnimInstanceId, double InTime, int32 InNodeId, int32 InStateMachineIndex, int32 InStateIndex, float InStateWeight, float InElapsedTime);
 
 	/** Add an anim notify */
-	void AppendNotify(uint64 InAnimInstanceId, double InTime, uint64 InAssetId, uint64 InNotifyId, uint32 InNameId, float InNotifyTime, float InNotifyDuration, EAnimNotifyMessageType InNotifyEventType);
+	void AppendNotify(uint64 InAnimInstanceId, double InTime, double InRecordingTime, uint64 InAssetId, uint64 InNotifyId, uint32 InNameId, float InNotifyTime, float InNotifyDuration, EAnimNotifyMessageType InNotifyEventType);
 
 	/** Append montage data */
 	void AppendMontage(uint64 InAnimInstanceId, double InProfileTime, double InRecordingTime, uint64 InMontageId, uint32 InCurrentSectionNameId, uint32 InNextSectionNameId, float InWeight, float InDesiredWeight, float InPosition, uint16 InFrameCounter);
diff --git a/Engine/Plugins/Animation/GameplayInsights/Source/GameplayInsights/Private/GameplayInsightsModule.cpp b/Engine/Plugins/Animation/GameplayInsights/Source/GameplayInsights/Private/GameplayInsightsModule.cpp
--- a/Engine/Plugins/Animation/GameplayInsights/Source/GameplayInsights/Private/GameplayInsightsModule.cpp
+++ b/Engine/Plugins/Animation/GameplayInsights/Source/GameplayInsights/Private/GameplayInsightsModule.cpp
@@ -16,13 +16,12 @@
 #include "Trace/StoreClient.h"
 #include "Stats/Stats.h"
 #include "ObjectPropertyTrace.h"
-#include "SBlendWeightsView.h"
 #include "SMontageView.h"
 #include "SObjectPropertiesView.h"
-#include "SNotifiesView.h"
 #include "AnimCurvesTrack.h"
 #include "BlendWeightsTrack.h"
 #include "MontagesTrack.h"
+#include "NotifiesTrack.h"
 
 #if WITH_EDITOR
 #include "IAnimationBlueprintEditorModule.h"
@@ -77,8 +76,6 @@ void FGameplayInsightsModule::StartupModule()
 	IModularFeatures::Get().RegisterModularFeature(IRewindDebuggerViewCreator::ModularFeatureName, &ObjectPropertiesViewCreator);
 	static FAnimGraphSchematicViewCreator AnimGraphSchematicViewCreator;
 	IModularFeatures::Get().RegisterModularFeature(IRewindDebuggerViewCreator::ModularFeatureName, &AnimGraphSchematicViewCreator);
-	static FNotifiesViewCreator NotifiesViewCreator;
-	IModularFeatures::Get().RegisterModularFeature(IRewindDebuggerViewCreator::ModularFeatureName, &NotifiesViewCreator);
 	
 	static RewindDebugger::FAnimationCurvesTrackCreator AnimationCurvesTrackCreator;
 	IModularFeatures::Get().RegisterModularFeature(RewindDebugger::IRewindDebuggerTrackCreator::ModularFeatureName, &AnimationCurvesTrackCreator);
@@ -86,6 +83,8 @@ void FGameplayInsightsModule::StartupModule()
 	IModularFeatures::Get().RegisterModularFeature(RewindDebugger::IRewindDebuggerTrackCreator::ModularFeatureName, &BlendWeightsTrackCreator);
 	static RewindDebugger::FMontagesTrackCreator MontagesTrackCreator;
 	IModularFeatures::Get().RegisterModularFeature(RewindDebugger::IRewindDebuggerTrackCreator::ModularFeatureName, &MontagesTrackCreator);
+	static RewindDebugger::FNotifiesTrackCreator NotifiesTrackCreator;
+	IModularFeatures::Get().RegisterModularFeature(RewindDebugger::IRewindDebuggerTrackCreator::ModularFeatureName, &NotifiesTrackCreator);
 
 
 	if (!IsRunningCommandlet())
diff --git a/Engine/Plugins/Animation/GameplayInsights/Source/GameplayInsights/Private/NotifiesTrack.cpp b/Engine/Plugins/Animation/GameplayInsights/Source/GameplayInsights/Private/NotifiesTrack.cpp
new file mode 100644
--- /dev/null
+++ b/Engine/Plugins/Animation/GameplayInsights/Source/GameplayInsights/Private/NotifiesTrack.cpp
@@ -0,0 +1,275 @@
+// Copyright Epic Games, Inc. All Rights Reserved.
+
+#include "NotifiesTrack.h"
+#include "SCurveTimelineView.h"
+#include "IRewindDebugger.h"
+#include "GameplayProvider.h"
+#include "AnimationProvider.h"
+#include "SNotifiesView.h"
+
+#define LOCTEXT_NAMESPACE "NotifiesTrack"
+
+namespace RewindDebugger
+{
+
+FNotifyTrack::FNotifyTrack(uint64 InObjectId, const FNotifyTrack::FNotifyTrackId& InNotifyTrackId) :
+	ObjectId(InObjectId),
+	NotifyTrackId(InNotifyTrackId)
+{
+	EventData = MakeShared<SEventTimelineView::FTimelineEventData>();
+	Icon = FSlateIcon("EditorStyle", "Sequencer.Tracks.Event", "Sequencer.Tracks.Event");
+}
+
+TSharedPtr<SEventTimelineView::FTimelineEventData> FNotifyTrack::GetEventData() const
+{
+	if (!EventData.IsValid())
+	{
+		EventData = MakeShared<SEventTimelineView::FTimelineEventData>();
+	}
+	
+	EventUpdateRequested++;
+	
+	return EventData;
+}
+	
+static FLinearColor MakeNotifyColor(uint32 InSeed, bool bInLine = false)
+{
+	FRandomStream Stream(InSeed);
+	const uint8 Hue = (uint8)(Stream.FRand() * 255.0f);
+	const uint8 SatVal = bInLine ? 196 : 128;
+	return FLinearColor::MakeFromHSV8(Hue, SatVal, SatVal);
+}
+
+static FNotifyTrack::ENotifyTrackType MakeTrackType(EAnimNotifyMessageType MessageType)
+{
+	switch (MessageType)
+	{
+		case EAnimNotifyMessageType::Begin: return FNotifyTrack::ENotifyTrackType::NotifyState;
+		case EAnimNotifyMessageType::SyncMarker: return FNotifyTrack::ENotifyTrackType::SyncMarker;
+		default: return FNotifyTrack::ENotifyTrackType::Notify;
+	}
+}
+	
+TSharedPtr<SWidget> FNotifyTrack::GetDetailsViewInternal() 
+{
+	IRewindDebugger* RewindDebugger = IRewindDebugger::Instance();
+	
+	TSharedPtr<SNotifiesView> NotifiesView = SNew(SNotifiesView, ObjectId, RewindDebugger->CurrentTraceTime(), *RewindDebugger->GetAnalysisSession())
+				.CurrentTime_Lambda([RewindDebugger]{ return RewindDebugger->CurrentTraceTime(); });
+	NotifiesView->SetNotifyFilter(NotifyTrackId.NameId);
+	return NotifiesView;
+}
+
+bool FNotifyTrack::UpdateInternal()
+{
+	IRewindDebugger* RewindDebugger = IRewindDebugger::Instance();
+	
+	TRange<double> TraceTimeRange = RewindDebugger->GetCurrentTraceRange();
+	double StartTime = TraceTimeRange.GetLowerBoundValue();
+	double EndTime = TraceTimeRange.GetUpperBoundValue();
+	
+	const TraceServices::IAnalysisSession* AnalysisSession = RewindDebugger->GetAnalysisSession();
+	const FGameplayProvider* GameplayProvider = AnalysisSession->ReadProvider<FGameplayProvider>(FGameplayProvider::ProviderName);
+	const FAnimationProvider* AnimationProvider = AnalysisSession->ReadProvider<FAnimationProvider>(FAnimationProvider::ProviderName);
+	
+	if(EventUpdateRequested > 10 && GameplayProvider && AnimationProvider)
+	{
+		EventUpdateRequested = 0;
+		
+		auto& EventPoints = EventData->Points;
+		EventPoints.SetNum(0,false);
+		EventData->Windows.SetNum(0);
+
+		TraceServices::FAnalysisSessionReadScope SessionReadScope(*AnalysisSession);
+
+		AnimationProvider->ReadNotifyTimeline(ObjectId, [this, StartTime, EndTime, AnimationProvider, AnalysisSession, &EventPoints](const FAnimationProvider::AnimNotifyTimeline& InTimeline)
+		{
+			InTimeline.EnumerateEvents(StartTime, EndTime, [this, StartTime, EndTime, AnimationProvider, AnalysisSession, &EventPoints](double InStartTime, double InEndTime, uint32 InDepth, const FAnimNotifyMessage& InMessage)
+			{
+				if (InEndTime > StartTime && InStartTime < EndTime)
+				{
+					if (InMessage.NameId == NotifyTrackId.NameId && MakeTrackType(InMessage.NotifyEventType) == NotifyTrackId.Type)
+					{
+						SEventTimelineView::FTimelineEventData::EventPoint Point;
+						Point.Time = InMessage.RecordingTime;
+						Point.Color = MakeNotifyColor(InMessage.NameId);
+						EventPoints.Add(Point);
+					}
+				}
+				return TraceServices::EEventEnumerate::Continue;
+			});
+		});
+
+		AnimationProvider->EnumerateNotifyStateTimelines(ObjectId, [this, StartTime, EndTime, AnimationProvider, AnalysisSession, &EventPoints](uint64 InNotifyId, const FAnimationProvider::AnimNotifyTimeline& InTimeline)
+		{
+			InTimeline.EnumerateEvents(StartTime, EndTime, [this, StartTime, EndTime, AnimationProvider, AnalysisSession, &EventPoints](double InStartTime, double InEndTime, uint32 InDepth, const FAnimNotifyMessage& InMessage)
+			{
+				if (InEndTime > StartTime && InStartTime < EndTime)
+				{
+					if (InMessage.NameId == NotifyTrackId.NameId && MakeTrackType(InMessage.NotifyEventType) == NotifyTrackId.Type)
+					{
+						SEventTimelineView::FTimelineEventData::EventWindow Window;
+						Window.TimeStart = InMessage.RecordingTime;
+						Window.TimeEnd = InMessage.RecordingTime + InMessage.Duration;
+						Window.Color = MakeNotifyColor(InMessage.NameId);
+						EventData->Windows.Add(Window);
+					}
+				}
+				
+				return TraceServices::EEventEnumerate::Continue;
+			});
+		});
+
+	}
+
+	bool bChanged = false;
+
+	if (TrackName.IsEmpty() && AnimationProvider)
+    {
+    	TrackName = FText::FromString(AnimationProvider->GetName(NotifyTrackId.NameId));
+    	bChanged = true;
+    }
+
+	return bChanged;
+	
+}
+
+TSharedPtr<SWidget> FNotifiesTrack::GetDetailsViewInternal() 
+{
+	IRewindDebugger* RewindDebugger = IRewindDebugger::Instance();
+	TSharedPtr<SNotifiesView> NotifiesView = SNew(SNotifiesView, ObjectId, RewindDebugger->CurrentTraceTime(), *RewindDebugger->GetAnalysisSession())
+				.CurrentTime_Lambda([RewindDebugger]{ return RewindDebugger->CurrentTraceTime(); });
+	return NotifiesView;
+}
+
+TSharedPtr<SWidget> FNotifyTrack::GetTimelineViewInternal()
+{
+	return SNew(SEventTimelineView)
+		.ViewRange_Lambda([]() { return IRewindDebugger::Instance()->GetCurrentViewRange(); })
+		.EventData_Raw(this, &FNotifyTrack::GetEventData);
+}
+
+static const FName NotifiesName("Notifies");
+
+FName FNotifiesTrackCreator::GetTargetTypeNameInternal() const
+{
+	static const FName AnimInstanceName("AnimInstance");
+	return AnimInstanceName;
+}
+	
+FName FNotifiesTrackCreator::GetNameInternal() const
+{
+	return NotifiesName;
+}
+
+TSharedPtr<RewindDebugger::FRewindDebuggerTrack> FNotifiesTrackCreator::CreateTrackInternal(uint64 ObjectId) const
+{
+	return MakeShared<RewindDebugger::FNotifiesTrack>(ObjectId);
+}
+
+		
+FNotifiesTrack::FNotifiesTrack(uint64 InObjectId) :
+	ObjectId(InObjectId)
+{
+	Icon = FSlateIcon("EditorStyle", "Sequencer.Tracks.Event", "Sequencer.Tracks.Event");
+}
+
+
+bool FNotifiesTrack::UpdateInternal()
+{
+	IRewindDebugger* RewindDebugger = IRewindDebugger::Instance();
+	
+	TRange<double> TraceTimeRange = RewindDebugger->GetCurrentTraceRange();
+	double StartTime = TraceTimeRange.GetLowerBoundValue();
+	double EndTime = TraceTimeRange.GetUpperBoundValue();
+	
+	const TraceServices::IAnalysisSession* AnalysisSession = RewindDebugger->GetAnalysisSession();
+	const FGameplayProvider* GameplayProvider = AnalysisSession->ReadProvider<FGameplayProvider>(FGameplayProvider::ProviderName);
+	const FAnimationProvider* AnimationProvider = AnalysisSession->ReadProvider<FAnimationProvider>(FAnimationProvider::ProviderName);
+	
+	bool bChanged = false;
+	
+	if(GameplayProvider && AnimationProvider)
+	{
+		TArray<FNotifyTrack::FNotifyTrackId> UniqueTrackIds;
+		
+		TraceServices::FAnalysisSessionReadScope SessionReadScope(*AnalysisSession);
+
+		AnimationProvider->ReadNotifyTimeline(ObjectId, [this, StartTime, EndTime, AnimationProvider, AnalysisSession, &UniqueTrackIds](const FAnimationProvider::AnimNotifyTimeline& InTimeline)
+		{
+			InTimeline.EnumerateEvents(StartTime, EndTime, [this, StartTime, EndTime, AnimationProvider, AnalysisSession, &UniqueTrackIds](double InStartTime, double InEndTime, uint32 InDepth, const FAnimNotifyMessage& InMessage)
+			{
+				if (InMessage.NotifyEventType != EAnimNotifyMessageType::Tick)
+				{
+					UniqueTrackIds.AddUnique({InMessage.NameId, MakeTrackType(InMessage.NotifyEventType)});
+				}
+				return TraceServices::EEventEnumerate::Continue;
+			});
+		});
+
+		AnimationProvider->EnumerateNotifyStateTimelines(ObjectId, [this, StartTime, EndTime, AnimationProvider, AnalysisSession, &UniqueTrackIds](uint64 InNotifyId, const FAnimationProvider::AnimNotifyTimeline& InTimeline)
+		{
+			InTimeline.EnumerateEvents(StartTime, EndTime, [this, StartTime, EndTime, AnimationProvider, AnalysisSession, &UniqueTrackIds](double InStartTime, double InEndTime, uint32 InDepth, const FAnimNotifyMessage& InMessage)
+			{
+				if (InEndTime > StartTime && InStartTime < EndTime)
+				{
+					if (InMessage.NotifyEventType == EAnimNotifyMessageType::Begin)
+					{
+						UniqueTrackIds.AddUnique({InMessage.NameId, MakeTrackType(InMessage.NotifyEventType)});
+					}
+				}
+				
+				return TraceServices::EEventEnumerate::Continue;
+			});
+		});
+		
+		UniqueTrackIds.StableSort();
+		const int32 TrackCount = UniqueTrackIds.Num();
+
+		if (Children.Num() != TrackCount)
+			bChanged = true;
+
+		Children.SetNum(UniqueTrackIds.Num());
+		for(int i = 0; i < TrackCount; i++)
+		{
+			if (!Children[i].IsValid() || !(Children[i].Get()->GetNotifyTrackId() == UniqueTrackIds[i]))
+			{
+				Children[i] = MakeShared<FNotifyTrack>(ObjectId, UniqueTrackIds[i]);
+				bChanged = true;
+			}
+
+			bChanged = bChanged || Children[i]->Update();
+		}
+	}
+	
+	return bChanged;
+}
+	
+void FNotifiesTrack::IterateSubTracksInternal(TFunction<void(TSharedPtr<FRewindDebuggerTrack> SubTrack)> IteratorFunction)
+{
+	for(TSharedPtr<FNotifyTrack>& Track : Children)
+	{
+		IteratorFunction(Track);
+	}
+};
+
+bool FNotifiesTrackCreator::HasDebugInfoInternal(uint64 ObjectId) const
+{
+	const TraceServices::IAnalysisSession* AnalysisSession = IRewindDebugger::Instance()->GetAnalysisSession();
+	
+	TraceServices::FAnalysisSessionReadScope SessionReadScope(*AnalysisSession);
+	bool bHasData = false;
+	if (const FAnimationProvider* AnimationProvider = AnalysisSession->ReadProvider<FAnimationProvider>(FAnimationProvider::ProviderName))
+	{
+		AnimationProvider->ReadNotifyTimeline(ObjectId, [&bHasData](const FAnimationProvider::AnimNotifyTimeline& InTimeline)
+		{
+			bHasData = true;
+		});
+	}
+	return bHasData;
+}
+
+	
+}
+
+#undef LOCTEXT_NAMESPACE
\ No newline at end of file
diff --git a/Engine/Plugins/Animation/GameplayInsights/Source/GameplayInsights/Private/NotifiesTrack.h b/Engine/Plugins/Animation/GameplayInsights/Source/GameplayInsights/Private/NotifiesTrack.h
new file mode 100644
--- /dev/null
+++ b/Engine/Plugins/Animation/GameplayInsights/Source/GameplayInsights/Private/NotifiesTrack.h
@@ -0,0 +1,108 @@
+// Copyright Epic Games, Inc. All Rights Reserved.
+#pragma once
+
+#include "IAnimationProvider.h"
+#include "IRewindDebuggerTrackCreator.h"
+#include "RewindDebuggerTrack.h"
+#include "Textures/SlateIcon.h"
+#include "SEventTimelineView.h"
+
+namespace RewindDebugger
+{
+
+class FNotifyTrack : public FRewindDebuggerTrack
+{
+public:
+	enum class ENotifyTrackType
+	{
+		SyncMarker,
+		Notify,
+		NotifyState
+	};
+	
+	struct FNotifyTrackId
+	{
+		uint64 NameId;
+		ENotifyTrackType Type;
+
+		bool operator == (const FNotifyTrackId& Other) const
+		{
+			return NameId == Other.NameId && Type == Other.Type;
+		}
+
+		bool operator <(const FNotifyTrack::FNotifyTrackId& Other) const
+		{
+			if (Type == Other.Type)
+			{
+				return NameId < Other.NameId; 
+			}
+			
+			return static_cast<int>(Type) < static_cast<int>(Other.Type);
+		}		
+	};
+
+	FNotifyTrack(uint64 InObjectId, const FNotifyTrackId& TrackId);
+	const FNotifyTrackId& GetNotifyTrackId() const { return NotifyTrackId; }
+private:
+	virtual bool UpdateInternal() override;
+	TSharedPtr<SWidget> GetDetailsViewInternal() override;
+	virtual TSharedPtr<SWidget> GetTimelineViewInternal() override;
+	virtual FSlateIcon GetIconInternal() override { return Icon; }
+	virtual FText GetDisplayNameInternal() const override { return TrackName; }
+	virtual uint64 GetObjectIdInternal() const override { return ObjectId; }
+
+	TSharedPtr<SEventTimelineView::FTimelineEventData> GetEventData() const;
+	
+	mutable TSharedPtr<SEventTimelineView::FTimelineEventData> EventData;
+	mutable int EventUpdateRequested = 0;
+	
+	uint64 ObjectId;
+	FNotifyTrackId NotifyTrackId;
+	FSlateIcon Icon;
+	FText TrackName;
+	
+	enum class EFilterType
+	{
+		Notify,
+		NotifyState,
+		SyncMarker,
+	};
+
+	EFilterType FilterType = EFilterType::Notify;
+};
+
+
+
+// Parent track of all the curves
+class FNotifiesTrack : public FRewindDebuggerTrack
+{
+public:
+	FNotifiesTrack(uint64 InObjectId);
+
+private:
+	virtual bool UpdateInternal() override;
+	virtual TSharedPtr<SWidget> GetDetailsViewInternal() override;
+
+	virtual FSlateIcon GetIconInternal() override { return Icon; }
+	virtual FName GetNameInternal() const override { return "Notifies"; }
+	virtual FText GetDisplayNameInternal() const override { return NSLOCTEXT("RewindDebugger", "NotifiesTrackName", "Notifies"); }
+	virtual uint64 GetObjectIdInternal() const override { return ObjectId; }
+	virtual void IterateSubTracksInternal(TFunction<void(TSharedPtr<FRewindDebuggerTrack> SubTrack)> IteratorFunction) override;
+
+	FSlateIcon Icon;
+	uint64 ObjectId;
+
+	TArray<TSharedPtr<FNotifyTrack>> Children;
+};
+
+
+class FNotifiesTrackCreator : public IRewindDebuggerTrackCreator
+{
+private:
+	virtual FName GetTargetTypeNameInternal() const;
+	virtual FName GetNameInternal() const override;
+	virtual TSharedPtr<FRewindDebuggerTrack> CreateTrackInternal(uint64 ObjectId) const override;
+	virtual bool HasDebugInfoInternal(uint64 ObjectId) const override;
+};
+
+}
\ No newline at end of file
diff --git a/Engine/Plugins/Animation/GameplayInsights/Source/GameplayInsights/Private/SEventTimelineView.cpp b/Engine/Plugins/Animation/GameplayInsights/Source/GameplayInsights/Private/SEventTimelineView.cpp
new file mode 100644
--- /dev/null
+++ b/Engine/Plugins/Animation/GameplayInsights/Source/GameplayInsights/Private/SEventTimelineView.cpp
@@ -0,0 +1,135 @@
+// Copyright Epic Games, Inc. All Rights Reserved.
+
+#include "SEventTimelineView.h"
+#include "Framework/Application/SlateApplication.h"
+#include "Rendering/DrawElements.h"
+#include "SSimpleTimeSlider.h"
+
+#define LOCTEXT_NAMESPACE "SEventTimelineView"
+
+
+void SEventTimelineView::Construct( const SEventTimelineView::FArguments& InArgs )
+{
+    ViewRange = InArgs._ViewRange;
+	DesiredSize = InArgs._DesiredSize;
+	EventData = InArgs._EventData;
+
+	// set clipping on by default, since the OnPaint function is drawing outside the bounds
+	// Clipping = EWidgetClipping::ClipToBounds;
+	EventBrush = FAppStyle::GetBrush("Sequencer.KeyDiamond");
+	EventBorderBrush = FAppStyle::GetBrush("Sequencer.KeyDiamondBorder");
+}
+
+int32 SEventTimelineView::OnPaint( const FPaintArgs& Args, const FGeometry& AllottedGeometry, const FSlateRect& MyCullingRect, FSlateWindowElementList& OutDrawElements, int32 LayerId, const FWidgetStyle& InWidgetStyle, bool bParentEnabled ) const
+{
+	int32 NewLayer = PaintEvents( AllottedGeometry, MyCullingRect, OutDrawElements, LayerId, InWidgetStyle, bParentEnabled );
+
+	return FMath::Max( NewLayer, SCompoundWidget::OnPaint(Args, AllottedGeometry, MyCullingRect, OutDrawElements, NewLayer, InWidgetStyle, ShouldBeEnabled( bParentEnabled ) ) );
+}
+
+
+FVector2D SEventTimelineView::ComputeDesiredSize( float ) const
+{
+	return DesiredSize.Get();
+}
+
+int32 SEventTimelineView::PaintEvents(const FGeometry& AllottedGeometry, const FSlateRect& MyCullingRect, FSlateWindowElementList& OutDrawElements, int32 LayerId, const FWidgetStyle& InWidgetStyle, bool bParentEnabled ) const
+{
+	TRange<double> DebugTimeRange = ViewRange.Get();
+	
+	SSimpleTimeSlider::FScrubRangeToScreen RangeToScreen( DebugTimeRange, AllottedGeometry.GetLocalSize() );
+	FVector2D Size = AllottedGeometry.GetLocalSize();
+
+	if (TSharedPtr<FTimelineEventData> Event = EventData.Get())
+	{
+		auto &EventPoints = Event->Points;
+		const int NumPoints = EventPoints.Num();
+		
+		const ESlateDrawEffect LineDrawEffects = ESlateDrawEffect::NoPixelSnapping;
+		
+		double PrevPointTime = 0;
+		int OverlappingPointCount = 0;
+
+		if (NumPoints > 0)
+		{
+			for(int i=0; i<NumPoints; i++)
+			{
+				const FTimelineEventData::EventPoint& Point = EventPoints[i];
+			
+				FVector2D EventSize = EventBrush->GetImageSize();
+				
+				float X = RangeToScreen.InputToLocalX(Point.Time);
+				X = X-EventSize.X/2;
+				float Y = (AllottedGeometry.Size.Y - EventSize.Y)/2;
+				if (Point.Time == PrevPointTime)
+				{
+					OverlappingPointCount++;
+
+					static const int OverlapOffsetAmount = 2;
+					Y += OverlapOffsetAmount * OverlappingPointCount;
+				}
+				else
+				{
+					OverlappingPointCount = 0;
+				}
+				
+				PrevPointTime = Point.Time;
+				
+				FSlateDrawElement::MakeBox(OutDrawElements, LayerId,
+					AllottedGeometry.ToPaintGeometry(FVector2D(X, Y) , EventSize, 1.0), EventBrush, ESlateDrawEffect::None, Point.Color);
+				
+				FSlateDrawElement::MakeBox(OutDrawElements, LayerId,
+					AllottedGeometry.ToPaintGeometry(FVector2D(X, Y) , EventSize, 1.0), EventBorderBrush, ESlateDrawEffect::None, FLinearColor::Black); 
+			}
+		}
+		
+		auto &EventWindows = Event->Windows;
+		const int NumWindows = EventWindows.Num();
+		
+	
+		const FSlateBrush* Brush = FAppStyle::GetBrush("Sequencer.SectionArea.Background");
+
+		
+		
+		if (NumWindows > 0)
+		{
+			for(int i=0; i<NumWindows; i++)
+			{
+				const FTimelineEventData::EventWindow& Window = EventWindows[i];
+			
+				FVector2D EventSize = EventBrush->GetImageSize();
+				
+				float XStart = RangeToScreen.InputToLocalX(Window.TimeStart);
+				float XStartDiamond = XStart - EventSize.X/2;
+				float XEnd = RangeToScreen.InputToLocalX(Window.TimeEnd);
+				float XEndDiamond = XEnd - EventSize.X/2;
+				float Y = (AllottedGeometry.Size.Y - EventSize.Y)/2;
+
+				// window bar
+				FSlateDrawElement::MakeBox(OutDrawElements, LayerId++,
+					AllottedGeometry.ToPaintGeometry(FVector2D(XStart,Y + 1),FVector2D(XEnd-XStart, EventSize.Y-2), 1), Brush, ESlateDrawEffect::None, Window.Color);
+
+				// key diamond at start
+				FSlateDrawElement::MakeBox(OutDrawElements, LayerId,
+					AllottedGeometry.ToPaintGeometry(FVector2D(XStartDiamond, Y) , EventSize, 1.0), EventBrush, ESlateDrawEffect::None, Window.Color);
+				
+				FSlateDrawElement::MakeBox(OutDrawElements, LayerId,
+					AllottedGeometry.ToPaintGeometry(FVector2D(XStartDiamond, Y) , EventSize, 1.0), EventBorderBrush, ESlateDrawEffect::None, FLinearColor::Black); 
+
+				// key diamond at end
+				FSlateDrawElement::MakeBox(OutDrawElements, LayerId,
+					AllottedGeometry.ToPaintGeometry(FVector2D(XEndDiamond, Y) , EventSize, 1.0), EventBrush, ESlateDrawEffect::None, Window.Color);
+				
+				FSlateDrawElement::MakeBox(OutDrawElements, LayerId,
+					AllottedGeometry.ToPaintGeometry(FVector2D(XEndDiamond, Y) , EventSize, 1.0), EventBorderBrush, ESlateDrawEffect::None, FLinearColor::Black);
+			}
+		}
+		
+
+		LayerId++;
+	}
+
+	return LayerId;
+}
+
+#undef LOCTEXT_NAMESPACE
diff --git a/Engine/Plugins/Animation/GameplayInsights/Source/GameplayInsights/Private/SEventTimelineView.h b/Engine/Plugins/Animation/GameplayInsights/Source/GameplayInsights/Private/SEventTimelineView.h
new file mode 100644
--- /dev/null
+++ b/Engine/Plugins/Animation/GameplayInsights/Source/GameplayInsights/Private/SEventTimelineView.h
@@ -0,0 +1,77 @@
+// Copyright Epic Games, Inc. All Rights Reserved.
+
+#pragma once
+
+#include "CoreMinimal.h"
+#include "Widgets/SCompoundWidget.h"
+#include "Widgets/DeclarativeSyntaxSupport.h"
+
+class FPaintArgs;
+class FSlateWindowElementList;
+class FTimeSliderController;
+
+class SEventTimelineView : public SCompoundWidget
+{
+public:
+	struct FTimelineEventData
+	{
+	public:
+		struct EventPoint
+		{
+			double Time;
+			FText Type;
+			FText Description;
+			FLinearColor Color;
+		};
+		
+		struct EventWindow
+		{
+			double TimeStart;
+			double TimeEnd;
+			FText Type;
+			FText Description;
+			FLinearColor Color;
+		};
+
+		TArray<EventPoint> Points;
+		TArray<EventWindow> Windows;
+	};
+	
+	SLATE_BEGIN_ARGS(SEventTimelineView)
+		: _ViewRange(TRange<double>(0,10))
+		, _DesiredSize(FVector2D(100.f,20.f))
+	{}
+    	/** View time range */
+    	SLATE_ATTRIBUTE(TRange<double>, ViewRange);
+
+		/** Data for events to render */
+		SLATE_ATTRIBUTE(TSharedPtr<FTimelineEventData>, EventData);
+
+		/** Desired widget size */
+		SLATE_ATTRIBUTE(FVector2D, DesiredSize);
+	
+	SLATE_END_ARGS()
+
+
+	/**
+	 * Construct the widget
+	 * 
+	 * @param InArgs   A declaration from which to construct the widget
+	 */
+	void Construct( const FArguments& InArgs );
+
+protected:
+	// SWidget interface
+	virtual FVector2D ComputeDesiredSize(float) const override;
+	virtual int32 OnPaint( const FPaintArgs& Args, const FGeometry& AllottedGeometry, const FSlateRect& MyCullingRect, FSlateWindowElementList& OutDrawElements, int32 LayerId, const FWidgetStyle& InWidgetStyle, bool bParentEnabled ) const override;
+	
+	int32 PaintEvents(const FGeometry& AllottedGeometry, const FSlateRect& MyCullingRect, FSlateWindowElementList& OutDrawElements, int32 LayerId, const FWidgetStyle& InWidgetStyle, bool bParentEnabled ) const;
+	
+	TAttribute<TRange<double>> ViewRange;
+	TAttribute<TSharedPtr<FTimelineEventData>> EventData;
+	TAttribute<FVector2D> DesiredSize;
+	
+	const FSlateBrush* EventBrush;
+	const FSlateBrush* EventBorderBrush;
+};
+
diff --git a/Engine/Plugins/Animation/GameplayInsights/Source/GameplayInsights/Private/SNotifiesView.cpp b/Engine/Plugins/Animation/GameplayInsights/Source/GameplayInsights/Private/SNotifiesView.cpp
--- a/Engine/Plugins/Animation/GameplayInsights/Source/GameplayInsights/Private/SNotifiesView.cpp
+++ b/Engine/Plugins/Animation/GameplayInsights/Source/GameplayInsights/Private/SNotifiesView.cpp
@@ -73,6 +73,11 @@ void SNotifiesView::GetVariantsAtFrame(const TraceServices::FFrame& InFrame, TAr
 		bool bForEvent = true;
 		auto ProcessEvent = [this, &AnimationProvider, &GameplayProvider, &Header, &bForEvent](double InStartTime, double InEndTime, uint32 InDepth, const FAnimNotifyMessage& InMessage)
 		{
+			if (bFilterIsSet && InMessage.NameId != FilterNotifyNameId)
+			{
+				return TraceServices::EEventEnumerate::Continue;
+			}
+			
 			const TCHAR* Name = AnimationProvider->GetName(InMessage.NameId);
 			TSharedRef<FVariantTreeNode> NotifyHeader = Header->AddChild(FVariantTreeNode::MakeHeader(FText::FromString(Name), InMessage.NameId));
 
@@ -113,52 +118,4 @@ FName SNotifiesView::GetName() const
 	return NotifiesName;
 }
 
-FName FNotifiesViewCreator::GetTargetTypeName() const
-{
-	static FName TargetTypeName = "AnimInstance";
-	return TargetTypeName;
-}
-
-FName FNotifiesViewCreator::GetName() const
-{
-	return NotifiesName;
-}
-
-FText FNotifiesViewCreator::GetTitle() const
-{
-	return LOCTEXT("Notifies", "Notifies");
-}
-
-FSlateIcon FNotifiesViewCreator::GetIcon() const
-{
-#if WITH_EDITOR
-	return FSlateIconFinder::FindIconForClass(UAnimInstance::StaticClass());
-#else
-	return FSlateIcon();
-#endif
-}
-
-TSharedPtr<IRewindDebuggerView> FNotifiesViewCreator::CreateDebugView(uint64 ObjectId, double CurrentTime, const TraceServices::IAnalysisSession& AnalysisSession) const
-{
-	return SNew(SNotifiesView, ObjectId, CurrentTime, AnalysisSession);
-}
-
-bool FNotifiesViewCreator::HasDebugInfo(uint64 ObjectId) const
-{
-	const TraceServices::IAnalysisSession* AnalysisSession = IRewindDebugger::Instance()->GetAnalysisSession();
-	
-	TraceServices::FAnalysisSessionReadScope SessionReadScope(*AnalysisSession);
-	bool bHasData = false;
-	if (const FAnimationProvider* AnimationProvider = AnalysisSession->ReadProvider<FAnimationProvider>(FAnimationProvider::ProviderName))
-	{
-		AnimationProvider->ReadNotifyTimeline(ObjectId, [&bHasData](const FAnimationProvider::AnimNotifyTimeline& InTimeline)
-		{
-			bHasData = true;
-		});
-	}
-
-	return bHasData;
-}
-
-
 #undef LOCTEXT_NAMESPACE
diff --git a/Engine/Plugins/Animation/GameplayInsights/Source/GameplayInsights/Private/SNotifiesView.h b/Engine/Plugins/Animation/GameplayInsights/Source/GameplayInsights/Private/SNotifiesView.h
--- a/Engine/Plugins/Animation/GameplayInsights/Source/GameplayInsights/Private/SNotifiesView.h
+++ b/Engine/Plugins/Animation/GameplayInsights/Source/GameplayInsights/Private/SNotifiesView.h
@@ -12,15 +12,13 @@ class SNotifiesView : public SPropertiesDebugViewBase
 public:
 	virtual void GetVariantsAtFrame(const TraceServices::FFrame& InFrame, TArray<TSharedRef<FVariantTreeNode>>& OutVariants) const override;
 	virtual FName GetName() const override;
-};
 
-class FNotifiesViewCreator : public IRewindDebuggerViewCreator
-{
-	public:
-		virtual FName GetTargetTypeName() const;
-		virtual FName GetName() const override;
-		virtual FText GetTitle() const override;
-		virtual FSlateIcon GetIcon() const override;
-		virtual TSharedPtr<IRewindDebuggerView> CreateDebugView(uint64 ObjectId, double CurrentTime, const TraceServices::IAnalysisSession& InAnalysisSession) const override;
-		virtual bool HasDebugInfo(uint64 ObjectId) const override;
+	void SetNotifyFilter(uint64 NameId)
+	{
+		FilterNotifyNameId = NameId;
+		bFilterIsSet = true;
+	}
+private:
+	uint64 FilterNotifyNameId = 0;
+	bool bFilterIsSet = false;
 };
\ No newline at end of file
diff --git a/Engine/Plugins/Animation/GameplayInsights/Source/GameplayInsights/Public/IAnimationProvider.h b/Engine/Plugins/Animation/GameplayInsights/Source/GameplayInsights/Public/IAnimationProvider.h
--- a/Engine/Plugins/Animation/GameplayInsights/Source/GameplayInsights/Public/IAnimationProvider.h
+++ b/Engine/Plugins/Animation/GameplayInsights/Source/GameplayInsights/Public/IAnimationProvider.h
@@ -218,6 +218,7 @@ struct FAnimNotifyMessage
 	uint64 NotifyId = 0;
 	const TCHAR* Name = nullptr;
 	uint32 NameId = 0;
+	double RecordingTime = 0.0f;
 	float Time = 0.0f; 
 	float Duration = 0.0f;
 	EAnimNotifyMessageType NotifyEventType = EAnimNotifyMessageType::Event;
diff --git a/Engine/Source/Runtime/Engine/Private/Animation/AnimTrace.cpp b/Engine/Source/Runtime/Engine/Private/Animation/AnimTrace.cpp
--- a/Engine/Source/Runtime/Engine/Private/Animation/AnimTrace.cpp
+++ b/Engine/Source/Runtime/Engine/Private/Animation/AnimTrace.cpp
@@ -202,8 +202,9 @@ UE_TRACE_EVENT_BEGIN(Animation, Name, NoSync|Important)
 	UE_TRACE_EVENT_FIELD(UE::Trace::WideString, Name)
 UE_TRACE_EVENT_END()
 
-UE_TRACE_EVENT_BEGIN(Animation, Notify)
+UE_TRACE_EVENT_BEGIN(Animation, Notify2)
 	UE_TRACE_EVENT_FIELD(uint64, Cycle)
+	UE_TRACE_EVENT_FIELD(double, RecordingTime)
 	UE_TRACE_EVENT_FIELD(uint64, AnimInstanceId)
 	UE_TRACE_EVENT_FIELD(uint64, AssetId)
 	UE_TRACE_EVENT_FIELD(uint64, NotifyId)
@@ -213,8 +214,9 @@ UE_TRACE_EVENT_BEGIN(Animation, Notify)
 	UE_TRACE_EVENT_FIELD(uint8, NotifyEventType)
 UE_TRACE_EVENT_END()
 
-UE_TRACE_EVENT_BEGIN(Animation, SyncMarker)
+UE_TRACE_EVENT_BEGIN(Animation, SyncMarker2)
 	UE_TRACE_EVENT_FIELD(uint64, Cycle)
+	UE_TRACE_EVENT_FIELD(double, RecordingTime)
 	UE_TRACE_EVENT_FIELD(uint64, AnimInstanceId)
 	UE_TRACE_EVENT_FIELD(uint32, NameId)
 UE_TRACE_EVENT_END()
@@ -1071,15 +1073,16 @@ void FAnimTrace::OutputAnimNotify(UAnimInstance* InAnimInstance, const FAnimNoti
 
 	const uint32 NameId = OutputName(InNotifyEvent.NotifyName);
 
-	UE_TRACE_LOG(Animation, Notify, AnimationChannel)
-		<< Notify.Cycle(FPlatformTime::Cycles64())
-		<< Notify.AnimInstanceId(FObjectTrace::GetObjectId(InAnimInstance))
-		<< Notify.AssetId(FObjectTrace::GetObjectId(NotifyAsset))
-		<< Notify.NotifyId(FObjectTrace::GetObjectId(NotifyObject))
-		<< Notify.NameId(NameId)
-		<< Notify.Time(InNotifyEvent.GetTime())
-		<< Notify.Duration(InNotifyEvent.GetDuration())
-		<< Notify.NotifyEventType((uint8)InEventType);
+	UE_TRACE_LOG(Animation, Notify2, AnimationChannel)
+		<< Notify2.Cycle(FPlatformTime::Cycles64())
+		<< Notify2.RecordingTime(FObjectTrace::GetWorldElapsedTime(InAnimInstance->GetWorld()))
+		<< Notify2.AnimInstanceId(FObjectTrace::GetObjectId(InAnimInstance))
+		<< Notify2.AssetId(FObjectTrace::GetObjectId(NotifyAsset))
+		<< Notify2.NotifyId(FObjectTrace::GetObjectId(NotifyObject))
+		<< Notify2.NameId(NameId)
+		<< Notify2.Time(InNotifyEvent.GetTime())
+		<< Notify2.Duration(InNotifyEvent.GetDuration())
+		<< Notify2.NotifyEventType((uint8)InEventType);
 }
 
 void FAnimTrace::OutputAnimSyncMarker(UAnimInstance* InAnimInstance, const FPassedMarker& InPassedSyncMarker)
@@ -1099,10 +1102,11 @@ void FAnimTrace::OutputAnimSyncMarker(UAnimInstance* InAnimInstance, const FPass
 
 	const uint32 NameId = OutputName(InPassedSyncMarker.PassedMarkerName);
 
-	UE_TRACE_LOG(Animation, SyncMarker, AnimationChannel)
-		<< SyncMarker.Cycle(FPlatformTime::Cycles64())
-		<< SyncMarker.AnimInstanceId(FObjectTrace::GetObjectId(InAnimInstance))
-		<< SyncMarker.NameId(NameId);
+	UE_TRACE_LOG(Animation, SyncMarker2, AnimationChannel)
+		<< SyncMarker2.Cycle(FPlatformTime::Cycles64())
+		<< SyncMarker2.AnimInstanceId(FObjectTrace::GetObjectId(InAnimInstance))
+		<< SyncMarker2.RecordingTime(FObjectTrace::GetWorldElapsedTime(InAnimInstance->GetWorld()))
+		<< SyncMarker2.NameId(NameId);
 }
 
 void FAnimTrace::OutputMontage(UAnimInstance* InAnimInstance, const FAnimMontageInstance& InMontageInstance)
