Fix for property tracing from github pull request #9146

diff --git a/Engine/Plugins/Animation/GameplayInsights/Source/GameplayInsights/Private/ObjectPropertyTrace.cpp b/Engine/Plugins/Animation/GameplayInsights/Source/GameplayInsights/Private/ObjectPropertyTrace.cpp
--- a/Engine/Plugins/Animation/GameplayInsights/Source/GameplayInsights/Private/ObjectPropertyTrace.cpp
+++ b/Engine/Plugins/Animation/GameplayInsights/Source/GameplayInsights/Private/ObjectPropertyTrace.cpp
@@ -98,11 +98,11 @@ namespace ObjectPropertyTrace
 					FString TypeString = FString::Printf(TEXT("{%s, %s}"), *MapProperty->KeyProp->GetCPPType(), *MapProperty->ValueProp->GetCPPType());
 					InFunction(TypeString, KeyString, TEXT("{...}"), MapEntryId, MapRootId);
 
-					const void* KeyPtr = Helper.GetKeyPtr(DynamicIndex);
-					IteratePropertiesRecursive(MapProperty->KeyProp, KeyPtr, MapProperty->KeyProp->GetName(), InFunction, InId, MapEntryId);
+					const void* PairPtr = Helper.GetPairPtr(DynamicIndex);
+					
+					IteratePropertiesRecursive(MapProperty->KeyProp, PairPtr, MapProperty->KeyProp->GetName(), InFunction, InId, MapEntryId);
 
-					const void* ValuePtr = Helper.GetValuePtr(DynamicIndex);
-					IteratePropertiesRecursive(MapProperty->ValueProp, ValuePtr, MapProperty->ValueProp->GetName(), InFunction, InId, MapEntryId);
+					IteratePropertiesRecursive(MapProperty->ValueProp, PairPtr, MapProperty->ValueProp->GetName(), InFunction, InId, MapEntryId);
 
 					--Num;
 				}
