Anim Blueprint Rewind Debugger debug optimizations

diff --git a/Engine/Plugins/Animation/GameplayInsights/Source/GameplayInsights/Private/AnimationProvider.cpp b/Engine/Plugins/Animation/GameplayInsights/Source/GameplayInsights/Private/AnimationProvider.cpp
--- a/Engine/Plugins/Animation/GameplayInsights/Source/GameplayInsights/Private/AnimationProvider.cpp
+++ b/Engine/Plugins/Animation/GameplayInsights/Source/GameplayInsights/Private/AnimationProvider.cpp
@@ -148,10 +148,7 @@ void FAnimationProvider::EnumerateAnimGraphTimelines(TFunctionRef<void(uint64 Ob
 	
 	for(auto& IndexMapping : ObjectIdToAnimGraphTimelines)
 	{
-		for (const TSharedRef<TraceServices::TIntervalTimeline<FAnimGraphMessage>>& Timeline : AnimGraphTimelines)
-		{
-			Callback(IndexMapping.Key, AnimGraphTimelines[IndexMapping.Value].Get());
-		}
+		Callback(IndexMapping.Key, AnimGraphTimelines[IndexMapping.Value].Get());
 	}
 }
 
diff --git a/Engine/Plugins/Animation/GameplayInsights/Source/RewindDebugger/Private/RewindDebugger.cpp b/Engine/Plugins/Animation/GameplayInsights/Source/RewindDebugger/Private/RewindDebugger.cpp
--- a/Engine/Plugins/Animation/GameplayInsights/Source/RewindDebugger/Private/RewindDebugger.cpp
+++ b/Engine/Plugins/Animation/GameplayInsights/Source/RewindDebugger/Private/RewindDebugger.cpp
@@ -41,6 +41,7 @@ FRewindDebugger::FRewindDebugger()  :
 	bAutoRecord(false),
 	bRecording(false),
 	PlaybackRate(1),
+	PreviousTraceTime(-1),
 	CurrentScrubTime(0),
 	CurrentViewRange(0,0),
 	CurrentTraceRange(0,0),
@@ -596,11 +597,11 @@ void FRewindDebugger::GetScrubTimeInformation(double InDebugTime, FScrubTimeInfo
 							{
 								StartEventIndex = MiddleEventIndex;
 							}
-						}
-						
+}
+
 						// Ensure there is not frames between start and end index
 						check(EndEventIndex == StartEventIndex + 1)
-				
+
 						const FRecordingInfoMessage& Event = Recording->GetEvent(StartEventIndex);
 						const FRecordingInfoMessage& NextEvent = Recording->GetEvent(EndEventIndex);
 
@@ -609,11 +610,11 @@ void FRewindDebugger::GetScrubTimeInformation(double InDebugTime, FScrubTimeInfo
 
 						// Choose frame that is nearest to the debug time
 						if (InDebugTime - Event.ElapsedTime < NextEvent.ElapsedTime - InDebugTime)
-						{
+{
 							ScrubFrameIndex = StartEventIndex;
 						}
 						else
-						{
+	{
 							ScrubFrameIndex = EndEventIndex;
 						}
 					}
@@ -639,6 +640,7 @@ const TraceServices::IAnalysisSession* FRewindDebugger::GetAnalysisSession() con
 
 void FRewindDebugger::Tick(float DeltaTime)
 {
+	TRACE_CPUPROFILER_EVENT_SCOPE(FRewindDebugger::Tick);
 	if (const TraceServices::IAnalysisSession* Session = GetAnalysisSession())
 	{
 		if (bRecording)
@@ -681,47 +683,54 @@ void FRewindDebugger::Tick(float DeltaTime)
 						}
 					}
 
-					double CurrentTraceTime = TraceTime.Get();
+					const double CurrentTraceTime = TraceTime.Get();
+					if (CurrentTraceTime != PreviousTraceTime)
+					{
+						PreviousTraceTime = CurrentTraceTime;
+						
 					const TraceServices::IFrameProvider& FrameProvider = TraceServices::ReadFrameProvider(*Session);
 					TraceServices::FFrame Frame;
 					if(FrameProvider.GetFrameFromTime(ETraceFrameType::TraceFrameType_Game, CurrentTraceTime, Frame))
 					{
-						
+							{
+								TRACE_CPUPROFILER_EVENT_SCOPE(FRewindDebugger::Tick_UpdateActorPosition);
 						// until we have actor transforms traced out, the first skeletal mesh component transform on the target actor be used as as the actor position 
 						uint64 TargetActorId = GetTargetActorId();
 						if (TargetActorId != 0)
 						{
-							if (UObject* ObjectInstance = FObjectTrace::GetObjectFromId(TargetActorId))
+						if(UObject* ObjectInstance = FObjectTrace::GetObjectFromId(TargetActorId))
+						{
+							if (AActor* TargetActor = Cast<AActor>(ObjectInstance))
 							{
-								if (AActor* TargetActor = Cast<AActor>(ObjectInstance))
+								TInlineComponentArray<USkeletalMeshComponent*> SkeletalMeshComponents;
+								TargetActor->GetComponents(SkeletalMeshComponents);
+
+								if (SkeletalMeshComponents.Num() > 0)
 								{
-									TInlineComponentArray<USkeletalMeshComponent*> SkeletalMeshComponents;
-									TargetActor->GetComponents(SkeletalMeshComponents);
+									int64 ObjectId = FObjectTrace::GetObjectId(SkeletalMeshComponents[0]);
 
-									if (SkeletalMeshComponents.Num() > 0)
+									AnimationProvider->ReadSkeletalMeshPoseTimeline(ObjectId, [this, Frame, ObjectId, AnimationProvider](const IAnimationProvider::SkeletalMeshPoseTimeline& TimelineData, bool bHasCurves)
 									{
-										int64 ObjectId = FObjectTrace::GetObjectId(SkeletalMeshComponents[0]);
-
-										AnimationProvider->ReadSkeletalMeshPoseTimeline(ObjectId, [this, Frame, ObjectId, AnimationProvider](const IAnimationProvider::SkeletalMeshPoseTimeline& TimelineData, bool bHasCurves)
+										TimelineData.EnumerateEvents(Frame.StartTime, Frame.EndTime,
+											[this](double InStartTime, double InEndTime, uint32 InDepth, const FSkeletalMeshPoseMessage& PoseMessage)
 											{
-												TimelineData.EnumerateEvents(Frame.StartTime, Frame.EndTime,
-													[this](double InStartTime, double InEndTime, uint32 InDepth, const FSkeletalMeshPoseMessage& PoseMessage)
-													{
-														bTargetActorPositionValid = true;
-														TargetActorPosition = PoseMessage.ComponentToWorld.GetTranslation();
-
-														return TraceServices::EEventEnumerate::Stop;
-													});
+												bTargetActorPositionValid = true;
+												TargetActorPosition = PoseMessage.ComponentToWorld.GetTranslation();
+										
+												return TraceServices::EEventEnumerate::Stop;
 											});
-									}
+									});
 								}
 							}
 						}
+							}
 						
 						// update pose on all SkeletalMeshComponents:
 						// - enumerate all skeletal mesh pose timelines
 						// - check if the corresponding mesh component still exists
 						// - apply the recorded pose for the current Frame
+							{
+								TRACE_CPUPROFILER_EVENT_SCOPE(FRewindDebugger::Tick_UpdatePoses);
 						AnimationProvider->EnumerateSkeletalMeshPoseTimelines([this, &Frame, AnimationProvider, GameplayProvider](uint64 ObjectId, const IAnimationProvider::SkeletalMeshPoseTimeline& TimelineData)
 						{
 							if(UObject* ObjectInstance = FObjectTrace::GetObjectFromId(ObjectId))
@@ -755,7 +764,10 @@ void FRewindDebugger::Tick(float DeltaTime)
 								}
 							}
 						});
+							}
 
+							{
+								TRACE_CPUPROFILER_EVENT_SCOPE(FRewindDebugger::Tick_AnimBlueprintsDebug);
 						// Apply Animation Blueprint Debugging Data:
 						// - enumerate over all anim graph timelines
 						// - check if their instance class still exists and is the debugging target for the Animation Blueprint Editor
@@ -772,15 +784,21 @@ void FRewindDebugger::Tick(float DeltaTime)
 										{
 											if(AnimBlueprint->IsObjectBeingDebugged(AnimInstance))
 											{
+														TRACE_CPUPROFILER_EVENT_SCOPE(FRewindDebugger::Tick_UpdateBlueprintDebug);
 												// update debug info for attached Animation Blueprint editors
 												uint64 Id = FObjectTrace::GetObjectId(AnimInstance);
 												const int32 NodeCount = InstanceClass->GetAnimNodeProperties().Num();
 						
 												FAnimBlueprintDebugData& DebugData = InstanceClass->GetAnimBlueprintDebugData();
+														{
+															TRACE_CPUPROFILER_EVENT_SCOPE(ResetNodeVisitStates);
 												DebugData.ResetNodeVisitSites();
+														}
 					
 												AnimGraphTimeline.EnumerateEvents(Frame.StartTime, Frame.EndTime, [Id, AnimationProvider, GameplayProvider, &DebugData, NodeCount](double InGraphStartTime, double InGraphEndTime, uint32 InDepth, const FAnimGraphMessage& InMessage)
 												{
+															TRACE_CPUPROFILER_EVENT_SCOPE(AnimGraphTimelineEvent);
+															
 													// Basic verification - check node count is the same
 													// @TODO: could add some form of node hash/CRC to the class to improve this
 													if(InMessage.NodeCount == NodeCount)
@@ -790,6 +808,7 @@ void FRewindDebugger::Tick(float DeltaTime)
 														{
 															AnimationProvider->ReadAnimNodesTimeline(Id, [InGraphStartTime, InGraphEndTime, &DebugData](const IAnimationProvider::AnimNodesTimeline& InNodesTimeline)
 															{
+																		TRACE_CPUPROFILER_EVENT_SCOPE(AnimGraphDebugNodeVisits);
 																InNodesTimeline.EnumerateEvents(InGraphStartTime, InGraphEndTime, [&DebugData](double InStartTime, double InEndTime, uint32 InDepth, const FAnimNodeMessage& InMessage)
 																{
 																	DebugData.RecordNodeVisit(InMessage.NodeId, InMessage.PreviousNodeId, InMessage.Weight);
@@ -799,6 +818,7 @@ void FRewindDebugger::Tick(float DeltaTime)
 					
 															AnimationProvider->ReadStateMachinesTimeline(Id, [InGraphStartTime, InGraphEndTime, &DebugData](const IAnimationProvider::StateMachinesTimeline& InStateMachinesTimeline)
 															{
+																		TRACE_CPUPROFILER_EVENT_SCOPE(AnimGraphDebugStateMachine);
 																InStateMachinesTimeline.EnumerateEvents(InGraphStartTime, InGraphEndTime, [&DebugData](double InStartTime, double InEndTime, uint32 InDepth, const FAnimStateMachineMessage& InMessage)
 																{
 																	DebugData.RecordStateData(InMessage.StateMachineIndex, InMessage.StateIndex, InMessage.StateWeight, InMessage.ElapsedTime);
@@ -808,6 +828,7 @@ void FRewindDebugger::Tick(float DeltaTime)
 					
 															AnimationProvider->ReadAnimSequencePlayersTimeline(Id, [InGraphStartTime, InGraphEndTime, GameplayProvider, &DebugData](const IAnimationProvider::AnimSequencePlayersTimeline& InSequencePlayersTimeline)
 															{
+																		TRACE_CPUPROFILER_EVENT_SCOPE(AnimGraphDebugSequencePlayers);
 																InSequencePlayersTimeline.EnumerateEvents(InGraphStartTime, InGraphEndTime, [&DebugData](double InStartTime, double InEndTime, uint32 InDepth, const FAnimSequencePlayerMessage& InMessage)
 																{
 																	DebugData.RecordSequencePlayer(InMessage.NodeId, InMessage.Position, InMessage.Length, InMessage.FrameCounter);
@@ -817,6 +838,7 @@ void FRewindDebugger::Tick(float DeltaTime)
 					
 															AnimationProvider->ReadAnimBlendSpacePlayersTimeline(Id, [InGraphStartTime, InGraphEndTime, GameplayProvider, &DebugData](const IAnimationProvider::BlendSpacePlayersTimeline& InBlendSpacePlayersTimeline)
 															{
+																		TRACE_CPUPROFILER_EVENT_SCOPE(AnimGraphBlendSpaces);
 																InBlendSpacePlayersTimeline.EnumerateEvents(InGraphStartTime, InGraphEndTime, [GameplayProvider, &DebugData](double InStartTime, double InEndTime, uint32 InDepth, const FBlendSpacePlayerMessage& InMessage)
 																{
 																	UBlendSpace* BlendSpace = nullptr;
@@ -833,6 +855,7 @@ void FRewindDebugger::Tick(float DeltaTime)
 					
 															AnimationProvider->ReadAnimSyncTimeline(Id, [InGraphStartTime, InGraphEndTime, AnimationProvider, &DebugData](const IAnimationProvider::AnimSyncTimeline& InAnimSyncTimeline)
 															{
+																		TRACE_CPUPROFILER_EVENT_SCOPE(AnimGraphAnimSync);
 																InAnimSyncTimeline.EnumerateEvents(InGraphStartTime, InGraphEndTime, [AnimationProvider, &DebugData](double InStartTime, double InEndTime, uint32 InDepth, const FAnimSyncMessage& InMessage)
 																{
 																	const TCHAR* GroupName = AnimationProvider->GetName(InMessage.GroupNameId);
@@ -851,6 +874,7 @@ void FRewindDebugger::Tick(float DeltaTime)
 														{
 															AnimationProvider->ReadAnimAttributesTimeline(Id, [InGraphStartTime, InGraphEndTime, AnimationProvider, &DebugData](const IAnimationProvider::AnimAttributeTimeline& InAnimAttributeTimeline)
 															{
+																		TRACE_CPUPROFILER_EVENT_SCOPE(AnimGraphAttributes);
 																InAnimAttributeTimeline.EnumerateEvents(InGraphStartTime, InGraphEndTime, [AnimationProvider, &DebugData](double InStartTime, double InEndTime, uint32 InDepth, const FAnimAttributeMessage& InMessage)
 																{
 																	const TCHAR* AttributeName = AnimationProvider->GetName(InMessage.AttributeNameId);
@@ -867,6 +891,7 @@ void FRewindDebugger::Tick(float DeltaTime)
 														// Anim node values can come from all phases
 														AnimationProvider->ReadAnimNodeValuesTimeline(Id, [InGraphStartTime, InGraphEndTime, AnimationProvider, &DebugData](const IAnimationProvider::AnimNodeValuesTimeline& InNodeValuesTimeline)
 														{
+																	TRACE_CPUPROFILER_EVENT_SCOPE(AnimGraphNodeValues);
 															InNodeValuesTimeline.EnumerateEvents(InGraphStartTime, InGraphEndTime, [AnimationProvider, &DebugData](double InStartTime, double InEndTime, uint32 InDepth, const FAnimNodeValueMessage& InMessage)
 															{
 																FText Text = AnimationProvider->FormatNodeKeyValue(InMessage);
@@ -888,6 +913,8 @@ void FRewindDebugger::Tick(float DeltaTime)
 				}
 			}
 		}
+			}
+		}
 
 		// update extensions
 		IterateExtensions([DeltaTime, this](IRewindDebuggerExtension* Extension)
diff --git a/Engine/Plugins/Animation/GameplayInsights/Source/RewindDebugger/Private/RewindDebugger.h b/Engine/Plugins/Animation/GameplayInsights/Source/RewindDebugger/Private/RewindDebugger.h
--- a/Engine/Plugins/Animation/GameplayInsights/Source/RewindDebugger/Private/RewindDebugger.h
+++ b/Engine/Plugins/Animation/GameplayInsights/Source/RewindDebugger/Private/RewindDebugger.h
@@ -145,6 +145,7 @@ private:
 	bool bRecording;
 
 	float PlaybackRate;
+	double PreviousTraceTime;
 	double CurrentScrubTime;
 	TRange<double> CurrentViewRange;
 	TRange<double> CurrentTraceRange;
