Fix merge/compile errors in RewindDebugger.cpp

diff --git a/Engine/Plugins/Animation/GameplayInsights/Source/RewindDebugger/Private/RewindDebugger.cpp b/Engine/Plugins/Animation/GameplayInsights/Source/RewindDebugger/Private/RewindDebugger.cpp
--- a/Engine/Plugins/Animation/GameplayInsights/Source/RewindDebugger/Private/RewindDebugger.cpp
+++ b/Engine/Plugins/Animation/GameplayInsights/Source/RewindDebugger/Private/RewindDebugger.cpp
@@ -597,7 +597,7 @@ void FRewindDebugger::GetScrubTimeInformation(double InDebugTime, FScrubTimeInfo
 							{
 								StartEventIndex = MiddleEventIndex;
 							}
-}
+						}
 
 						// Ensure there is not frames between start and end index
 						check(EndEventIndex == StartEventIndex + 1)
@@ -610,11 +610,11 @@ void FRewindDebugger::GetScrubTimeInformation(double InDebugTime, FScrubTimeInfo
 
 						// Choose frame that is nearest to the debug time
 						if (InDebugTime - Event.ElapsedTime < NextEvent.ElapsedTime - InDebugTime)
-{
+						{
 							ScrubFrameIndex = StartEventIndex;
 						}
 						else
-	{
+						{
 							ScrubFrameIndex = EndEventIndex;
 						}
 					}
@@ -688,233 +688,234 @@ void FRewindDebugger::Tick(float DeltaTime)
 					{
 						PreviousTraceTime = CurrentTraceTime;
 						
-					const TraceServices::IFrameProvider& FrameProvider = TraceServices::ReadFrameProvider(*Session);
-					TraceServices::FFrame Frame;
-					if(FrameProvider.GetFrameFromTime(ETraceFrameType::TraceFrameType_Game, CurrentTraceTime, Frame))
-					{
-							{
-								TRACE_CPUPROFILER_EVENT_SCOPE(FRewindDebugger::Tick_UpdateActorPosition);
-						// until we have actor transforms traced out, the first skeletal mesh component transform on the target actor be used as as the actor position 
-						uint64 TargetActorId = GetTargetActorId();
-						if (TargetActorId != 0)
+						const TraceServices::IFrameProvider& FrameProvider = TraceServices::ReadFrameProvider(*Session);
+						TraceServices::FFrame Frame;
+						if(FrameProvider.GetFrameFromTime(ETraceFrameType::TraceFrameType_Game, CurrentTraceTime, Frame))
 						{
-						if(UObject* ObjectInstance = FObjectTrace::GetObjectFromId(TargetActorId))
-						{
-							if (AActor* TargetActor = Cast<AActor>(ObjectInstance))
 							{
-								TInlineComponentArray<USkeletalMeshComponent*> SkeletalMeshComponents;
-								TargetActor->GetComponents(SkeletalMeshComponents);
-
-								if (SkeletalMeshComponents.Num() > 0)
+								TRACE_CPUPROFILER_EVENT_SCOPE(FRewindDebugger::Tick_UpdateActorPosition);
+								// until we have actor transforms traced out, the first skeletal mesh component transform on the target actor be used as as the actor position 
+								uint64 TargetActorId = GetTargetActorId();
+								if (TargetActorId != 0)
 								{
-									int64 ObjectId = FObjectTrace::GetObjectId(SkeletalMeshComponents[0]);
-
-									AnimationProvider->ReadSkeletalMeshPoseTimeline(ObjectId, [this, Frame, ObjectId, AnimationProvider](const IAnimationProvider::SkeletalMeshPoseTimeline& TimelineData, bool bHasCurves)
+									if(UObject* ObjectInstance = FObjectTrace::GetObjectFromId(TargetActorId))
 									{
-										TimelineData.EnumerateEvents(Frame.StartTime, Frame.EndTime,
-											[this](double InStartTime, double InEndTime, uint32 InDepth, const FSkeletalMeshPoseMessage& PoseMessage)
+										if (AActor* TargetActor = Cast<AActor>(ObjectInstance))
+										{
+											TInlineComponentArray<USkeletalMeshComponent*> SkeletalMeshComponents;
+											TargetActor->GetComponents(SkeletalMeshComponents);
+
+											if (SkeletalMeshComponents.Num() > 0)
 											{
-												bTargetActorPositionValid = true;
-												TargetActorPosition = PoseMessage.ComponentToWorld.GetTranslation();
-										
-												return TraceServices::EEventEnumerate::Stop;
-											});
-									});
+												int64 ObjectId = FObjectTrace::GetObjectId(SkeletalMeshComponents[0]);
+
+												AnimationProvider->ReadSkeletalMeshPoseTimeline(ObjectId, [this, Frame, ObjectId, AnimationProvider](const IAnimationProvider::SkeletalMeshPoseTimeline& TimelineData, bool bHasCurves)
+												{
+													TimelineData.EnumerateEvents(Frame.StartTime, Frame.EndTime,
+														[this](double InStartTime, double InEndTime, uint32 InDepth, const FSkeletalMeshPoseMessage& PoseMessage)
+														{
+															bTargetActorPositionValid = true;
+															TargetActorPosition = PoseMessage.ComponentToWorld.GetTranslation();
+													
+															return TraceServices::EEventEnumerate::Stop;
+														});
+												});
+											}
+										}
+									}
 								}
 							}
-						}
-							}
-						
-						// update pose on all SkeletalMeshComponents:
-						// - enumerate all skeletal mesh pose timelines
-						// - check if the corresponding mesh component still exists
-						// - apply the recorded pose for the current Frame
+							
+							// update pose on all SkeletalMeshComponents:
+							// - enumerate all skeletal mesh pose timelines
+							// - check if the corresponding mesh component still exists
+							// - apply the recorded pose for the current Frame
 							{
 								TRACE_CPUPROFILER_EVENT_SCOPE(FRewindDebugger::Tick_UpdatePoses);
-						AnimationProvider->EnumerateSkeletalMeshPoseTimelines([this, &Frame, AnimationProvider, GameplayProvider](uint64 ObjectId, const IAnimationProvider::SkeletalMeshPoseTimeline& TimelineData)
-						{
-							if(UObject* ObjectInstance = FObjectTrace::GetObjectFromId(ObjectId))
-							{
-								if(USkeletalMeshComponent* MeshComponent = Cast<USkeletalMeshComponent>(ObjectInstance))
+								AnimationProvider->EnumerateSkeletalMeshPoseTimelines([this, &Frame, AnimationProvider, GameplayProvider](uint64 ObjectId, const IAnimationProvider::SkeletalMeshPoseTimeline& TimelineData)
 								{
-									AnimationProvider->ReadSkeletalMeshPoseTimeline(ObjectId, [this, &Frame, ObjectId, MeshComponent, AnimationProvider](const IAnimationProvider::SkeletalMeshPoseTimeline& TimelineData, bool bHasCurves)
+									if(UObject* ObjectInstance = FObjectTrace::GetObjectFromId(ObjectId))
 									{
-										TimelineData.EnumerateEvents(Frame.StartTime, Frame.EndTime,
-											[this, ObjectId, MeshComponent, AnimationProvider](double InStartTime, double InEndTime, uint32 InDepth, const FSkeletalMeshPoseMessage& PoseMessage)
+										if(USkeletalMeshComponent* MeshComponent = Cast<USkeletalMeshComponent>(ObjectInstance))
+										{
+											AnimationProvider->ReadSkeletalMeshPoseTimeline(ObjectId, [this, &Frame, ObjectId, MeshComponent, AnimationProvider](const IAnimationProvider::SkeletalMeshPoseTimeline& TimelineData, bool bHasCurves)
 											{
-												FTransform ComponentWorldTransform;
-												const FSkeletalMeshInfo* SkeletalMeshInfo = AnimationProvider->FindSkeletalMeshInfo(PoseMessage.MeshId);
-												AnimationProvider->GetSkeletalMeshComponentSpacePose(PoseMessage, *SkeletalMeshInfo, ComponentWorldTransform, MeshComponent->GetEditableComponentSpaceTransforms());
-												MeshComponent->ApplyEditedComponentSpaceTransforms();
+												TimelineData.EnumerateEvents(Frame.StartTime, Frame.EndTime,
+													[this, ObjectId, MeshComponent, AnimationProvider](double InStartTime, double InEndTime, uint32 InDepth, const FSkeletalMeshPoseMessage& PoseMessage)
+													{
+														FTransform ComponentWorldTransform;
+														const FSkeletalMeshInfo* SkeletalMeshInfo = AnimationProvider->FindSkeletalMeshInfo(PoseMessage.MeshId);
+														AnimationProvider->GetSkeletalMeshComponentSpacePose(PoseMessage, *SkeletalMeshInfo, ComponentWorldTransform, MeshComponent->GetEditableComponentSpaceTransforms());
+														MeshComponent->ApplyEditedComponentSpaceTransforms();
 
-												if (MeshComponentsToReset.Find(ObjectId) == nullptr)
-												{
-													FMeshComponentResetData ResetData;
-													ResetData.Component = MeshComponent;
-													ResetData.RelativeTransform = MeshComponent->GetRelativeTransform();
-													MeshComponentsToReset.Add(ObjectId, ResetData);
-												}
+														if (MeshComponentsToReset.Find(ObjectId) == nullptr)
+														{
+															FMeshComponentResetData ResetData;
+															ResetData.Component = MeshComponent;
+															ResetData.RelativeTransform = MeshComponent->GetRelativeTransform();
+															MeshComponentsToReset.Add(ObjectId, ResetData);
+														}
 
-												// todo: we need to take into account tick order requirements for attached objects here
-												MeshComponent->SetWorldTransform(ComponentWorldTransform, false, nullptr, ETeleportType::TeleportPhysics);
-												MeshComponent->SetForcedLOD(PoseMessage.LodIndex + 1);
-												return TraceServices::EEventEnumerate::Stop;
+														// todo: we need to take into account tick order requirements for attached objects here
+														MeshComponent->SetWorldTransform(ComponentWorldTransform, false, nullptr, ETeleportType::TeleportPhysics);
+														MeshComponent->SetForcedLOD(PoseMessage.LodIndex + 1);
+														return TraceServices::EEventEnumerate::Stop;
+													});
 											});
-									});
-								}
-							}
-						});
+										}
+									}
+								});
 							}
 
 							{
 								TRACE_CPUPROFILER_EVENT_SCOPE(FRewindDebugger::Tick_AnimBlueprintsDebug);
-						// Apply Animation Blueprint Debugging Data:
-						// - enumerate over all anim graph timelines
-						// - check if their instance class still exists and is the debugging target for the Animation Blueprint Editor
-						// - if it is copy that debug data into the class debug data for the blueprint debugger
-						AnimationProvider->EnumerateAnimGraphTimelines([&Frame, AnimationProvider, GameplayProvider](uint64 ObjectId, const IAnimationProvider::AnimGraphTimeline& AnimGraphTimeline)
-						{
-							if(UObject* ObjectInstance = FObjectTrace::GetObjectFromId(ObjectId))
-							{
-								if(UAnimInstance* AnimInstance = Cast<UAnimInstance>(ObjectInstance))
+								// Apply Animation Blueprint Debugging Data:
+								// - enumerate over all anim graph timelines
+								// - check if their instance class still exists and is the debugging target for the Animation Blueprint Editor
+								// - if it is copy that debug data into the class debug data for the blueprint debugger
+								AnimationProvider->EnumerateAnimGraphTimelines([&Frame, AnimationProvider, GameplayProvider](uint64 ObjectId, const IAnimationProvider::AnimGraphTimeline& AnimGraphTimeline)
 								{
-									if(UAnimBlueprintGeneratedClass* InstanceClass = Cast<UAnimBlueprintGeneratedClass>(AnimInstance->GetClass()))
+									if(UObject* ObjectInstance = FObjectTrace::GetObjectFromId(ObjectId))
 									{
-										if(UAnimBlueprint* AnimBlueprint = Cast<UAnimBlueprint>(InstanceClass->ClassGeneratedBy))
+										if(UAnimInstance* AnimInstance = Cast<UAnimInstance>(ObjectInstance))
 										{
-											if(AnimBlueprint->IsObjectBeingDebugged(AnimInstance))
+											if(UAnimBlueprintGeneratedClass* InstanceClass = Cast<UAnimBlueprintGeneratedClass>(AnimInstance->GetClass()))
 											{
+												if(UAnimBlueprint* AnimBlueprint = Cast<UAnimBlueprint>(InstanceClass->ClassGeneratedBy))
+												{
+													if(AnimBlueprint->IsObjectBeingDebugged(AnimInstance))
+													{
 														TRACE_CPUPROFILER_EVENT_SCOPE(FRewindDebugger::Tick_UpdateBlueprintDebug);
-												// update debug info for attached Animation Blueprint editors
-												uint64 Id = FObjectTrace::GetObjectId(AnimInstance);
-												const int32 NodeCount = InstanceClass->GetAnimNodeProperties().Num();
-						
-												FAnimBlueprintDebugData& DebugData = InstanceClass->GetAnimBlueprintDebugData();
+														// update debug info for attached Animation Blueprint editors
+														uint64 Id = FObjectTrace::GetObjectId(AnimInstance);
+														const int32 NodeCount = InstanceClass->GetAnimNodeProperties().Num();
+								
+														FAnimBlueprintDebugData& DebugData = InstanceClass->GetAnimBlueprintDebugData();
 														{
 															TRACE_CPUPROFILER_EVENT_SCOPE(ResetNodeVisitStates);
-												DebugData.ResetNodeVisitSites();
+															DebugData.ResetNodeVisitSites();
 														}
-					
-												AnimGraphTimeline.EnumerateEvents(Frame.StartTime, Frame.EndTime, [Id, AnimationProvider, GameplayProvider, &DebugData, NodeCount](double InGraphStartTime, double InGraphEndTime, uint32 InDepth, const FAnimGraphMessage& InMessage)
-												{
-															TRACE_CPUPROFILER_EVENT_SCOPE(AnimGraphTimelineEvent);
-															
-													// Basic verification - check node count is the same
-													// @TODO: could add some form of node hash/CRC to the class to improve this
-													if(InMessage.NodeCount == NodeCount)
-													{
-														// Check for an update phase (which contains weights)
-														if(InMessage.Phase == EAnimGraphPhase::Update)
+							
+														AnimGraphTimeline.EnumerateEvents(Frame.StartTime, Frame.EndTime, [Id, AnimationProvider, GameplayProvider, &DebugData, NodeCount](double InGraphStartTime, double InGraphEndTime, uint32 InDepth, const FAnimGraphMessage& InMessage)
 														{
-															AnimationProvider->ReadAnimNodesTimeline(Id, [InGraphStartTime, InGraphEndTime, &DebugData](const IAnimationProvider::AnimNodesTimeline& InNodesTimeline)
+															TRACE_CPUPROFILER_EVENT_SCOPE(AnimGraphTimelineEvent);
+																	
+															// Basic verification - check node count is the same
+															// @TODO: could add some form of node hash/CRC to the class to improve this
+															if(InMessage.NodeCount == NodeCount)
 															{
-																		TRACE_CPUPROFILER_EVENT_SCOPE(AnimGraphDebugNodeVisits);
-																InNodesTimeline.EnumerateEvents(InGraphStartTime, InGraphEndTime, [&DebugData](double InStartTime, double InEndTime, uint32 InDepth, const FAnimNodeMessage& InMessage)
+																// Check for an update phase (which contains weights)
+																if(InMessage.Phase == EAnimGraphPhase::Update)
 																{
-																	DebugData.RecordNodeVisit(InMessage.NodeId, InMessage.PreviousNodeId, InMessage.Weight);
-																	return TraceServices::EEventEnumerate::Continue;
-																});
-															});
-					
-															AnimationProvider->ReadStateMachinesTimeline(Id, [InGraphStartTime, InGraphEndTime, &DebugData](const IAnimationProvider::StateMachinesTimeline& InStateMachinesTimeline)
-															{
+																	AnimationProvider->ReadAnimNodesTimeline(Id, [InGraphStartTime, InGraphEndTime, &DebugData](const IAnimationProvider::AnimNodesTimeline& InNodesTimeline)
+																	{
+																		TRACE_CPUPROFILER_EVENT_SCOPE(AnimGraphDebugNodeVisits);
+																		InNodesTimeline.EnumerateEvents(InGraphStartTime, InGraphEndTime, [&DebugData](double InStartTime, double InEndTime, uint32 InDepth, const FAnimNodeMessage& InMessage)
+																		{
+																			DebugData.RecordNodeVisit(InMessage.NodeId, InMessage.PreviousNodeId, InMessage.Weight);
+																			return TraceServices::EEventEnumerate::Continue;
+																		});
+																	});
+							
+																	AnimationProvider->ReadStateMachinesTimeline(Id, [InGraphStartTime, InGraphEndTime, &DebugData](const IAnimationProvider::StateMachinesTimeline& InStateMachinesTimeline)
+																	{
 																		TRACE_CPUPROFILER_EVENT_SCOPE(AnimGraphDebugStateMachine);
-																InStateMachinesTimeline.EnumerateEvents(InGraphStartTime, InGraphEndTime, [&DebugData](double InStartTime, double InEndTime, uint32 InDepth, const FAnimStateMachineMessage& InMessage)
-																{
-																	DebugData.RecordStateData(InMessage.StateMachineIndex, InMessage.StateIndex, InMessage.StateWeight, InMessage.ElapsedTime);
-																	return TraceServices::EEventEnumerate::Continue;
-																});
-															});
-					
-															AnimationProvider->ReadAnimSequencePlayersTimeline(Id, [InGraphStartTime, InGraphEndTime, GameplayProvider, &DebugData](const IAnimationProvider::AnimSequencePlayersTimeline& InSequencePlayersTimeline)
-															{
+																		InStateMachinesTimeline.EnumerateEvents(InGraphStartTime, InGraphEndTime, [&DebugData](double InStartTime, double InEndTime, uint32 InDepth, const FAnimStateMachineMessage& InMessage)
+																		{
+																			DebugData.RecordStateData(InMessage.StateMachineIndex, InMessage.StateIndex, InMessage.StateWeight, InMessage.ElapsedTime);
+																			return TraceServices::EEventEnumerate::Continue;
+																		});
+																	});
+							
+																	AnimationProvider->ReadAnimSequencePlayersTimeline(Id, [InGraphStartTime, InGraphEndTime, GameplayProvider, &DebugData](const IAnimationProvider::AnimSequencePlayersTimeline& InSequencePlayersTimeline)
+																	{
 																		TRACE_CPUPROFILER_EVENT_SCOPE(AnimGraphDebugSequencePlayers);
-																InSequencePlayersTimeline.EnumerateEvents(InGraphStartTime, InGraphEndTime, [&DebugData](double InStartTime, double InEndTime, uint32 InDepth, const FAnimSequencePlayerMessage& InMessage)
-																{
-																	DebugData.RecordSequencePlayer(InMessage.NodeId, InMessage.Position, InMessage.Length, InMessage.FrameCounter);
-																	return TraceServices::EEventEnumerate::Continue;
-																});
-															});
-					
-															AnimationProvider->ReadAnimBlendSpacePlayersTimeline(Id, [InGraphStartTime, InGraphEndTime, GameplayProvider, &DebugData](const IAnimationProvider::BlendSpacePlayersTimeline& InBlendSpacePlayersTimeline)
-															{
+																		InSequencePlayersTimeline.EnumerateEvents(InGraphStartTime, InGraphEndTime, [&DebugData](double InStartTime, double InEndTime, uint32 InDepth, const FAnimSequencePlayerMessage& InMessage)
+																		{
+																			DebugData.RecordSequencePlayer(InMessage.NodeId, InMessage.Position, InMessage.Length, InMessage.FrameCounter);
+																			return TraceServices::EEventEnumerate::Continue;
+																		});
+																	});
+							
+																	AnimationProvider->ReadAnimBlendSpacePlayersTimeline(Id, [InGraphStartTime, InGraphEndTime, GameplayProvider, &DebugData](const IAnimationProvider::BlendSpacePlayersTimeline& InBlendSpacePlayersTimeline)
+																	{
 																		TRACE_CPUPROFILER_EVENT_SCOPE(AnimGraphBlendSpaces);
-																InBlendSpacePlayersTimeline.EnumerateEvents(InGraphStartTime, InGraphEndTime, [GameplayProvider, &DebugData](double InStartTime, double InEndTime, uint32 InDepth, const FBlendSpacePlayerMessage& InMessage)
-																{
-																	UBlendSpace* BlendSpace = nullptr;
-																	const FObjectInfo* BlendSpaceInfo = GameplayProvider->FindObjectInfo(InMessage.BlendSpaceId);
-																	if(BlendSpaceInfo)
+																		InBlendSpacePlayersTimeline.EnumerateEvents(InGraphStartTime, InGraphEndTime, [GameplayProvider, &DebugData](double InStartTime, double InEndTime, uint32 InDepth, const FBlendSpacePlayerMessage& InMessage)
+																		{
+																			UBlendSpace* BlendSpace = nullptr;
+																			const FObjectInfo* BlendSpaceInfo = GameplayProvider->FindObjectInfo(InMessage.BlendSpaceId);
+																			if(BlendSpaceInfo)
+																			{
+																				BlendSpace = TSoftObjectPtr<UBlendSpace>(FSoftObjectPath(BlendSpaceInfo->PathName)).LoadSynchronous();
+																			}
+							
+																			DebugData.RecordBlendSpacePlayer(InMessage.NodeId, BlendSpace, FVector(InMessage.PositionX, InMessage.PositionY, InMessage.PositionZ), FVector(InMessage.FilteredPositionX, InMessage.FilteredPositionY, InMessage.FilteredPositionZ));
+																			return TraceServices::EEventEnumerate::Continue;
+																		});
+																	});
+							
+																	AnimationProvider->ReadAnimSyncTimeline(Id, [InGraphStartTime, InGraphEndTime, AnimationProvider, &DebugData](const IAnimationProvider::AnimSyncTimeline& InAnimSyncTimeline)
 																	{
-																		BlendSpace = TSoftObjectPtr<UBlendSpace>(FSoftObjectPath(BlendSpaceInfo->PathName)).LoadSynchronous();
-																	}
-					
-																	DebugData.RecordBlendSpacePlayer(InMessage.NodeId, BlendSpace, FVector(InMessage.PositionX, InMessage.PositionY, InMessage.PositionZ), FVector(InMessage.FilteredPositionX, InMessage.FilteredPositionY, InMessage.FilteredPositionZ));
-																	return TraceServices::EEventEnumerate::Continue;
-																});
-															});
-					
-															AnimationProvider->ReadAnimSyncTimeline(Id, [InGraphStartTime, InGraphEndTime, AnimationProvider, &DebugData](const IAnimationProvider::AnimSyncTimeline& InAnimSyncTimeline)
-															{
 																		TRACE_CPUPROFILER_EVENT_SCOPE(AnimGraphAnimSync);
-																InAnimSyncTimeline.EnumerateEvents(InGraphStartTime, InGraphEndTime, [AnimationProvider, &DebugData](double InStartTime, double InEndTime, uint32 InDepth, const FAnimSyncMessage& InMessage)
+																		InAnimSyncTimeline.EnumerateEvents(InGraphStartTime, InGraphEndTime, [AnimationProvider, &DebugData](double InStartTime, double InEndTime, uint32 InDepth, const FAnimSyncMessage& InMessage)
+																		{
+																			const TCHAR* GroupName = AnimationProvider->GetName(InMessage.GroupNameId);
+																			if(GroupName)
+																			{
+																				DebugData.RecordNodeSync(InMessage.SourceNodeId, FName(GroupName));
+																			}
+																
+																			return TraceServices::EEventEnumerate::Continue;
+																		});
+																	});
+																}
+							
+																// Some traces come from both update and evaluate phases
+																if(InMessage.Phase == EAnimGraphPhase::Update || InMessage.Phase == EAnimGraphPhase::Evaluate)
 																{
-																	const TCHAR* GroupName = AnimationProvider->GetName(InMessage.GroupNameId);
-																	if(GroupName)
+																	AnimationProvider->ReadAnimAttributesTimeline(Id, [InGraphStartTime, InGraphEndTime, AnimationProvider, &DebugData](const IAnimationProvider::AnimAttributeTimeline& InAnimAttributeTimeline)
 																	{
-																		DebugData.RecordNodeSync(InMessage.SourceNodeId, FName(GroupName));
-																	}
-														
-																	return TraceServices::EEventEnumerate::Continue;
-																});
-															});
-														}
-					
-														// Some traces come from both update and evaluate phases
-														if(InMessage.Phase == EAnimGraphPhase::Update || InMessage.Phase == EAnimGraphPhase::Evaluate)
-														{
-															AnimationProvider->ReadAnimAttributesTimeline(Id, [InGraphStartTime, InGraphEndTime, AnimationProvider, &DebugData](const IAnimationProvider::AnimAttributeTimeline& InAnimAttributeTimeline)
-															{
 																		TRACE_CPUPROFILER_EVENT_SCOPE(AnimGraphAttributes);
-																InAnimAttributeTimeline.EnumerateEvents(InGraphStartTime, InGraphEndTime, [AnimationProvider, &DebugData](double InStartTime, double InEndTime, uint32 InDepth, const FAnimAttributeMessage& InMessage)
+																		InAnimAttributeTimeline.EnumerateEvents(InGraphStartTime, InGraphEndTime, [AnimationProvider, &DebugData](double InStartTime, double InEndTime, uint32 InDepth, const FAnimAttributeMessage& InMessage)
+																		{
+																			const TCHAR* AttributeName = AnimationProvider->GetName(InMessage.AttributeNameId);
+																			if(AttributeName)
+																			{
+																				DebugData.RecordNodeAttribute(InMessage.TargetNodeId, InMessage.SourceNodeId, FName(AttributeName));
+																			}
+																
+																			return TraceServices::EEventEnumerate::Continue;
+																		});
+																	});
+																}
+							
+																// Anim node values can come from all phases
+																AnimationProvider->ReadAnimNodeValuesTimeline(Id, [InGraphStartTime, InGraphEndTime, AnimationProvider, &DebugData](const IAnimationProvider::AnimNodeValuesTimeline& InNodeValuesTimeline)
 																{
-																	const TCHAR* AttributeName = AnimationProvider->GetName(InMessage.AttributeNameId);
-																	if(AttributeName)
+																	TRACE_CPUPROFILER_EVENT_SCOPE(AnimGraphNodeValues);
+																	InNodeValuesTimeline.EnumerateEvents(InGraphStartTime, InGraphEndTime, [AnimationProvider, &DebugData](double InStartTime, double InEndTime, uint32 InDepth, const FAnimNodeValueMessage& InMessage)
 																	{
-																		DebugData.RecordNodeAttribute(InMessage.TargetNodeId, InMessage.SourceNodeId, FName(AttributeName));
-																	}
-														
-																	return TraceServices::EEventEnumerate::Continue;
+																		FText Text = AnimationProvider->FormatNodeKeyValue(InMessage);
+																		DebugData.RecordNodeValue(InMessage.NodeId, Text.ToString());
+																		return TraceServices::EEventEnumerate::Continue;
+																	});
 																});
-															});
-														}
-					
-														// Anim node values can come from all phases
-														AnimationProvider->ReadAnimNodeValuesTimeline(Id, [InGraphStartTime, InGraphEndTime, AnimationProvider, &DebugData](const IAnimationProvider::AnimNodeValuesTimeline& InNodeValuesTimeline)
-														{
-																	TRACE_CPUPROFILER_EVENT_SCOPE(AnimGraphNodeValues);
-															InNodeValuesTimeline.EnumerateEvents(InGraphStartTime, InGraphEndTime, [AnimationProvider, &DebugData](double InStartTime, double InEndTime, uint32 InDepth, const FAnimNodeValueMessage& InMessage)
-															{
-																FText Text = AnimationProvider->FormatNodeKeyValue(InMessage);
-																DebugData.RecordNodeValue(InMessage.NodeId, Text.ToString());
-																return TraceServices::EEventEnumerate::Continue;
-															});
+															}
+															return TraceServices::EEventEnumerate::Continue;
 														});
 													}
-													return TraceServices::EEventEnumerate::Continue;
-												});
+												}
 											}
 										}
 									}
-								}
+									return TraceServices::EEventEnumerate::Continue;
+								});
 							}
-							return TraceServices::EEventEnumerate::Continue;
-						});
+						}
 					}
 				}
 			}
 		}
-			}
-		}
 
 		// update extensions
 		IterateExtensions([DeltaTime, this](IRewindDebuggerExtension* Extension)
