Fix Animation Insights window stopping Rewind Debugger from working.

diff --git a/Engine/Plugins/Animation/GameplayInsights/Source/GameplayInsights/Private/GameplayInsightsModule.cpp b/Engine/Plugins/Animation/GameplayInsights/Source/GameplayInsights/Private/GameplayInsightsModule.cpp
--- a/Engine/Plugins/Animation/GameplayInsights/Source/GameplayInsights/Private/GameplayInsightsModule.cpp
+++ b/Engine/Plugins/Animation/GameplayInsights/Source/GameplayInsights/Private/GameplayInsightsModule.cpp
@@ -419,13 +419,16 @@ bool FGameplayInsightsModule::IsObjectPropertyTraceEnabled(UObject* Object)
 
 void FGameplayInsightsModule::StartTrace()
 {
-	const bool bConnected = FTraceAuxiliary::Start(
-		FTraceAuxiliary::EConnectionType::Network,
-		TEXT("127.0.0.1"),
-		nullptr);
+	if (!bTraceStarted)
+	{
+		bTraceStarted = FTraceAuxiliary::Start(
+			FTraceAuxiliary::EConnectionType::Network,
+			TEXT("127.0.0.1"),
+			nullptr);
 	
-	IUnrealInsightsModule& UnrealInsightsModule = FModuleManager::LoadModuleChecked<IUnrealInsightsModule>("TraceInsights");
-	UnrealInsightsModule.StartAnalysisForLastLiveSession();
+		IUnrealInsightsModule& UnrealInsightsModule = FModuleManager::LoadModuleChecked<IUnrealInsightsModule>("TraceInsights");
+		UnrealInsightsModule.StartAnalysisForLastLiveSession();
+	}
 }
 
 #endif
diff --git a/Engine/Plugins/Animation/GameplayInsights/Source/GameplayInsights/Private/GameplayInsightsModule.h b/Engine/Plugins/Animation/GameplayInsights/Source/GameplayInsights/Private/GameplayInsightsModule.h
--- a/Engine/Plugins/Animation/GameplayInsights/Source/GameplayInsights/Private/GameplayInsightsModule.h
+++ b/Engine/Plugins/Animation/GameplayInsights/Source/GameplayInsights/Private/GameplayInsightsModule.h
@@ -60,4 +60,6 @@ private:
 	TSharedPtr<UE::Trace::FStoreService> StoreService;
 #endif
 	TWeakPtr<FTabManager> WeakTimingProfilerTabManager;
+	
+	bool bTraceStarted = false;
 };
