Fix for Rewind Debugger debug information not updating after you stop recording.

diff --git a/Engine/Plugins/Animation/GameplayInsights/Source/RewindDebugger/Private/RewindDebugger.cpp b/Engine/Plugins/Animation/GameplayInsights/Source/RewindDebugger/Private/RewindDebugger.cpp
--- a/Engine/Plugins/Animation/GameplayInsights/Source/RewindDebugger/Private/RewindDebugger.cpp
+++ b/Engine/Plugins/Animation/GameplayInsights/Source/RewindDebugger/Private/RewindDebugger.cpp
@@ -643,11 +643,7 @@ void FRewindDebugger::Tick(float DeltaTime)
 	TRACE_CPUPROFILER_EVENT_SCOPE(FRewindDebugger::Tick);
 	if (const TraceServices::IAnalysisSession* Session = GetAnalysisSession())
 	{
-		if (bRecording)
-		{
-			// if you select a debug target before you start recording, update component list when it becomes valid
-			RefreshDebugTracks();
-		}
+		RefreshDebugTracks();
 		
 		const IAnimationProvider* AnimationProvider = Session->ReadProvider<IAnimationProvider>("AnimationProvider");
 		const IGameplayProvider* GameplayProvider = Session->ReadProvider<IGameplayProvider>("GameplayProvider");
