Rewind Debugger Settings (Project Settings -> Plugins -> Rewind Debugger)

**New**
- RewindDebuggerSettings.cpp
- RewindDebuggerSettings.h

**Changed**
- Updated `SRewindDebugger` to use `RewindDebuggerSettings` to store flag for  Show Empty Object Tracks - Filter
- Updated `RewindDebuggerCamera` class to use `RewindDebuggerSettings` to store the Camera Mode
- Updated `RewindDebugger class` to use `RewindDebuggerSettings` to store AutoRecord setting
- Update the build.cs file to RewindDebugger code module to have DeveloperSettings as a private depedency

diff --git a/Engine/Plugins/Animation/GameplayInsights/Source/RewindDebugger/Private/RewindDebugger.cpp b/Engine/Plugins/Animation/GameplayInsights/Source/RewindDebugger/Private/RewindDebugger.cpp
--- a/Engine/Plugins/Animation/GameplayInsights/Source/RewindDebugger/Private/RewindDebugger.cpp
+++ b/Engine/Plugins/Animation/GameplayInsights/Source/RewindDebugger/Private/RewindDebugger.cpp
@@ -20,6 +20,7 @@
 #include "Animation/AnimBlueprint.h"
 #include "Animation/AnimBlueprintGeneratedClass.h"
 #include "ToolMenus.h"
+#include "RewindDebuggerSettings.h"
 
 static void IterateExtensions(TFunction<void(IRewindDebuggerExtension* Extension)> IteratorFunction)
 {
@@ -38,7 +39,6 @@ FRewindDebugger::FRewindDebugger()  :
 	ControlState(FRewindDebugger::EControlState::Pause),
 	bPIEStarted(false),
 	bPIESimulating(false),
-	bAutoRecord(false),
 	bRecording(false),
 	PlaybackRate(1),
 	PreviousTraceTime(-1),
@@ -116,7 +116,7 @@ void FRewindDebugger::OnPIEStarted(bool bSimulating)
 
 	UE::Trace::ToggleChannel(TEXT("Object"), true);
 
-	if (bAutoRecord)
+	if (ShouldAutoRecordOnPIE())
 	{
 		StartRecording();
 	}
@@ -297,6 +297,16 @@ void FRewindDebugger::StartRecording()
 	FObjectTrace::SetWorldRecordingIndex(World, RecordingIndex);
 }
 
+bool FRewindDebugger::ShouldAutoRecordOnPIE() const
+{
+	return URewindDebuggerSettings::Get().bShouldAutoRecordOnPIE;
+}
+
+void FRewindDebugger::SetShouldAutoRecordOnPIE(bool value)
+{
+	URewindDebuggerSettings::Get().bShouldAutoRecordOnPIE = value;
+}
+
 void FRewindDebugger::StopRecording()
 {
 	if (bRecording)
diff --git a/Engine/Plugins/Animation/GameplayInsights/Source/RewindDebugger/Private/RewindDebugger.h b/Engine/Plugins/Animation/GameplayInsights/Source/RewindDebugger/Private/RewindDebugger.h
--- a/Engine/Plugins/Animation/GameplayInsights/Source/RewindDebugger/Private/RewindDebugger.h
+++ b/Engine/Plugins/Animation/GameplayInsights/Source/RewindDebugger/Private/RewindDebugger.h
@@ -60,8 +60,8 @@ public:
 
 	bool CanStartRecording() const { return !IsRecording() && bPIESimulating; }
 
-	bool AutoRecord() const { return bAutoRecord; }
-	void SetAutoRecord(bool value) { bAutoRecord = value; }
+	bool ShouldAutoRecordOnPIE() const;
+	void SetShouldAutoRecordOnPIE(bool value);
 
 	// Stop recording: Stop tracing Object + Animation Data.
 	void StopRecording();
@@ -140,8 +140,7 @@ private:
 
 	bool bPIEStarted;
 	bool bPIESimulating;
-
-	bool bAutoRecord;
+	
 	bool bRecording;
 
 	float PlaybackRate;
diff --git a/Engine/Plugins/Animation/GameplayInsights/Source/RewindDebugger/Private/RewindDebuggerCamera.cpp b/Engine/Plugins/Animation/GameplayInsights/Source/RewindDebugger/Private/RewindDebuggerCamera.cpp
--- a/Engine/Plugins/Animation/GameplayInsights/Source/RewindDebugger/Private/RewindDebuggerCamera.cpp
+++ b/Engine/Plugins/Animation/GameplayInsights/Source/RewindDebugger/Private/RewindDebuggerCamera.cpp
@@ -28,13 +28,13 @@ void FRewindDebuggerCamera::Initialize()
 	Menu->AddSection("Camera Mode", LOCTEXT("Camera Mode","Camera Mode"));
 	Menu->AddMenuEntry("Camera Mode",
 			FToolMenuEntry::InitMenuEntry("CameraModeDisabled",
-				LOCTEXT("Camera Mode Dosan;ed", "Disabled"),
+				LOCTEXT("Camera Mode Disabled", "Disabled"),
 				FText(),
 				FSlateIcon(),
 				FUIAction(
-					FExecuteAction::CreateRaw(this, &FRewindDebuggerCamera::SetCameraMode, ECameraMode::Disabled),
+					FExecuteAction::CreateRaw(this, &FRewindDebuggerCamera::SetCameraMode, ERewindDebuggerCameraMode::Disabled),
 					FCanExecuteAction(),
-					FGetActionCheckState::CreateLambda([this] { return Mode == ECameraMode::Disabled ? ECheckBoxState::Checked : ECheckBoxState::Unchecked; } )),
+					FGetActionCheckState::CreateLambda([this] { return CameraMode() == ERewindDebuggerCameraMode::Disabled ? ECheckBoxState::Checked : ECheckBoxState::Unchecked; } )),
 				EUserInterfaceActionType::Check
 				)
 			);
@@ -44,9 +44,9 @@ void FRewindDebuggerCamera::Initialize()
 				LOCTEXT("Camera Mode Follow", "Follow Target Actor"),
 				FText(),
 				FSlateIcon(),
-				FUIAction( FExecuteAction::CreateRaw(this, &FRewindDebuggerCamera::SetCameraMode, ECameraMode::FollowTargetActor),
+				FUIAction( FExecuteAction::CreateRaw(this, &FRewindDebuggerCamera::SetCameraMode, ERewindDebuggerCameraMode::FollowTargetActor),
 					FCanExecuteAction(),
-					FGetActionCheckState::CreateLambda([this] { return Mode == ECameraMode::FollowTargetActor ? ECheckBoxState::Checked : ECheckBoxState::Unchecked; } )),
+					FGetActionCheckState::CreateLambda([this] { return CameraMode() == ERewindDebuggerCameraMode::FollowTargetActor ? ECheckBoxState::Checked : ECheckBoxState::Unchecked; } )),
 				EUserInterfaceActionType::Check
 				)
 			);
@@ -56,33 +56,40 @@ void FRewindDebuggerCamera::Initialize()
 				LOCTEXT("Camera Mode Recorded", "Replay Recorded Camera"),
 				FText(),
 				FSlateIcon(),
-				FUIAction( FExecuteAction::CreateRaw(this, &FRewindDebuggerCamera::SetCameraMode, ECameraMode::Replay),
+				FUIAction( FExecuteAction::CreateRaw(this, &FRewindDebuggerCamera::SetCameraMode, ERewindDebuggerCameraMode::Replay),
 					FCanExecuteAction(),
-					FGetActionCheckState::CreateLambda([this] { return Mode == ECameraMode::Replay ? ECheckBoxState::Checked : ECheckBoxState::Unchecked; } )),
+					FGetActionCheckState::CreateLambda([this] { return CameraMode() == ERewindDebuggerCameraMode::Replay ? ECheckBoxState::Checked : ECheckBoxState::Unchecked; } )),
 				EUserInterfaceActionType::Check
 				)
 			);
 }
 
-void FRewindDebuggerCamera::SetCameraMode(ECameraMode InMode)
+ERewindDebuggerCameraMode FRewindDebuggerCamera::CameraMode() const
+{
+	return URewindDebuggerSettings::Get().CameraMode;
+}
+
+void FRewindDebuggerCamera::SetCameraMode(ERewindDebuggerCameraMode InMode)
 {
 	FLevelEditorModule& LevelEditor = FModuleManager::GetModuleChecked<FLevelEditorModule>("LevelEditor");
 	SLevelViewport* LevelViewport = LevelEditor.GetFirstActiveLevelViewport().Get();
 	FLevelEditorViewportClient& LevelViewportClient = LevelViewport->GetLevelViewportClient();
 
-	if (Mode == ECameraMode::Replay && InMode != ECameraMode::Replay)
+	URewindDebuggerSettings& RewindDebuggerSettings = URewindDebuggerSettings::Get();
+	
+	if (RewindDebuggerSettings.CameraMode == ERewindDebuggerCameraMode::Replay && InMode != ERewindDebuggerCameraMode::Replay)
 	{
 		LevelViewportClient.SetActorLock(nullptr);
 	}
-	else if (Mode == ECameraMode::Replay)
+	else if (RewindDebuggerSettings.CameraMode == ERewindDebuggerCameraMode::Replay)
 	{
 		if (CameraActor.IsValid())
 		{
 			LevelViewportClient.SetActorLock(CameraActor.Get());
 		}
 	}
-
-	Mode = InMode;
+	
+	RewindDebuggerSettings.CameraMode = InMode;
 }
 
 void FRewindDebuggerCamera::Update(float DeltaTime, IRewindDebugger* RewindDebugger)
@@ -110,7 +117,7 @@ void FRewindDebuggerCamera::Update(float DeltaTime, IRewindDebugger* RewindDebug
 			// only update camera in playback or scrubbing when the time has changed (allow free movement when paused)
 			LastCameraScrubTime = CurrentTraceTime;
 
-			if (Mode == ECameraMode::FollowTargetActor)
+			if (CameraMode() == ERewindDebuggerCameraMode::FollowTargetActor)
 			{
 				// Follow Actor mode: apply position changes from the target actor to the camera
 				if (bTargetActorPositionValid)
@@ -139,7 +146,6 @@ void FRewindDebuggerCamera::Update(float DeltaTime, IRewindDebugger* RewindDebug
 									FActorSpawnParameters SpawnParameters;
 									SpawnParameters.ObjectFlags |= RF_Transient;
 									CameraActor = RewindDebugger->GetWorldToVisualize()->SpawnActor<ACameraActor>(ViewMessage.Position, ViewMessage.Rotation, SpawnParameters);
-									UCameraComponent* Camera = CameraActor->GetCameraComponent();
 									CameraActor->SetActorLabel("RewindDebuggerCamera"); 
 								}
 
@@ -154,7 +160,7 @@ void FRewindDebuggerCamera::Update(float DeltaTime, IRewindDebugger* RewindDebug
 				});
 			}
 
-			if (Mode == ECameraMode::Replay)
+			if (CameraMode() == ERewindDebuggerCameraMode::Replay)
 			{
 				if (CameraActor.IsValid())
 				{
diff --git a/Engine/Plugins/Animation/GameplayInsights/Source/RewindDebugger/Private/RewindDebuggerCamera.h b/Engine/Plugins/Animation/GameplayInsights/Source/RewindDebugger/Private/RewindDebuggerCamera.h
--- a/Engine/Plugins/Animation/GameplayInsights/Source/RewindDebugger/Private/RewindDebuggerCamera.h
+++ b/Engine/Plugins/Animation/GameplayInsights/Source/RewindDebugger/Private/RewindDebuggerCamera.h
@@ -3,6 +3,7 @@
 
 #include "CoreMinimal.h"
 #include "IRewindDebuggerExtension.h"
+#include "RewindDebuggerSettings.h"
 #include "UObject/WeakObjectPtr.h"
 
 class ACameraActor;
@@ -14,12 +15,6 @@ class ACameraActor;
 class FRewindDebuggerCamera : public IRewindDebuggerExtension
 {
 public:
-	enum class ECameraMode
-	{
-		Replay,
-		FollowTargetActor,
-		Disabled,
-	};
 
 
 	FRewindDebuggerCamera();
@@ -28,12 +23,14 @@ public:
 
 	virtual void Update(float DeltaTime, IRewindDebugger* RewindDebugger) override;
 
-	ECameraMode CameraMode();
-	void SetCameraMode(ECameraMode Mode);
+	ERewindDebuggerCameraMode CameraMode() const;
+	void SetCameraMode(ERewindDebuggerCameraMode Mode);
+	
 private:
+
+	
 	bool LastPositionValid;
 	FVector LastPosition;
-
-	ECameraMode Mode = ECameraMode::Replay;
+	
 	TWeakObjectPtr<ACameraActor> CameraActor; 
 };
diff --git a/Engine/Plugins/Animation/GameplayInsights/Source/RewindDebugger/Private/RewindDebuggerSettings.cpp b/Engine/Plugins/Animation/GameplayInsights/Source/RewindDebugger/Private/RewindDebuggerSettings.cpp
new file mode 100644
--- /dev/null
+++ b/Engine/Plugins/Animation/GameplayInsights/Source/RewindDebugger/Private/RewindDebuggerSettings.cpp
@@ -0,0 +1,39 @@
+﻿// Copyright Epic Games, Inc. All Rights Reserved.
+
+#include "RewindDebuggerSettings.h"
+#include "Misc/CoreDelegates.h"
+
+URewindDebuggerSettings::URewindDebuggerSettings() : CameraMode(ERewindDebuggerCameraMode::Replay), /*bShouldAutoDetach(false),*/ bShouldAutoRecordOnPIE(false)
+{
+	FCoreDelegates::OnPreExit.AddLambda([]()
+	{
+		Get().SaveConfig();
+	});
+}
+
+#if WITH_EDITOR
+
+FText URewindDebuggerSettings::GetSectionText() const
+{
+	return NSLOCTEXT("RewindDebugger", "RewindDebuggerSettingsName", "Rewind Debugger");
+}
+
+FText URewindDebuggerSettings::GetSectionDescription() const
+{
+	return NSLOCTEXT("RewindDebugger", "RewindDebuggerSettingsDescription", "Configure options for the Rewind Debugger.");
+}
+
+#endif
+
+FName URewindDebuggerSettings::GetCategoryName() const
+{
+	return FName(TEXT("Plugins"));
+}
+
+URewindDebuggerSettings& URewindDebuggerSettings::Get()
+{
+	URewindDebuggerSettings* MutableCDO = GetMutableDefault<URewindDebuggerSettings>();
+	check(MutableCDO != nullptr)
+	
+	return *MutableCDO;
+}
diff --git a/Engine/Plugins/Animation/GameplayInsights/Source/RewindDebugger/Private/RewindDebuggerSettings.h b/Engine/Plugins/Animation/GameplayInsights/Source/RewindDebugger/Private/RewindDebuggerSettings.h
new file mode 100644
--- /dev/null
+++ b/Engine/Plugins/Animation/GameplayInsights/Source/RewindDebugger/Private/RewindDebuggerSettings.h
@@ -0,0 +1,55 @@
+﻿// Copyright Epic Games, Inc. All Rights Reserved.
+
+#pragma once
+
+#include "CoreMinimal.h"
+#include "Engine/DeveloperSettings.h"
+#include "RewindDebuggerSettings.generated.h"
+
+class URewindDebuggerExtensionSettings;
+
+UENUM()
+enum class ERewindDebuggerCameraMode
+{
+	Replay UMETA(Tooltip="Replay Recorded Camera"),
+	FollowTargetActor UMETA(Tooltip="Follow Target Actor"),
+	Disabled UMETA(Tooltip="Disable Camera On Playback"),
+};
+
+/**
+ * Implements the settings for the Rewind Debugger.
+ */
+UCLASS(Config=Editor, meta=(DisplayName="Rewind Debugger"))
+class URewindDebuggerSettings : public UDeveloperSettings
+{
+	GENERATED_BODY()
+	
+public:
+	URewindDebuggerSettings();
+
+#if WITH_EDITOR
+	virtual FText GetSectionText() const override;
+	virtual FText GetSectionDescription() const override;
+#endif
+
+	virtual FName GetCategoryName() const override;
+	
+	/** Rewind Debugger Playback Camera Mode */
+	UPROPERTY(EditAnywhere, Config, Category = Camera)
+	ERewindDebuggerCameraMode CameraMode;
+	
+	// /** If enabled, automatically detach player controller at the start of PIE */
+	// UPROPERTY(EditAnywhere, Config, Category = Other, meta = (DisplayName = "Auto Detach Player Controller on PIE"))
+	// bool bShouldAutoDetach;
+	
+	/** If enabled, start recording information at the start of PIE */
+	UPROPERTY(EditAnywhere, Config, Category = Other)
+	bool bShouldAutoRecordOnPIE;
+
+	/** If enabled, show empty tracks on Rewind Debugger Timeline*/
+	UPROPERTY(EditAnywhere, Config, Category = Filters)
+	bool bShowEmptyObjectTracks;
+	
+	/** Get Mutable CDO of URewindDebuggerSettings */
+	static URewindDebuggerSettings & Get();
+};
diff --git a/Engine/Plugins/Animation/GameplayInsights/Source/RewindDebugger/Private/SRewindDebugger.cpp b/Engine/Plugins/Animation/GameplayInsights/Source/RewindDebugger/Private/SRewindDebugger.cpp
--- a/Engine/Plugins/Animation/GameplayInsights/Source/RewindDebugger/Private/SRewindDebugger.cpp
+++ b/Engine/Plugins/Animation/GameplayInsights/Source/RewindDebugger/Private/SRewindDebugger.cpp
@@ -25,6 +25,7 @@
 #include "Widgets/Input/SSearchBox.h"
 #include "Widgets/Layout/SExpandableArea.h"
 #include "Widgets/Layout/SScrollBox.h"
+#include "RewindDebuggerSettings.h"
 
 #define LOCTEXT_NAMESPACE "SRewindDebugger"
 
@@ -32,6 +33,7 @@ SRewindDebugger::SRewindDebugger()
 	: SCompoundWidget()
 	, ViewRange(0,10)
 	, DebugComponents(nullptr)
+	, Settings(URewindDebuggerSettings::Get())
 { 
 }
 
@@ -86,6 +88,17 @@ void SRewindDebugger::SetViewRange(TRange<double> NewRange)
 	OnViewRangeChanged.ExecuteIfBound(NewRange);
 }
 
+void SRewindDebugger::ToggleDisplayEmptyTracks()
+{
+	Settings.bShowEmptyObjectTracks = !Settings.bShowEmptyObjectTracks;
+	RefreshDebugComponents();
+}
+
+bool SRewindDebugger::ShouldDisplayEmptyTracks() const
+{
+	return Settings.bShowEmptyObjectTracks;
+}
+
 void SRewindDebugger::SetDebugTargetActor(AActor* Actor)
 {
 	DebugTargetActor.Set(Actor->GetName());
diff --git a/Engine/Plugins/Animation/GameplayInsights/Source/RewindDebugger/Private/SRewindDebugger.h b/Engine/Plugins/Animation/GameplayInsights/Source/RewindDebugger/Private/SRewindDebugger.h
--- a/Engine/Plugins/Animation/GameplayInsights/Source/RewindDebugger/Private/SRewindDebugger.h
+++ b/Engine/Plugins/Animation/GameplayInsights/Source/RewindDebugger/Private/SRewindDebugger.h
@@ -10,6 +10,7 @@
 #include "Widgets/DeclarativeSyntaxSupport.h"
 #include "Widgets/Input/SSearchBox.h"
 #include "Widgets/SCompoundWidget.h"
+#include "RewindDebuggerSettings.h"
 
 class SDockTab;
 
@@ -66,9 +67,9 @@ private:
 
 	void MainFrameCreationFinished(TSharedPtr<SWindow> InRootWindow, bool bIsNewProjectWindow);
 
-	void ToggleDisplayEmptyTracks() { bDisplayEmptyTracks = !bDisplayEmptyTracks; RefreshDebugComponents(); }
-	bool ShouldDisplayEmptyTracks() { return bDisplayEmptyTracks; }
-	
+	void ToggleDisplayEmptyTracks();
+	bool ShouldDisplayEmptyTracks() const;
+
 	// Time Slider
 	TAttribute<double> ScrubTimeAttribute;
 	TAttribute<bool> TrackScrubbingAttribute;
@@ -106,4 +107,6 @@ private:
 	
 	bool bInSelectionChanged = false;
 	bool bDisplayEmptyTracks = false;
+
+	URewindDebuggerSettings & Settings;
 };
diff --git a/Engine/Plugins/Animation/GameplayInsights/Source/RewindDebugger/RewindDebugger.Build.cs b/Engine/Plugins/Animation/GameplayInsights/Source/RewindDebugger/RewindDebugger.Build.cs
--- a/Engine/Plugins/Animation/GameplayInsights/Source/RewindDebugger/RewindDebugger.Build.cs
+++ b/Engine/Plugins/Animation/GameplayInsights/Source/RewindDebugger/RewindDebugger.Build.cs
@@ -36,7 +36,8 @@ namespace UnrealBuildTool.Rules
 				"UnrealEd",
 				"RewindDebuggerInterface",
 				"ToolWidgets",
-				"MainFrame"
+				"MainFrame",
+				"DeveloperSettings",
 			});
 		}
 	}
