Rewind Debugger tab changes:
- Default layout using Level Editor Tab Exteneder to dock main and details panels nicely
- Only allow docking those tabs in the level editor (eg not into a blueprint editor)
- Don't repeatedly reopen the details panel after it's been closed

diff --git a/Engine/Plugins/Animation/GameplayInsights/Source/RewindDebugger/Private/RewindDebugger.cpp b/Engine/Plugins/Animation/GameplayInsights/Source/RewindDebugger/Private/RewindDebugger.cpp
--- a/Engine/Plugins/Animation/GameplayInsights/Source/RewindDebugger/Private/RewindDebugger.cpp
+++ b/Engine/Plugins/Animation/GameplayInsights/Source/RewindDebugger/Private/RewindDebugger.cpp
@@ -21,6 +21,8 @@
 #include "Animation/AnimBlueprintGeneratedClass.h"
 #include "ToolMenus.h"
 #include "RewindDebuggerSettings.h"
+#include "LevelEditor.h"
+#include "RewindDebuggerModule.h"
 
 static void IterateExtensions(TFunction<void(IRewindDebuggerExtension* Extension)> IteratorFunction)
 {
@@ -46,7 +48,8 @@ FRewindDebugger::FRewindDebugger()  :
 	CurrentViewRange(0,0),
 	CurrentTraceRange(0,0),
 	RecordingIndex(0),
-	bTargetActorPositionValid(false)
+	bTargetActorPositionValid(false),
+	bIsDetailsPanelOpen(true)
 {
 	RecordingDuration.Set(0);
 
@@ -936,31 +939,43 @@ void FRewindDebugger::ComponentSelectionChanged(TSharedPtr<RewindDebugger::FRewi
 {
 	SelectedTrack = SelectedObject;
 
-	// todo: if not hidden (add a toggle button to toolbar to hide/show details)
-	TSharedPtr<SDockTab> DetailsTab = FGlobalTabmanager::Get()->TryInvokeTab(FName("RewindDebuggerDetails"));
-
-	TSharedPtr<SWidget> DetailsView;
-
-	if (SelectedObject)
+	if (bIsDetailsPanelOpen)
 	{
-		if (const TraceServices::IAnalysisSession* Session = GetAnalysisSession())
+		FLevelEditorModule& LevelEditorModule = FModuleManager::GetModuleChecked<FLevelEditorModule>("LevelEditor");
+		TSharedPtr<FTabManager> LevelEditorTabManager = LevelEditorModule.GetLevelEditorTabManager();
+		TSharedPtr<SDockTab> DetailsTab = LevelEditorTabManager->TryInvokeTab(FRewindDebuggerModule::DetailsTabName);
+
+		if (DetailsTab.IsValid())
 		{
-			DetailsView = SelectedObject->GetDetailsView();
+			UpdateDetailsPanel(DetailsTab.ToSharedRef());
 		}
 	}
-	
-	if (DetailsView)
-	{
-		DetailsTab->SetContent(DetailsView.ToSharedRef());
-	}
-	else
+}
+
+void FRewindDebugger::UpdateDetailsPanel(TSharedRef<SDockTab> DetailsTab)
+{
+	if (bIsDetailsPanelOpen)
 	{
-		static TSharedPtr<SWidget> EmptyDetails;
-		if (EmptyDetails == nullptr)
+		TSharedPtr<SWidget> DetailsView;
+
+		if (SelectedTrack)
+		{
+			DetailsView = SelectedTrack->GetDetailsView();
+		}
+
+		if (DetailsView)
 		{
-			EmptyDetails = 	SNew(SSpacer);
+			DetailsTab->SetContent(DetailsView.ToSharedRef());
+		}
+		else
+		{
+			static TSharedPtr<SWidget> EmptyDetails;
+			if (EmptyDetails == nullptr)
+			{
+				EmptyDetails = SNew(SSpacer);
+			}
+			DetailsTab->SetContent(EmptyDetails.ToSharedRef());
 		}
-		DetailsTab->SetContent(EmptyDetails.ToSharedRef());
 	}
 }
 
diff --git a/Engine/Plugins/Animation/GameplayInsights/Source/RewindDebugger/Private/RewindDebugger.h b/Engine/Plugins/Animation/GameplayInsights/Source/RewindDebugger/Private/RewindDebugger.h
--- a/Engine/Plugins/Animation/GameplayInsights/Source/RewindDebugger/Private/RewindDebugger.h
+++ b/Engine/Plugins/Animation/GameplayInsights/Source/RewindDebugger/Private/RewindDebugger.h
@@ -17,6 +17,7 @@ namespace TraceServices
 }
 
 class SWidget;
+class SDockTab;
 
 // Singleton class that handles the logic for the Rewind Debugger
 // handles:
@@ -104,6 +105,8 @@ public:
 	void ComponentSelectionChanged(TSharedPtr<RewindDebugger::FRewindDebuggerTrack> SelectedObject);
 	TSharedPtr<SWidget> BuildComponentContextMenu();
 
+	void UpdateDetailsPanel(TSharedRef<SDockTab> DetailsTab);
+
 	DECLARE_DELEGATE_OneParam( FOnTrackCursor, bool)
 	void OnTrackCursor(const FOnTrackCursor& TrackCursorCallback);
 
@@ -111,6 +114,9 @@ public:
 	TBindableProperty<float>* GetRecordingDurationProperty() { return &RecordingDuration; }
 	TBindableProperty<FString, BindingType_Out>* GetDebugTargetActorProperty() { return &DebugTargetActor; }
 
+	void SetIsDetailsPanelOpen(bool bIsOpen) { bIsDetailsPanelOpen = bIsOpen; }
+	bool IsDetailsPanelOpen(bool bIsOpen) { return bIsDetailsPanelOpen; }
+
 private:
 	void RefreshDebugComponents(TArray<TSharedPtr<RewindDebugger::FRewindDebuggerTrack>>& InTracks, TArray<TSharedPtr<FDebugObjectInfo>>& OutComponents);
 	
@@ -181,4 +187,6 @@ private:
 
 	bool bTargetActorPositionValid;
 	FVector TargetActorPosition;
+
+	bool bIsDetailsPanelOpen;
 };
diff --git a/Engine/Plugins/Animation/GameplayInsights/Source/RewindDebugger/Private/RewindDebuggerModule.cpp b/Engine/Plugins/Animation/GameplayInsights/Source/RewindDebugger/Private/RewindDebuggerModule.cpp
--- a/Engine/Plugins/Animation/GameplayInsights/Source/RewindDebugger/Private/RewindDebuggerModule.cpp
+++ b/Engine/Plugins/Animation/GameplayInsights/Source/RewindDebugger/Private/RewindDebuggerModule.cpp
@@ -19,12 +19,12 @@
 #include "AnimInstanceHelpers.h"
 #include "PropertyTraceMenu.h"
 #include "ToolMenus.h"
+#include "LevelEditor.h"
 
 #define LOCTEXT_NAMESPACE "RewindDebuggerModule"
 
-static const FName RewindDebuggerTabName("RewindDebugger");
-static const FName RewindDebuggerDetailsTabName("RewindDebuggerDetails");
-
+const FName FRewindDebuggerModule::MainTabName("RewindDebugger2");
+const FName FRewindDebuggerModule::DetailsTabName("RewindDebuggerDetails2");
 
 TSharedRef<SDockTab> FRewindDebuggerModule::SpawnRewindDebuggerDetailsTab(const FSpawnTabArgs& SpawnTabArgs)
 {
@@ -32,13 +32,17 @@ TSharedRef<SDockTab> FRewindDebuggerModule::SpawnRewindDebuggerDetailsTab(const
 	{
 		FRewindDebugger::Initialize();
 	}
+
+	FRewindDebugger* RewindDebugger = FRewindDebugger::Instance();
 	
-	const TSharedRef<SDockTab> MajorTab = SNew(SDockTab)
-		.TabRole(ETabRole::NomadTab);
+	RewindDebugger->SetIsDetailsPanelOpen(true);
+	
+	TSharedRef<SDockTab> MajorTab = SNew(SDockTab)
+		.TabRole(ETabRole::PanelTab);
 
-	RewindDebuggerDetailsWidget = SNew(SRewindDebuggerDetails, MajorTab, SpawnTabArgs.GetOwnerWindow());
+	MajorTab->SetOnTabClosed(SDockTab::FOnTabClosedCallback::CreateLambda( [](TSharedRef<SDockTab>) { FRewindDebugger::Instance()->SetIsDetailsPanelOpen(false); }));
 
-	MajorTab->SetContent(RewindDebuggerDetailsWidget.ToSharedRef());
+	RewindDebugger->UpdateDetailsPanel(MajorTab);
 
 	return MajorTab;
 }
@@ -51,7 +55,7 @@ TSharedRef<SDockTab> FRewindDebuggerModule::SpawnRewindDebuggerTab(const FSpawnT
 	}
 	
 	const TSharedRef<SDockTab> MajorTab = SNew(SDockTab)
-		.TabRole(ETabRole::NomadTab)
+		.TabRole(ETabRole::PanelTab)
 		.OnTabClosed_Lambda([this](TSharedRef<SDockTab>)
 		{
 			// clear reference to widget so it will be destroyed
@@ -149,21 +153,36 @@ void FRewindDebuggerModule::StartupModule()
 	FRewindDebuggerStyle::Initialize();
 	FRewindDebuggerCommands::Register();
 
-	FGlobalTabmanager::Get()->RegisterNomadTabSpawner(
-		RewindDebuggerTabName, FOnSpawnTab::CreateRaw(this, &FRewindDebuggerModule::SpawnRewindDebuggerTab))
-		.SetGroup(WorkspaceMenu::GetMenuStructure().GetDeveloperToolsDebugCategory())
-		.SetDisplayName(LOCTEXT("TabTitle", "Rewind Debugger"))
-		.SetIcon(FSlateIcon("RewindDebuggerStyle", "RewindDebugger.RewindIcon"))
-		.SetTooltipText(LOCTEXT("TooltipText", "Opens Rewind Debugger."));
-	
-	FGlobalTabmanager::Get()->RegisterNomadTabSpawner(
-		RewindDebuggerDetailsTabName,
-    	FOnSpawnTab::CreateRaw(this, &FRewindDebuggerModule::SpawnRewindDebuggerDetailsTab))
-		.SetGroup(WorkspaceMenu::GetMenuStructure().GetDeveloperToolsDebugCategory())
-		.SetDisplayName(LOCTEXT("DetailsTabTitle", "Rewind Debugger Details"))
-		.SetIcon(FSlateIcon("RewindDebuggerStyle", "RewindDebugger.RewindDetailsIcon"))
-		.SetTooltipText(LOCTEXT("DetailsWindowTooltipText", "Opens Rewind Debugger Details Window."));
+	FLevelEditorModule& LevelEditorModule = FModuleManager::GetModuleChecked<FLevelEditorModule>("LevelEditor");
 
+	/* LevelEditorTabManagerChangedHandle = */ LevelEditorModule.OnTabManagerChanged().AddLambda([this]()
+		{
+			FLevelEditorModule& LevelEditorModule = FModuleManager::GetModuleChecked<FLevelEditorModule>("LevelEditor");
+			TSharedPtr<FTabManager> LevelEditorTabManager = LevelEditorModule.GetLevelEditorTabManager();
+
+			LevelEditorTabManager->RegisterTabSpawner(
+				MainTabName, FOnSpawnTab::CreateRaw(this, &FRewindDebuggerModule::SpawnRewindDebuggerTab))
+				.SetGroup(WorkspaceMenu::GetMenuStructure().GetDeveloperToolsDebugCategory())
+				.SetDisplayName(LOCTEXT("TabTitle", "Rewind Debugger"))
+				.SetIcon(FSlateIcon("RewindDebuggerStyle", "RewindDebugger.RewindIcon"))
+				.SetTooltipText(LOCTEXT("TooltipText", "Opens Rewind Debugger."));
+
+			LevelEditorTabManager->RegisterTabSpawner(
+				DetailsTabName, FOnSpawnTab::CreateRaw(this, &FRewindDebuggerModule::SpawnRewindDebuggerDetailsTab))
+				.SetGroup(WorkspaceMenu::GetMenuStructure().GetDeveloperToolsDebugCategory())
+				.SetDisplayName(LOCTEXT("DetailsTabTitle", "Rewind Debugger Details"))
+				.SetIcon(FSlateIcon("RewindDebuggerStyle", "RewindDebugger.RewindDetailsIcon"))
+				.SetTooltipText(LOCTEXT("DetailsWindowTooltipText", "Opens Rewind Debugger Details Window."));
+		});
+
+	/*LevelEditorLayoutExtensionHandle = */ LevelEditorModule.OnRegisterLayoutExtensions().AddLambda(
+		[this](FLayoutExtender& Extender)
+		{
+			Extender.ExtendLayout(FName("LevelEditorSelectionDetails"), ELayoutExtensionPosition::After, FTabManager::FTab(DetailsTabName, ETabState::ClosedTab));
+			Extender.ExtendLayout(FName("Sequencer"), ELayoutExtensionPosition::After, FTabManager::FTab(MainTabName, ETabState::ClosedTab));
+		}
+	);
+	
 	RewindDebuggerCameraExtension.Initialize();
 	IModularFeatures::Get().RegisterModularFeature(IRewindDebuggerExtension::ModularFeatureName, &RewindDebuggerCameraExtension);
 	IModularFeatures::Get().RegisterModularFeature(IRewindDebuggerDoubleClickHandler::ModularFeatureName, &AnimInstanceDoubleClickHandler);
diff --git a/Engine/Plugins/Animation/GameplayInsights/Source/RewindDebugger/Private/RewindDebuggerModule.h b/Engine/Plugins/Animation/GameplayInsights/Source/RewindDebugger/Private/RewindDebuggerModule.h
--- a/Engine/Plugins/Animation/GameplayInsights/Source/RewindDebugger/Private/RewindDebuggerModule.h
+++ b/Engine/Plugins/Animation/GameplayInsights/Source/RewindDebugger/Private/RewindDebuggerModule.h
@@ -20,6 +20,9 @@ public:
 	TSharedRef<SDockTab> SpawnRewindDebuggerTab(const FSpawnTabArgs& SpawnTabArgs);
 	TSharedRef<SDockTab> SpawnRewindDebuggerDetailsTab(const FSpawnTabArgs& SpawnTabArgs);
 
+	static const FName MainTabName;
+	static const FName DetailsTabName;
+
 private:
 	TSharedPtr<SRewindDebugger> RewindDebuggerWidget;
 	TSharedPtr<SRewindDebuggerDetails> RewindDebuggerDetailsWidget;
