Fix for Rewind Debugger tab ordering in the Tools Menu. Fix for right click on empty area in component section.

diff --git a/Engine/Plugins/Animation/GameplayInsights/Source/RewindDebugger/Private/RewindDebugger.cpp b/Engine/Plugins/Animation/GameplayInsights/Source/RewindDebugger/Private/RewindDebugger.cpp
--- a/Engine/Plugins/Animation/GameplayInsights/Source/RewindDebugger/Private/RewindDebugger.cpp
+++ b/Engine/Plugins/Animation/GameplayInsights/Source/RewindDebugger/Private/RewindDebugger.cpp
@@ -1057,11 +1057,19 @@ TSharedPtr<FDebugObjectInfo> FRewindDebugger::GetSelectedComponent() const
 {
 	if (!SelectedComponent.IsValid())
 	{
-		SelectedComponent = MakeShared<FDebugObjectInfo>(0,"");
+		SelectedComponent = MakeShared<FDebugObjectInfo>(0, "");
+	}
+	
+	if (SelectedTrack.IsValid())
+	{
+		SelectedComponent->ObjectId = SelectedTrack->GetObjectId();
+		SelectedComponent->ObjectName = SelectedTrack->GetDisplayName().ToString();
+		return SelectedComponent;
+	}
+	else
+	{
+		return TSharedPtr<FDebugObjectInfo>();
 	}
-	SelectedComponent->ObjectId = SelectedTrack->GetObjectId();
-	SelectedComponent->ObjectName = SelectedTrack->GetDisplayName().ToString();
-	return SelectedComponent;
 }
 
 // build a component tree that's compatible with the public api from 5.0 for GetDebugComponents.
diff --git a/Engine/Plugins/Animation/GameplayInsights/Source/RewindDebugger/Private/RewindDebuggerModule.cpp b/Engine/Plugins/Animation/GameplayInsights/Source/RewindDebugger/Private/RewindDebuggerModule.cpp
--- a/Engine/Plugins/Animation/GameplayInsights/Source/RewindDebugger/Private/RewindDebuggerModule.cpp
+++ b/Engine/Plugins/Animation/GameplayInsights/Source/RewindDebugger/Private/RewindDebuggerModule.cpp
@@ -162,17 +162,18 @@ void FRewindDebuggerModule::StartupModule()
 
 			LevelEditorTabManager->RegisterTabSpawner(
 				MainTabName, FOnSpawnTab::CreateRaw(this, &FRewindDebuggerModule::SpawnRewindDebuggerTab))
-				.SetGroup(WorkspaceMenu::GetMenuStructure().GetDeveloperToolsDebugCategory())
 				.SetDisplayName(LOCTEXT("TabTitle", "Rewind Debugger"))
 				.SetIcon(FSlateIcon("RewindDebuggerStyle", "RewindDebugger.RewindIcon"))
-				.SetTooltipText(LOCTEXT("TooltipText", "Opens Rewind Debugger."));
+				.SetTooltipText(LOCTEXT("TooltipText", "Opens Rewind Debugger."))
+				.SetGroup(WorkspaceMenu::GetMenuStructure().GetDeveloperToolsDebugCategory());
 
 			LevelEditorTabManager->RegisterTabSpawner(
 				DetailsTabName, FOnSpawnTab::CreateRaw(this, &FRewindDebuggerModule::SpawnRewindDebuggerDetailsTab))
-				.SetGroup(WorkspaceMenu::GetMenuStructure().GetDeveloperToolsDebugCategory())
 				.SetDisplayName(LOCTEXT("DetailsTabTitle", "Rewind Debugger Details"))
 				.SetIcon(FSlateIcon("RewindDebuggerStyle", "RewindDebugger.RewindDetailsIcon"))
-				.SetTooltipText(LOCTEXT("DetailsWindowTooltipText", "Opens Rewind Debugger Details Window."));
+				.SetTooltipText(LOCTEXT("DetailsWindowTooltipText", "Opens Rewind Debugger Details Window."))
+				.SetGroup(WorkspaceMenu::GetMenuStructure().GetDeveloperToolsDebugCategory());
+			
 		});
 
 	/*LevelEditorLayoutExtensionHandle = */ LevelEditorModule.OnRegisterLayoutExtensions().AddLambda(
