Shortcuts for Rewind Debugger VCR Controls in Tool Bar

diff --git a/Engine/Plugins/Animation/GameplayInsights/Source/RewindDebugger/Private/RewindDebuggerCommands.h b/Engine/Plugins/Animation/GameplayInsights/Source/RewindDebugger/Private/RewindDebuggerCommands.h
--- a/Engine/Plugins/Animation/GameplayInsights/Source/RewindDebugger/Private/RewindDebuggerCommands.h
+++ b/Engine/Plugins/Animation/GameplayInsights/Source/RewindDebugger/Private/RewindDebuggerCommands.h
@@ -14,22 +14,23 @@ public:
 
 	/** Default constructor. */
 	FRewindDebuggerCommands()
-		: TCommands<FRewindDebuggerCommands>("RewindDebugger", NSLOCTEXT("Contexts", "RewindDebugger", "Animation Insights 2"), NAME_None, "RewindDebuggerStyle")
+		: TCommands<FRewindDebuggerCommands>(TEXT("RewindDebugger"), LOCTEXT("RewindDebugger", "Rewind Debugger"), NAME_None, "RewindDebuggerStyle")
 	{ }
 
 	// TCommands interface
 
 	virtual void RegisterCommands() override
 	{
-		UI_COMMAND(StartRecording, "Start Recording", "Start recording Animation Insights data", EUserInterfaceActionType::Button, FInputChord());
-		UI_COMMAND(StopRecording, "Stop Recording", "Stop recording Animation Insights data", EUserInterfaceActionType::Button, FInputChord());
-		UI_COMMAND(FirstFrame, "First Frame", "Jump to first recorded frame", EUserInterfaceActionType::Button, FInputChord());
-		UI_COMMAND(PreviousFrame, "Previous Frame", "Step one frame back", EUserInterfaceActionType::Button, FInputChord());
-		UI_COMMAND(ReversePlay, "Reverse Play", "Playback recorded data in reverse", EUserInterfaceActionType::Button, FInputChord());
+		UI_COMMAND(PauseOrPlay, "Pause or Play Recording", "Toggle recording playing", EUserInterfaceActionType::Button, FInputChord(EKeys::SpaceBar) );
+		UI_COMMAND(StartRecording, "Start Recording", "Start recording Animation Insights data", EUserInterfaceActionType::Button, FInputChord(EModifierKey::Control | EModifierKey::Shift, EKeys::R));
+		UI_COMMAND(StopRecording, "Stop Recording", "Stop recording Animation Insights data", EUserInterfaceActionType::Button, FInputChord(EModifierKey::Control | EModifierKey::Shift, EKeys::S));
+		UI_COMMAND(FirstFrame, "First Frame", "Jump to first recorded frame", EUserInterfaceActionType::Button, FInputChord(EKeys::Up));
+		UI_COMMAND(PreviousFrame, "Previous Frame", "Step one frame back", EUserInterfaceActionType::Button, FInputChord(EKeys::Left));
+		UI_COMMAND(ReversePlay, "Reverse Play", "Playback recorded data in reverse", EUserInterfaceActionType::Button, FInputChord(EModifierKey::Control | EModifierKey::Shift, EKeys::SpaceBar));
 		UI_COMMAND(Pause, "Pause", "Pause playback", EUserInterfaceActionType::Button, FInputChord());
-		UI_COMMAND(Play, "Play", "Playback recorded data", EUserInterfaceActionType::Button, FInputChord());
-		UI_COMMAND(NextFrame, "Next Frame", "Step one frame forward", EUserInterfaceActionType::Button, FInputChord());
-		UI_COMMAND(LastFrame, "Last Frame", "Jump to last recorded frame", EUserInterfaceActionType::Button, FInputChord());
+		UI_COMMAND(Play, "Play", "Playback recorded data", EUserInterfaceActionType::Button, FInputChord(EKeys::Down));
+		UI_COMMAND(NextFrame, "Next Frame", "Step one frame forward", EUserInterfaceActionType::Button, FInputChord(EKeys::Right));
+		UI_COMMAND(LastFrame, "Last Frame", "Jump to last recorded frame", EUserInterfaceActionType::Button, FInputChord(EModifierKey::Control, EKeys::Up));
 	}
 
 
@@ -42,6 +43,7 @@ public:
 	TSharedPtr<FUICommandInfo> Play;
 	TSharedPtr<FUICommandInfo> NextFrame;
 	TSharedPtr<FUICommandInfo> LastFrame;
+	TSharedPtr<FUICommandInfo> PauseOrPlay;
 };
 
 
diff --git a/Engine/Plugins/Animation/GameplayInsights/Source/RewindDebugger/Private/RewindDebuggerModule.cpp b/Engine/Plugins/Animation/GameplayInsights/Source/RewindDebugger/Private/RewindDebuggerModule.cpp
--- a/Engine/Plugins/Animation/GameplayInsights/Source/RewindDebugger/Private/RewindDebuggerModule.cpp
+++ b/Engine/Plugins/Animation/GameplayInsights/Source/RewindDebugger/Private/RewindDebuggerModule.cpp
@@ -20,6 +20,7 @@
 #include "PropertyTraceMenu.h"
 #include "ToolMenus.h"
 #include "LevelEditor.h"
+#include "Kismet2/DebuggerCommands.h"
 
 #define LOCTEXT_NAMESPACE "RewindDebuggerModule"
 
@@ -79,6 +80,25 @@ TSharedRef<SDockTab> FRewindDebuggerModule::SpawnRewindDebuggerTab(const FSpawnT
 							FIsActionChecked(),
 							FIsActionButtonVisible());
 
+	CommandList->MapAction(Commands.PauseOrPlay,
+							FExecuteAction::CreateLambda([DebuggerInstance]()
+							{
+								if (DebuggerInstance->CanPause())
+								{
+									DebuggerInstance->Pause();
+								}
+								else if (DebuggerInstance->CanPlay())
+								{
+									DebuggerInstance->Play();
+								}
+							}), 
+							FCanExecuteAction::CreateLambda([DebuggerInstance]()
+							{
+								return DebuggerInstance->CanPause() || DebuggerInstance->CanPlay();
+							}),
+							FIsActionChecked(),
+							FIsActionButtonVisible());
+	
 	CommandList->MapAction(Commands.ReversePlay,
 							FExecuteAction::CreateRaw(DebuggerInstance, &FRewindDebugger::PlayReverse), 
 							FCanExecuteAction::CreateRaw(DebuggerInstance, &FRewindDebugger::CanPlayReverse),
@@ -122,7 +142,24 @@ TSharedRef<SDockTab> FRewindDebuggerModule::SpawnRewindDebuggerTab(const FSpawnT
 							 FIsActionChecked(),
 							 FIsActionButtonVisible::CreateRaw(DebuggerInstance, &FRewindDebugger::CanStopRecording));
 
+	// Register PIE Rewind Debugger Commands
+	if (GEditor != nullptr)
+	{
+		check(FPlayWorldCommands::GlobalPlayWorldActions.IsValid());
 
+		FPlayWorldCommands::GlobalPlayWorldActions->MapAction(Commands.StartRecording, 
+							FExecuteAction::CreateRaw(DebuggerInstance, &FRewindDebugger::StartRecording), 
+							FCanExecuteAction::CreateRaw(DebuggerInstance, &FRewindDebugger::CanStartRecording),
+							FIsActionChecked(),
+							FIsActionButtonVisible::CreateLambda([]() { return !FRewindDebugger::Instance()->IsRecording();}));
+
+		FPlayWorldCommands::GlobalPlayWorldActions->MapAction(Commands.StopRecording,
+							 FExecuteAction::CreateRaw(DebuggerInstance, &FRewindDebugger::StopRecording),
+							 FCanExecuteAction::CreateRaw(DebuggerInstance, &FRewindDebugger::CanStopRecording),
+							 FIsActionChecked(),
+							 FIsActionButtonVisible::CreateRaw(DebuggerInstance, &FRewindDebugger::CanStopRecording));
+	}
+	
 	RewindDebuggerWidget = SNew(SRewindDebugger, CommandList.ToSharedRef(), MajorTab, SpawnTabArgs.GetOwnerWindow())
 								.DebugTargetActor(DebuggerInstance->GetDebugTargetActorProperty())
 								.RecordingDuration(DebuggerInstance->GetRecordingDurationProperty())
diff --git a/Engine/Plugins/Animation/GameplayInsights/Source/RewindDebugger/Private/SRewindDebugger.cpp b/Engine/Plugins/Animation/GameplayInsights/Source/RewindDebugger/Private/SRewindDebugger.cpp
--- a/Engine/Plugins/Animation/GameplayInsights/Source/RewindDebugger/Private/SRewindDebugger.cpp
+++ b/Engine/Plugins/Animation/GameplayInsights/Source/RewindDebugger/Private/SRewindDebugger.cpp
@@ -32,6 +32,7 @@
 SRewindDebugger::SRewindDebugger() 
 	: SCompoundWidget()
 	, ViewRange(0,10)
+	, Commands(FRewindDebuggerCommands::Get())
 	, DebugComponents(nullptr)
 	, Settings(URewindDebuggerSettings::Get())
 { 
@@ -99,6 +100,26 @@ bool SRewindDebugger::ShouldDisplayEmptyTracks() const
 	return Settings.bShowEmptyObjectTracks;
 }
 
+FReply SRewindDebugger::OnPreviewKeyDown(const FGeometry& MyGeometry, const FKeyEvent& InKeyEvent)
+{
+	const FInputChord KeyEventAsInputChord = FInputChord(InKeyEvent.GetKey(), EModifierKey::FromBools(InKeyEvent.IsControlDown(), InKeyEvent.IsAltDown(), InKeyEvent.IsShiftDown(), InKeyEvent.IsCommandDown()));
+	FReply bReply = FReply::Unhandled();
+
+	// Handle Rewind Debugger VCR Commands
+	if (CommandList->ProcessCommandBindings(InKeyEvent))
+	{
+		bReply = FReply::Handled();
+	}
+
+	// Prevent bubbling up shortcut for "ReversePlay"
+	if (Commands.ReversePlay->HasDefaultChord(KeyEventAsInputChord) && !IsPIESimulating.Get())
+	{
+		bReply = FReply::Handled();
+	}
+
+	return bReply;
+}
+
 void SRewindDebugger::SetDebugTargetActor(AActor* Actor)
 {
 	DebugTargetActor.Set(Actor->GetName());
@@ -199,7 +220,7 @@ TSharedRef<SWidget> SRewindDebugger::MakeSelectActorMenu()
 	return MenuBuilder.MakeWidget();
 }
 
-void SRewindDebugger::Construct(const FArguments& InArgs, TSharedRef<FUICommandList> CommandList, const TSharedRef<SDockTab>& ConstructUnderMajorTab, const TSharedPtr<SWindow>& ConstructUnderWindow)
+void SRewindDebugger::Construct(const FArguments& InArgs, TSharedRef<FUICommandList> InCommandList, const TSharedRef<SDockTab>& ConstructUnderMajorTab, const TSharedPtr<SWindow>& ConstructUnderWindow)
 {
 	OnScrubPositionChanged = InArgs._OnScrubPositionChanged;
 	OnViewRangeChanged = InArgs._OnViewRangeChanged;
@@ -210,7 +231,8 @@ void SRewindDebugger::Construct(const FArguments& InArgs, TSharedRef<FUICommandL
 	TraceTime.Initialize(InArgs._TraceTime);
 	RecordingDuration.Initialize(InArgs._RecordingDuration);
 	DebugTargetActor.Initialize(InArgs._DebugTargetActor);
-
+	IsPIESimulating = InArgs._IsPIESimulating;
+	CommandList = InCommandList;
 	
 	TrackFilterBox = SNew(SSearchBox).HintText(LOCTEXT("Filter Tracks","Filter Tracks")).OnTextChanged_Lambda([this](const FText&)
 	{
@@ -218,17 +240,15 @@ void SRewindDebugger::Construct(const FArguments& InArgs, TSharedRef<FUICommandL
 	});
 	
 	FSlimHorizontalToolBarBuilder ToolBarBuilder(CommandList, FMultiBoxCustomization::None, nullptr, true);
-
-	const FRewindDebuggerCommands& Commands = FRewindDebuggerCommands::Get();
-
+	
 	ToolBarBuilder.SetStyle(&FAppStyle::Get(), "PaletteToolBar");
 	ToolBarBuilder.BeginSection("Debugger");
 	{
 		ToolBarBuilder.AddToolBarButton(Commands.FirstFrame);
 		ToolBarBuilder.AddToolBarButton(Commands.PreviousFrame);
 		ToolBarBuilder.AddToolBarButton(Commands.ReversePlay);
-		ToolBarBuilder.AddToolBarButton(Commands.Pause);
-		ToolBarBuilder.AddToolBarButton(Commands.Play);
+		ToolBarBuilder.AddToolBarButton(Commands.Pause, NAME_None, {}, FText::Format(LOCTEXT("PauseButtonTooltip", "{0} ({1})"), Commands.Pause->GetDescription(), Commands.PauseOrPlay->GetInputText()));
+		ToolBarBuilder.AddToolBarButton(Commands.Play, NAME_None, {}, FText::Format(LOCTEXT("PlayButtonTooltip", "{0} ({1}) or"), Commands.Play->GetDescription(), Commands.PauseOrPlay->GetInputText(), Commands.Play->GetInputText()));
 		ToolBarBuilder.AddToolBarButton(Commands.NextFrame);
 		ToolBarBuilder.AddToolBarButton(Commands.LastFrame);
 		ToolBarBuilder.AddToolBarButton(Commands.StartRecording);
diff --git a/Engine/Plugins/Animation/GameplayInsights/Source/RewindDebugger/Private/SRewindDebugger.h b/Engine/Plugins/Animation/GameplayInsights/Source/RewindDebugger/Private/SRewindDebugger.h
--- a/Engine/Plugins/Animation/GameplayInsights/Source/RewindDebugger/Private/SRewindDebugger.h
+++ b/Engine/Plugins/Animation/GameplayInsights/Source/RewindDebugger/Private/SRewindDebugger.h
@@ -5,6 +5,7 @@
 #include "CoreMinimal.h"
 #include "Framework/Commands/UICommandList.h"
 #include "IRewindDebuggerView.h"
+#include "RewindDebuggerCommands.h"
 #include "SRewindDebuggerComponentTree.h"
 #include "SRewindDebuggerTimelines.h"
 #include "Widgets/DeclarativeSyntaxSupport.h"
@@ -12,6 +13,7 @@
 #include "Widgets/SCompoundWidget.h"
 #include "RewindDebuggerSettings.h"
 
+class FRewindDebuggerModule;
 class SDockTab;
 
 class SSearchBox;
@@ -35,6 +37,7 @@ public:
 		SLATE_ARGUMENT(TBindablePropertyInitializer<double>, TraceTime);
 		SLATE_ARGUMENT(TBindablePropertyInitializer<float>, RecordingDuration);
 		SLATE_ATTRIBUTE(double, ScrubTime);
+		SLATE_ATTRIBUTE(bool, IsPIESimulating);
 		SLATE_EVENT( FOnScrubPositionChanged, OnScrubPositionChanged);
 		SLATE_EVENT( FOnViewRangeChanged, OnViewRangeChanged);
 		SLATE_EVENT( FBuildComponentContextMenu, BuildComponentContextMenu );
@@ -69,9 +72,13 @@ private:
 
 	void ToggleDisplayEmptyTracks();
 	bool ShouldDisplayEmptyTracks() const;
-
+	
+	virtual FReply OnPreviewKeyDown(const FGeometry& MyGeometry, const FKeyEvent& InKeyEvent) override;
+	virtual bool SupportsKeyboardFocus() const override { return true; }
+	
 	// Time Slider
 	TAttribute<double> ScrubTimeAttribute;
+	TAttribute<bool> IsPIESimulating;
 	TAttribute<bool> TrackScrubbingAttribute;
 	FOnScrubPositionChanged OnScrubPositionChanged;
 	FOnViewRangeChanged OnViewRangeChanged;
@@ -79,6 +86,9 @@ private:
 	TBindableProperty<double> TraceTime;
 	TBindableProperty<float> RecordingDuration;
 
+	TSharedPtr<FUICommandList> CommandList;
+	const FRewindDebuggerCommands & Commands;
+	
 	// debug actor selector
 	TSharedRef<SWidget> MakeSelectActorMenu();
 	void SetDebugTargetActor(AActor* Actor);
