Add Pose Watch support to the Rewind Debugger
* Recording an instance of an anim blueprint with the Animation Blueprint Editor debug object attached will now trace all active Pose Watches.
* Pose Watches have separate tracks with regions shown for when that pose is evaluated in the anim graph.

diff --git a/Engine/Plugins/Animation/GameplayInsights/Source/GameplayInsights/Private/AnimationAnalyzer.cpp b/Engine/Plugins/Animation/GameplayInsights/Source/GameplayInsights/Private/AnimationAnalyzer.cpp
--- a/Engine/Plugins/Animation/GameplayInsights/Source/GameplayInsights/Private/AnimationAnalyzer.cpp
+++ b/Engine/Plugins/Animation/GameplayInsights/Source/GameplayInsights/Private/AnimationAnalyzer.cpp
@@ -48,6 +48,7 @@ void FAnimationAnalyzer::OnAnalysisBegin(const FOnAnalysisContext& Context)
 	Builder.RouteEvent(RouteId_Montage, "Animation", "Montage");
 	Builder.RouteEvent(RouteId_Montage2, "Animation", "Montage2");
 	Builder.RouteEvent(RouteId_Sync, "Animation", "Sync");
+	Builder.RouteEvent(RouteId_PoseWatch, "Animation", "PoseWatch");
 }
 
 bool FAnimationAnalyzer::OnEvent(uint16 RouteId, EStyle Style, const FOnEventContext& Context)
@@ -413,6 +414,19 @@ bool FAnimationAnalyzer::OnEvent(uint16 RouteId, EStyle Style, const FOnEventCon
 		AnimationProvider.AppendSync(AnimInstanceId, Context.EventTime.AsSeconds(Cycle), SourceNodeId, GroupNameId);
 		break;
 	}
+	case RouteId_PoseWatch:
+	{
+		uint64 Cycle = EventData.GetValue<uint64>("Cycle");
+		double RecordingTime = EventData.GetValue<double>("RecordingTime");
+		uint64 AnimInstanceId = EventData.GetValue<uint64>("AnimInstanceId");
+		uint64 PoseWatchId = EventData.GetValue<uint64>("PoseWatchId");
+		bool bIsEnabled = EventData.GetValue<bool>("bIsEnabled");
+		TArrayView<const float> WorldTransformFloatArray = EventData.GetArrayView<float>("WorldTransform");
+		TArrayView<const uint16> RequiredBonesIntArray = EventData.GetArrayView<uint16>("RequiredBones");
+		TArrayView<const float> BoneTransformsFloatArray = EventData.GetArrayView<float>("BoneTransforms");
+		AnimationProvider.AppendPoseWatch(AnimInstanceId, Context.EventTime.AsSeconds(Cycle), RecordingTime, PoseWatchId, BoneTransformsFloatArray, RequiredBonesIntArray, WorldTransformFloatArray, bIsEnabled);
+		break;
+	}
 	}
 
 	return true;
diff --git a/Engine/Plugins/Animation/GameplayInsights/Source/GameplayInsights/Private/AnimationAnalyzer.h b/Engine/Plugins/Animation/GameplayInsights/Source/GameplayInsights/Private/AnimationAnalyzer.h
--- a/Engine/Plugins/Animation/GameplayInsights/Source/GameplayInsights/Private/AnimationAnalyzer.h
+++ b/Engine/Plugins/Animation/GameplayInsights/Source/GameplayInsights/Private/AnimationAnalyzer.h
@@ -49,6 +49,7 @@ private:
 		RouteId_Montage,
 		RouteId_Montage2,
 		RouteId_Sync,
+		RouteId_PoseWatch
 	};
 
 	TraceServices::IAnalysisSession& Session;
diff --git a/Engine/Plugins/Animation/GameplayInsights/Source/GameplayInsights/Private/AnimationProvider.cpp b/Engine/Plugins/Animation/GameplayInsights/Source/GameplayInsights/Private/AnimationProvider.cpp
--- a/Engine/Plugins/Animation/GameplayInsights/Source/GameplayInsights/Private/AnimationProvider.cpp
+++ b/Engine/Plugins/Animation/GameplayInsights/Source/GameplayInsights/Private/AnimationProvider.cpp
@@ -14,6 +14,7 @@ FAnimationProvider::FAnimationProvider(TraceServices::IAnalysisSession& InSessio
 	, SkeletalMeshPoseTransforms(InSession.GetLinearAllocator(), 256)
 	, SkeletalMeshCurves(InSession.GetLinearAllocator(), 256)
 	, SkeletalMeshParentIndices(InSession.GetLinearAllocator(), 256)
+	, PoseWatchRequiredBones(InSession.GetLinearAllocator(), 256)
 	, bHasAnyData(false)
 {
 	GameplayProvider.OnObjectEndPlay().AddRaw(this, &FAnimationProvider::HandleObjectEndPlay);
@@ -76,6 +77,40 @@ void FAnimationProvider::GetSkeletalMeshComponentSpacePose(const FSkeletalMeshPo
 	}
 }
 
+void FAnimationProvider::GetPoseWatchData(const FPoseWatchMessage& InMessage, TArray<FTransform>& BoneTransforms, TArray<uint16>& RequiredBones) const
+{
+	Session.ReadAccessCheck();
+
+	// Pre-alloc arrays
+	BoneTransforms.Empty();
+	BoneTransforms.SetNumUninitialized(InMessage.NumBoneTransforms);
+
+	RequiredBones.Empty();
+	RequiredBones.SetNumUninitialized(InMessage.NumRequiredBones);
+
+	// Write into BoneTransforms
+	{
+		uint64 StartTransformIndex = InMessage.BoneTransformsStartIndex;
+		uint64 EndTransformIndex = InMessage.BoneTransformsStartIndex + InMessage.NumBoneTransforms;
+		int32 TargetTransformIndex = 0;
+		for (uint64 SourceTransformIndex = StartTransformIndex; SourceTransformIndex < EndTransformIndex; ++SourceTransformIndex, ++TargetTransformIndex)
+		{
+			BoneTransforms[TargetTransformIndex] = SkeletalMeshPoseTransforms[SourceTransformIndex];
+		}
+	}
+
+	// Write into RequiredBones
+	{
+		uint64 StartIndex = InMessage.RequiredBonesStartIndex;
+		uint64 EndIndex = InMessage.RequiredBonesStartIndex + InMessage.NumRequiredBones;
+		int32 TargetIndex = 0;
+		for (uint64 SourceIndex = StartIndex; SourceIndex < EndIndex; ++SourceIndex, ++TargetIndex)
+		{
+			RequiredBones[TargetIndex] = PoseWatchRequiredBones[SourceIndex];
+		}
+	}
+}
+
 void FAnimationProvider::EnumerateSkeletalMeshCurveIds(uint64 InObjectId, TFunctionRef<void(uint32)> Callback) const
 {
 	Session.ReadAccessCheck();
@@ -356,6 +391,23 @@ bool FAnimationProvider::ReadAnimSyncTimeline(uint64 InObjectId, TFunctionRef<vo
 	return false;
 }
 
+bool FAnimationProvider::ReadPoseWatchTimeline(uint64 InObjectId, TFunctionRef<void(const PoseWatchTimeline&)> Callback) const
+{
+	Session.ReadAccessCheck();
+
+	const uint32* IndexPtr = ObjectIdToPoseWatchTimelines.Find(InObjectId);
+	if (IndexPtr != nullptr)
+	{
+		if (*IndexPtr < uint32(PoseWatchTimelines.Num()))
+		{
+			Callback(*PoseWatchTimelines[*IndexPtr]);
+			return true;
+		}
+	}
+
+	return false;
+}
+
 const FSkeletalMeshInfo* FAnimationProvider::FindSkeletalMeshInfo(uint64 InObjectId) const
 {
 	Session.ReadAccessCheck();
@@ -1253,4 +1305,69 @@ void FAnimationProvider::AppendSync(uint64 InAnimInstanceId, double InTime, int3
 	Session.UpdateDurationSeconds(InTime);
 }
 
+void FAnimationProvider::AppendPoseWatch(uint64 InAnimInstanceId, double InTime, double InRecordingTime, uint64 PoseWatchId, const TArrayView<const float>& BoneTransformsRaw, const TArrayView<const uint16>& RequiredBones, const TArrayView<const float>& WorldTransformRaw, const bool bIsEnabled)
+{
+	Session.WriteAccessCheck();
+
+	bHasAnyData = true;
+
+	TSharedPtr<TraceServices::TPointTimeline<FPoseWatchMessage>> Timeline;
+	uint32* IndexPtr = ObjectIdToPoseWatchTimelines.Find(InAnimInstanceId);
+	if (IndexPtr != nullptr)
+	{
+		check(PoseWatchTimelines.IsValidIndex(*IndexPtr));
+		Timeline = PoseWatchTimelines[*IndexPtr];
+	}
+	else
+	{
+		Timeline = MakeShared<TraceServices::TPointTimeline<FPoseWatchMessage>>(Session.GetLinearAllocator());
+		ObjectIdToPoseWatchTimelines.Add(InAnimInstanceId, PoseWatchTimelines.Num());
+		PoseWatchTimelines.Add(Timeline.ToSharedRef());
+	}
+
+	FTransform WorldTransform;
+	FMemory::Memcpy(&WorldTransform, &WorldTransformRaw[0], sizeof(FTransform));
+
+	FPoseWatchMessage Message;
+	Message.RecordingTime = InRecordingTime;
+	Message.AnimInstanceId = InAnimInstanceId;
+	Message.PoseWatchId = PoseWatchId;
+	Message.BoneTransformsStartIndex = SkeletalMeshPoseTransforms.Num();
+	Message.NumBoneTransforms = BoneTransformsRaw.Num() / (sizeof(FTransform) / sizeof(float));
+	Message.WorldTransform = WorldTransform;
+	Message.RequiredBonesStartIndex = PoseWatchRequiredBones.Num();
+	Message.NumRequiredBones = RequiredBones.Num();
+	Message.bIsEnabled = bIsEnabled;
+
+	Timeline->AppendEvent(InTime, Message);
+
+	const int CaptureTransformSize = sizeof(FTransform) / sizeof(float); // TODO
+	const int LocalTransformSize = sizeof(FTransform) / sizeof(float);
+	const int PoseTransformCount = BoneTransformsRaw.Num() / CaptureTransformSize;
+
+	// Dump PoseBoneTransformsRaw into SkeletalMeshPoseTransforms
+	if (CaptureTransformSize == LocalTransformSize)
+	{
+		for (int i = 0; i < PoseTransformCount; i++)
+		{
+			FMemory::Memcpy(&SkeletalMeshPoseTransforms.PushBack(), &BoneTransformsRaw[CaptureTransformSize * i], sizeof(FTransform));
+		}
+	}
+	else
+	{
+		for (int i = 0; i < PoseTransformCount; i++)
+		{
+			SkeletalMeshPoseTransforms.PushBack() = ConvertTransform(CaptureTransformSize, &BoneTransformsRaw[CaptureTransformSize * i]);
+		}
+	}
+
+	// Dump RequiredBones into PoseWatchRequiredBones
+	for (const uint16 RequiredBone : RequiredBones)
+	{
+		PoseWatchRequiredBones.PushBack() = RequiredBone;
+	}
+
+	Session.UpdateDurationSeconds(InTime);
+}
+
 #undef LOCTEXT_NAMESPACE
diff --git a/Engine/Plugins/Animation/GameplayInsights/Source/GameplayInsights/Private/AnimationProvider.h b/Engine/Plugins/Animation/GameplayInsights/Source/GameplayInsights/Private/AnimationProvider.h
--- a/Engine/Plugins/Animation/GameplayInsights/Source/GameplayInsights/Private/AnimationProvider.h
+++ b/Engine/Plugins/Animation/GameplayInsights/Source/GameplayInsights/Private/AnimationProvider.h
@@ -26,6 +26,7 @@ public:
 	virtual void EnumerateSkeletalMeshPoseTimelines(TFunctionRef<void(uint64 ObjectId, const SkeletalMeshPoseTimeline&)> Callback) const override;
 	virtual bool ReadSkeletalMeshPoseTimeline(uint64 InObjectId, TFunctionRef<void(const SkeletalMeshPoseTimeline&, bool)> Callback) const override;
 	virtual void GetSkeletalMeshComponentSpacePose(const FSkeletalMeshPoseMessage& InMessage, const FSkeletalMeshInfo& InMeshInfo, FTransform& OutComponentToWorld, TArray<FTransform>& OutTransforms) const override;
+	virtual void GetPoseWatchData(const FPoseWatchMessage& InMessage, TArray<FTransform>& BoneTransforms, TArray<uint16>& RequiredBones) const;
 	virtual void EnumerateSkeletalMeshCurveIds(uint64 InObjectId, TFunctionRef<void(uint32)> Callback) const override;
 	virtual void EnumerateSkeletalMeshCurves(const FSkeletalMeshPoseMessage& InMessage, TFunctionRef<void(const FSkeletalMeshNamedCurve&)> Callback) const override;
 	virtual bool ReadTickRecordTimeline(uint64 InObjectId, TFunctionRef<void(const TickRecordTimeline&)> Callback) const override;
@@ -43,6 +44,7 @@ public:
 	virtual bool ReadMontageTimeline(uint64 InObjectId, TFunctionRef<void(const AnimMontageTimeline&)> Callback) const override;
 	virtual void EnumerateMontageIds(uint64 InObjectId, TFunctionRef<void(uint64)> Callback) const override;
 	virtual bool ReadAnimSyncTimeline(uint64 InObjectId, TFunctionRef<void(const AnimSyncTimeline&)> Callback) const override;
+	virtual bool ReadPoseWatchTimeline(uint64 InObjectId, TFunctionRef<void(const PoseWatchTimeline&)> Callback) const override;
 	virtual const FSkeletalMeshInfo* FindSkeletalMeshInfo(uint64 InObjectId) const override;
 	virtual const TCHAR* GetName(uint32 InId) const override;
 	virtual FText FormatNodeKeyValue(const FAnimNodeValueMessage& InMessage) const override;
@@ -107,6 +109,9 @@ public:
 	/** Append sync data */
 	void AppendSync(uint64 InAnimInstanceId, double InTime, int32 InSourceNodeId, uint32 InGroupNameId);
 
+	/** Append pose watch data */
+	void AppendPoseWatch(uint64 InAnimInstanceId, double InTime, double InRecordingTime, uint64 PoseWatchId, const TArrayView<const float>& BoneTransformsRaw, const TArrayView<const uint16>& RequiredBones, const TArrayView<const float>& WorldTransformRaw, const bool bIsEnabled);
+
 private:
 	/** Add anim node values helper */
 	void AppendAnimNodeValue(uint64 InAnimInstanceId, double InTime, uint16 InFrameCounter, int32 InNodeId, const TCHAR* InKey, FAnimNodeValueMessage& InMessage);
@@ -133,6 +138,7 @@ private:
 	TMap<uint64, uint32> ObjectIdToAnimMontageTimelines;
 	TMap<uint64, uint32> ObjectIdToAnimAttributeTimelines;
 	TMap<uint64, uint32> ObjectIdToAnimSyncTimelines;
+	TMap<uint64, uint32> ObjectIdToPoseWatchTimelines;
 
 	/** All the skeletal mesh info we have seen, grow only for stable indices */
 	TArray<FSkeletalMeshInfo> SkeletalMeshInfos;
@@ -185,6 +191,8 @@ private:
 	TraceServices::TPagedArray<int32> SkeletalMeshParentIndices;
 	TArray<TSharedRef<TraceServices::TPointTimeline<FAnimAttributeMessage>>> AnimAttributeTimelines;
 	TArray<TSharedRef<TraceServices::TPointTimeline<FAnimSyncMessage>>> AnimSyncTimelines;
+	TArray<TSharedRef<TraceServices::TPointTimeline<FPoseWatchMessage>>> PoseWatchTimelines;
+	TraceServices::TPagedArray<uint16> PoseWatchRequiredBones;
 
 	// Flag to indicate if any data is present
 	bool bHasAnyData;
diff --git a/Engine/Plugins/Animation/GameplayInsights/Source/GameplayInsights/Private/GameplayInsightsModule.cpp b/Engine/Plugins/Animation/GameplayInsights/Source/GameplayInsights/Private/GameplayInsightsModule.cpp
--- a/Engine/Plugins/Animation/GameplayInsights/Source/GameplayInsights/Private/GameplayInsightsModule.cpp
+++ b/Engine/Plugins/Animation/GameplayInsights/Source/GameplayInsights/Private/GameplayInsightsModule.cpp
@@ -22,6 +22,7 @@
 #include "BlendWeightsTrack.h"
 #include "MontagesTrack.h"
 #include "NotifiesTrack.h"
+#include "PoseWatchTrack.h"
 
 #if WITH_EDITOR
 #include "IAnimationBlueprintEditorModule.h"
@@ -85,6 +86,8 @@ void FGameplayInsightsModule::StartupModule()
 	IModularFeatures::Get().RegisterModularFeature(RewindDebugger::IRewindDebuggerTrackCreator::ModularFeatureName, &MontagesTrackCreator);
 	static RewindDebugger::FNotifiesTrackCreator NotifiesTrackCreator;
 	IModularFeatures::Get().RegisterModularFeature(RewindDebugger::IRewindDebuggerTrackCreator::ModularFeatureName, &NotifiesTrackCreator);
+	static RewindDebugger::FPoseWatchesTrackCreator PoseWatchesTrackCreator;
+	IModularFeatures::Get().RegisterModularFeature(RewindDebugger::IRewindDebuggerTrackCreator::ModularFeatureName, &PoseWatchesTrackCreator);
 
 
 	if (!IsRunningCommandlet())
diff --git a/Engine/Plugins/Animation/GameplayInsights/Source/GameplayInsights/Private/PoseWatchTrack.cpp b/Engine/Plugins/Animation/GameplayInsights/Source/GameplayInsights/Private/PoseWatchTrack.cpp
new file mode 100644
--- /dev/null
+++ b/Engine/Plugins/Animation/GameplayInsights/Source/GameplayInsights/Private/PoseWatchTrack.cpp
@@ -0,0 +1,284 @@
+// Copyright Epic Games, Inc. All Rights Reserved.
+
+#include "PoseWatchTrack.h"
+#include "SCurveTimelineView.h"
+#include "IRewindDebugger.h"
+#include "GameplayProvider.h"
+#include "AnimationProvider.h"
+#include "SNotifiesView.h"
+#include "Animation/AnimBlueprint.h"
+#include "Animation/AnimBlueprintGeneratedClass.h"
+#include "Animation/AnimInstance.h"
+#include "Engine/PoseWatch.h"
+#include "ObjectTrace.h"
+
+#define LOCTEXT_NAMESPACE "PoseWatchesTrack"
+
+namespace RewindDebugger
+{
+
+FPoseWatchTrack::FPoseWatchTrack(uint64 InObjectId, const FPoseWatchTrack::FPoseWatchTrackId& InPoseWatchTrackId)
+	: ObjectId(InObjectId)
+	, PostWatchTrackId(InPoseWatchTrackId)
+	, PoseWatchOwner(nullptr)
+{
+	EnabledSegments = MakeShared<SSegmentedTimelineView::FSegmentData>();
+	Icon = UPoseWatchPoseElement::StaticGetIcon();
+
+	if (UObject* ObjectInstance = FObjectTrace::GetObjectFromId(ObjectId))
+	{
+		if (UAnimInstance* AnimInstance = Cast<UAnimInstance>(ObjectInstance))
+		{
+			if (UAnimBlueprintGeneratedClass* InstanceClass = Cast<UAnimBlueprintGeneratedClass>(AnimInstance->GetClass()))
+			{
+				FAnimBlueprintDebugData& DebugData = InstanceClass->GetAnimBlueprintDebugData();
+				for (const FAnimNodePoseWatch& PoseWatch : DebugData.AnimNodePoseWatch)
+				{
+					if (PoseWatch.NodeID == PostWatchTrackId.NameId)
+					{
+						PoseWatchOwner = PoseWatch.PoseWatchPoseElement.Get();
+						break;
+					}
+				}
+			}
+		}
+	}
+}
+
+FText FPoseWatchTrack::GetDisplayNameInternal() const
+{
+	return TrackName;
+}
+
+TSharedPtr<SSegmentedTimelineView::FSegmentData> FPoseWatchTrack::GetSegmentData() const
+{
+	return EnabledSegments;
+}
+
+bool FPoseWatchTrack::UpdateInternal()
+{
+	IRewindDebugger* RewindDebugger = IRewindDebugger::Instance();
+	
+	TRange<double> TraceTimeRange = RewindDebugger->GetCurrentTraceRange();
+	double StartTime = TraceTimeRange.GetLowerBoundValue();
+	double EndTime = TraceTimeRange.GetUpperBoundValue();
+	
+	const TraceServices::IAnalysisSession* AnalysisSession = RewindDebugger->GetAnalysisSession();
+	const FGameplayProvider* GameplayProvider = AnalysisSession->ReadProvider<FGameplayProvider>(FGameplayProvider::ProviderName);
+	const FAnimationProvider* AnimationProvider = AnalysisSession->ReadProvider<FAnimationProvider>(FAnimationProvider::ProviderName);
+
+	if(GameplayProvider && AnimationProvider && EnabledSegments.IsValid())
+	{
+		EnabledSegments->Segments.SetNum(0, false);
+
+		struct FPoseWatchEnabledTime
+		{
+			double RecordingTime;
+			bool bIsEnabled;
+
+			bool operator <(const FPoseWatchEnabledTime& Other) const
+			{
+				return RecordingTime < Other.RecordingTime;
+			}
+		};
+
+		TArray<FPoseWatchEnabledTime> RecordingTimes;
+
+		TraceServices::FAnalysisSessionReadScope SessionReadScope(*AnalysisSession);
+
+		AnimationProvider->ReadPoseWatchTimeline(ObjectId, [this, StartTime, EndTime, &RecordingTimes](const FAnimationProvider::PoseWatchTimeline& InPoseWatchTimeline)
+		{
+			InPoseWatchTimeline.EnumerateEvents(StartTime, EndTime, [this, StartTime, EndTime, &RecordingTimes](double InStartTime, double InEndTime, uint32 InDepth, const FPoseWatchMessage& InMessage)
+			{
+				if (InMessage.PoseWatchId == PostWatchTrackId.NameId)
+				{
+					FPoseWatchEnabledTime EnabledTime;
+					EnabledTime.RecordingTime = InMessage.RecordingTime;
+					EnabledTime.bIsEnabled = InMessage.bIsEnabled;
+
+					RecordingTimes.Add(EnabledTime);
+				}
+				return TraceServices::EEventEnumerate::Continue;
+			});
+		});
+
+		if (!RecordingTimes.IsEmpty())
+		{
+			RecordingTimes.Sort();
+			double RangeStart = 0.0;
+
+			bool bIsEnabled = RecordingTimes[0].bIsEnabled;
+			for (int32 Index = 1; Index < RecordingTimes.Num(); ++Index)
+			{
+				// Going from disabled to enabled
+				if (!bIsEnabled && RecordingTimes[Index].bIsEnabled)
+				{
+					RangeStart = RecordingTimes[Index].RecordingTime;
+					bIsEnabled = true;
+				}
+
+				// Going from enabled to disabled
+				else if (bIsEnabled && !RecordingTimes[Index].bIsEnabled)
+				{
+					double RangeEnd = RecordingTimes[Index].RecordingTime;
+					bIsEnabled = false;
+					EnabledSegments->Segments.Add(TRange<double>(RangeStart, RangeEnd));
+				}
+			}
+
+			if (bIsEnabled)
+			{
+				// This region hasn't closed yet
+				EnabledSegments->Segments.Add(TRange<double>(RangeStart, RecordingTimes.Last().RecordingTime));
+			}
+		}
+	}
+
+	bool bChanged = false;
+	if (PoseWatchOwner && PoseWatchOwner->GetParent())
+	{
+		bChanged = !TrackName.EqualTo(PoseWatchOwner->GetParent()->GetLabel());
+		TrackName = PoseWatchOwner->GetParent()->GetLabel();
+	}
+
+	return bChanged;
+}
+
+TSharedPtr<SWidget> FPoseWatchTrack::GetTimelineViewInternal()
+{
+	const auto GetPoseWatchColorLambda = [this]() -> FLinearColor { return PoseWatchOwner ? FLinearColor(PoseWatchOwner->GetColor()) : FLinearColor::White; };
+
+	const auto TimelineView = SNew(SSegmentedTimelineView)
+		.FillColor_Lambda(GetPoseWatchColorLambda)
+		.ViewRange_Lambda([]() { return IRewindDebugger::Instance()->GetCurrentViewRange(); })
+		.SegmentData_Raw(this, &FPoseWatchTrack::GetSegmentData);
+
+	return TimelineView;
+}
+
+FName FPoseWatchesTrackCreator::GetTargetTypeNameInternal() const
+{
+	static const FName AnimInstanceName("AnimInstance");
+	return AnimInstanceName;
+}
+	
+FName FPoseWatchesTrackCreator::GetNameInternal() const
+{
+	static const FName PoseWatchesName("PoseWatches");
+	return PoseWatchesName;
+}
+
+TSharedPtr<RewindDebugger::FRewindDebuggerTrack> FPoseWatchesTrackCreator::CreateTrackInternal(uint64 ObjectId) const
+{
+	return MakeShared<RewindDebugger::FPoseWatchesTrack>(ObjectId);
+}
+
+		
+FPoseWatchesTrack::FPoseWatchesTrack(uint64 InObjectId)
+	: ObjectId(InObjectId)
+{
+	Icon = UPoseWatchPoseElement::StaticGetIcon();
+
+	if (UObject* ObjectInstance = FObjectTrace::GetObjectFromId(ObjectId))
+	{
+		if (UAnimInstance* AnimInstance = Cast<UAnimInstance>(ObjectInstance))
+		{
+			if (UAnimBlueprintGeneratedClass* InstanceClass = Cast<UAnimBlueprintGeneratedClass>(AnimInstance->GetClass()))
+			{
+				AnimBPGenClass = InstanceClass;
+			}
+		}
+	}
+}
+
+
+bool FPoseWatchesTrack::UpdateInternal()
+{
+	IRewindDebugger* RewindDebugger = IRewindDebugger::Instance();
+	
+	TRange<double> TraceTimeRange = RewindDebugger->GetCurrentTraceRange();
+	double StartTime = TraceTimeRange.GetLowerBoundValue();
+	double EndTime = TraceTimeRange.GetUpperBoundValue();
+	
+	const TraceServices::IAnalysisSession* AnalysisSession = RewindDebugger->GetAnalysisSession();
+	const FGameplayProvider* GameplayProvider = AnalysisSession->ReadProvider<FGameplayProvider>(FGameplayProvider::ProviderName);
+	const FAnimationProvider* AnimationProvider = AnalysisSession->ReadProvider<FAnimationProvider>(FAnimationProvider::ProviderName);
+	
+	bool bChanged = false;
+	
+	if(GameplayProvider && AnimationProvider)
+	{
+		TArray<FPoseWatchTrack::FPoseWatchTrackId> UniqueTrackIds;
+		
+		TraceServices::FAnalysisSessionReadScope SessionReadScope(*AnalysisSession);
+
+		AnimationProvider->ReadPoseWatchTimeline(ObjectId, [this, StartTime, EndTime, &UniqueTrackIds](const FAnimationProvider::PoseWatchTimeline& InPoseWatchTimeline)
+		{
+			InPoseWatchTimeline.EnumerateEvents(StartTime, EndTime, [this, &UniqueTrackIds](double InStartTime, double InEndTime, uint32 InDepth, const FPoseWatchMessage& InMessage)
+			{
+				// Only add the pose watch track if the pose watch still exists to attach to
+				if (AnimBPGenClass)
+				{
+					for (const FAnimNodePoseWatch& PoseWatch : AnimBPGenClass->GetAnimBlueprintDebugData().AnimNodePoseWatch)
+					{
+						if (PoseWatch.NodeID == InMessage.PoseWatchId)
+						{
+							UniqueTrackIds.AddUnique({ InMessage.PoseWatchId });
+							break;
+						}
+					}
+				}
+				return TraceServices::EEventEnumerate::Continue;
+			});
+		});
+		
+		UniqueTrackIds.StableSort();
+		const int32 TrackCount = UniqueTrackIds.Num();
+
+		if (Children.Num() != TrackCount)
+			bChanged = true;
+
+		Children.SetNum(UniqueTrackIds.Num());
+		for(int i = 0; i < TrackCount; i++)
+		{
+			if (!Children[i].IsValid() || !(Children[i].Get()->GetPoseWatchTrackId() == UniqueTrackIds[i]))
+			{
+				Children[i] = MakeShared<FPoseWatchTrack>(ObjectId, UniqueTrackIds[i]);
+				bChanged = true;
+			}
+
+			bChanged = bChanged || Children[i]->Update();
+		}
+	}
+	
+	return bChanged;
+}
+	
+void FPoseWatchesTrack::IterateSubTracksInternal(TFunction<void(TSharedPtr<FRewindDebuggerTrack> SubTrack)> IteratorFunction)
+{
+	for(TSharedPtr<FPoseWatchTrack>& Track : Children)
+	{
+		IteratorFunction(Track);
+	}
+};
+
+bool FPoseWatchesTrackCreator::HasDebugInfoInternal(uint64 ObjectId) const
+{
+	const TraceServices::IAnalysisSession* AnalysisSession = IRewindDebugger::Instance()->GetAnalysisSession();
+	
+	TraceServices::FAnalysisSessionReadScope SessionReadScope(*AnalysisSession);
+	bool bHasData = false;
+	if (const FAnimationProvider* AnimationProvider = AnalysisSession->ReadProvider<FAnimationProvider>(FAnimationProvider::ProviderName))
+	{
+		AnimationProvider->ReadPoseWatchTimeline(ObjectId, [&bHasData](const FAnimationProvider::PoseWatchTimeline& InTimeline)
+		{
+			bHasData = true;
+		});
+	}
+	return bHasData;
+}
+
+	
+}
+
+#undef LOCTEXT_NAMESPACE
\ No newline at end of file
diff --git a/Engine/Plugins/Animation/GameplayInsights/Source/GameplayInsights/Private/PoseWatchTrack.h b/Engine/Plugins/Animation/GameplayInsights/Source/GameplayInsights/Private/PoseWatchTrack.h
new file mode 100644
--- /dev/null
+++ b/Engine/Plugins/Animation/GameplayInsights/Source/GameplayInsights/Private/PoseWatchTrack.h
@@ -0,0 +1,90 @@
+// Copyright Epic Games, Inc. All Rights Reserved.
+#pragma once
+
+#include "IAnimationProvider.h"
+#include "IRewindDebuggerTrackCreator.h"
+#include "RewindDebuggerTrack.h"
+#include "Textures/SlateIcon.h"
+#include "SSegmentedTimelineView.h"
+
+class UPoseWatchPoseElement;
+class UAnimBlueprintGeneratedClass;
+
+namespace RewindDebugger
+{
+
+class FPoseWatchTrack : public FRewindDebuggerTrack
+{
+public:
+	struct FPoseWatchTrackId
+	{
+		uint64 NameId;
+
+		bool operator == (const FPoseWatchTrackId& Other) const
+		{
+			return NameId == Other.NameId;
+		}
+
+		bool operator <(const FPoseWatchTrack::FPoseWatchTrackId& Other) const
+		{
+			return NameId < Other.NameId; 
+		}		
+	};
+
+	FPoseWatchTrack(uint64 InObjectId, const FPoseWatchTrackId& TrackId);
+	const FPoseWatchTrackId& GetPoseWatchTrackId() const { return PostWatchTrackId; }
+
+private:
+	virtual bool UpdateInternal() override;
+	virtual TSharedPtr<SWidget> GetTimelineViewInternal() override;
+
+	virtual FSlateIcon GetIconInternal() override { return Icon; }
+	virtual FText GetDisplayNameInternal() const override;
+	virtual uint64 GetObjectIdInternal() const override { return ObjectId; }
+	
+	TSharedPtr<SSegmentedTimelineView::FSegmentData> GetSegmentData() const;
+	
+	uint64 ObjectId;
+	FPoseWatchTrackId PostWatchTrackId;
+	FSlateIcon Icon;
+	FText TrackName;
+
+	const UPoseWatchPoseElement* PoseWatchOwner;
+	TSharedPtr<SSegmentedTimelineView::FSegmentData> EnabledSegments;
+};
+
+
+// Parent track of all the pose watches
+class FPoseWatchesTrack : public FRewindDebuggerTrack
+{
+public:
+	FPoseWatchesTrack(uint64 InObjectId);
+
+private:
+	virtual bool UpdateInternal() override;
+	virtual void IterateSubTracksInternal(TFunction<void(TSharedPtr<FRewindDebuggerTrack> SubTrack)> IteratorFunction) override;
+
+	virtual FName GetNameInternal() const override { return "PoseWatches"; }
+	virtual FText GetDisplayNameInternal() const override { return NSLOCTEXT("RewindDebugger", "PoseWatchesTrackName", "Pose Watches"); }
+	virtual uint64 GetObjectIdInternal() const override { return ObjectId; }
+	virtual FSlateIcon GetIconInternal() override { return Icon; }
+
+	FSlateIcon Icon;
+	uint64 ObjectId;
+
+	class UAnimBlueprintGeneratedClass* AnimBPGenClass;
+
+	TArray<TSharedPtr<FPoseWatchTrack>> Children;
+};
+
+
+class FPoseWatchesTrackCreator : public IRewindDebuggerTrackCreator
+{
+private:
+	virtual FName GetTargetTypeNameInternal() const;
+	virtual FName GetNameInternal() const override;
+	virtual TSharedPtr<FRewindDebuggerTrack> CreateTrackInternal(uint64 ObjectId) const override;
+	virtual bool HasDebugInfoInternal(uint64 ObjectId) const override;
+};
+
+}
\ No newline at end of file
diff --git a/Engine/Plugins/Animation/GameplayInsights/Source/GameplayInsights/Public/IAnimationProvider.h b/Engine/Plugins/Animation/GameplayInsights/Source/GameplayInsights/Public/IAnimationProvider.h
--- a/Engine/Plugins/Animation/GameplayInsights/Source/GameplayInsights/Public/IAnimationProvider.h
+++ b/Engine/Plugins/Animation/GameplayInsights/Source/GameplayInsights/Public/IAnimationProvider.h
@@ -36,6 +36,19 @@ struct FSkeletalMeshPoseMessage
 	uint16 LodIndex = 0;
 };
 
+struct FPoseWatchMessage
+{
+	double RecordingTime = 0.0;
+	uint64 AnimInstanceId = 0;
+	uint64 PoseWatchId = 0;
+	uint64 BoneTransformsStartIndex = 0;
+	uint16 NumBoneTransforms = 0;
+	FTransform WorldTransform;
+	uint64 RequiredBonesStartIndex = 0;
+	uint16 NumRequiredBones = 0;
+	bool bIsEnabled = false;
+};
+
 struct FSkeletalMeshFrameMessage
 {
 	uint64 ComponentId = 0;
@@ -265,11 +278,13 @@ public:
 	typedef TraceServices::ITimeline<FAnimMontageMessage> AnimMontageTimeline;
 	typedef TraceServices::ITimeline<FAnimAttributeMessage> AnimAttributeTimeline;
 	typedef TraceServices::ITimeline<FAnimSyncMessage> AnimSyncTimeline;
+	typedef TraceServices::ITimeline<FPoseWatchMessage> PoseWatchTimeline;
 
 	virtual void EnumerateSkeletalMeshPoseTimelines(TFunctionRef<void(uint64 ObjectId, const SkeletalMeshPoseTimeline&)> Callback) const =0;
 	virtual bool ReadSkeletalMeshPoseTimeline(uint64 InObjectId, TFunctionRef<void(const SkeletalMeshPoseTimeline&, bool)> Callback) const = 0;
 	virtual void GetSkeletalMeshComponentSpacePose(const FSkeletalMeshPoseMessage& InMessage, const FSkeletalMeshInfo& InMeshInfo, FTransform& OutComponentToWorld, TArray<FTransform>& OutTransforms) const = 0;
 	virtual void EnumerateSkeletalMeshCurveIds(uint64 InObjectId, TFunctionRef<void(uint32)> Callback) const = 0;
+	virtual void GetPoseWatchData(const FPoseWatchMessage& InMessage, TArray<FTransform>& BoneTransforms, TArray<uint16>& RequiredBones) const = 0;
 	virtual void EnumerateSkeletalMeshCurves(const FSkeletalMeshPoseMessage& InMessage, TFunctionRef<void(const FSkeletalMeshNamedCurve&)> Callback) const = 0;
 	virtual bool ReadTickRecordTimeline(uint64 InObjectId, TFunctionRef<void(const TickRecordTimeline&)> Callback) const = 0;
 	virtual void EnumerateTickRecordIds(uint64 InObjectId, TFunctionRef<void(uint64, int32)> Callback) const = 0;
@@ -286,6 +301,7 @@ public:
 	virtual bool ReadMontageTimeline(uint64 InObjectId, TFunctionRef<void(const AnimMontageTimeline&)> Callback) const = 0;
 	virtual void EnumerateMontageIds(uint64 InObjectId, TFunctionRef<void(uint64)> Callback) const = 0;
 	virtual bool ReadAnimSyncTimeline(uint64 InObjectId, TFunctionRef<void(const AnimSyncTimeline&)> Callback) const = 0;
+	virtual bool ReadPoseWatchTimeline(uint64 InObjectId, TFunctionRef<void(const PoseWatchTimeline&)> Callback) const = 0;
 	virtual const FSkeletalMeshInfo* FindSkeletalMeshInfo(uint64 InObjectId) const = 0;
 	virtual const TCHAR* GetName(uint32 InId) const = 0;
 	virtual FText FormatNodeKeyValue(const FAnimNodeValueMessage& InMessage) const = 0;
diff --git a/Engine/Plugins/Animation/GameplayInsights/Source/RewindDebugger/Private/RewindDebugger.cpp b/Engine/Plugins/Animation/GameplayInsights/Source/RewindDebugger/Private/RewindDebugger.cpp
--- a/Engine/Plugins/Animation/GameplayInsights/Source/RewindDebugger/Private/RewindDebugger.cpp
+++ b/Engine/Plugins/Animation/GameplayInsights/Source/RewindDebugger/Private/RewindDebugger.cpp
@@ -23,6 +23,7 @@
 #include "RewindDebuggerSettings.h"
 #include "LevelEditor.h"
 #include "RewindDebuggerModule.h"
+#include "Engine/PoseWatch.h"
 
 static void IterateExtensions(TFunction<void(IRewindDebuggerExtension* Extension)> IteratorFunction)
 {
@@ -804,6 +805,14 @@ void FRewindDebugger::Tick(float DeltaTime)
 															TRACE_CPUPROFILER_EVENT_SCOPE(ResetNodeVisitStates);
 															DebugData.ResetNodeVisitSites();
 														}
+
+														for (FAnimNodePoseWatch& PoseWatchNode : DebugData.AnimNodePoseWatch)
+														{
+															if (PoseWatchNode.PoseWatch.IsValid())
+															{
+																PoseWatchNode.PoseWatch->SetIsNodeEnabled(false);
+															}
+														}
 							
 														AnimGraphTimeline.EnumerateEvents(Frame.StartTime, Frame.EndTime, [Id, AnimationProvider, GameplayProvider, &DebugData, NodeCount](double InGraphStartTime, double InGraphEndTime, uint32 InDepth, const FAnimGraphMessage& InMessage)
 														{
@@ -896,6 +905,31 @@ void FRewindDebugger::Tick(float DeltaTime)
 																			return TraceServices::EEventEnumerate::Continue;
 																		});
 																	});
+
+																	
+																	AnimationProvider->ReadPoseWatchTimeline(Id, [InGraphStartTime, InGraphEndTime, AnimationProvider, &DebugData](const IAnimationProvider::PoseWatchTimeline& InPoseWatchTimeline)
+																		{
+																			InPoseWatchTimeline.EnumerateEvents(InGraphStartTime, InGraphEndTime, [AnimationProvider, &DebugData](double InStartTime, double InEndTime, uint32 InDepth, const FPoseWatchMessage& InMessage)
+																				{
+																					for (FAnimNodePoseWatch& PoseWatch : DebugData.AnimNodePoseWatch)
+																					{
+																						if (PoseWatch.NodeID == InMessage.PoseWatchId)
+																						{
+																							TArray<FBoneIndexType> RequiredBones;
+																							TArray<FTransform> BoneTransforms;
+																							AnimationProvider->GetPoseWatchData(InMessage, BoneTransforms, RequiredBones);
+
+																							PoseWatch.SetPose(RequiredBones, BoneTransforms);
+																							PoseWatch.SetWorldTransform(InMessage.WorldTransform);
+
+																							PoseWatch.PoseWatch->SetIsNodeEnabled(true);
+																							break;
+																						}
+																					}
+																					return TraceServices::EEventEnumerate::Continue;
+																				});
+																		});
+
 																}
 							
 																// Anim node values can come from all phases
diff --git a/Engine/Plugins/Animation/GameplayInsights/Source/RewindDebugger/Private/RewindDebuggerFallbackTrack.cpp b/Engine/Plugins/Animation/GameplayInsights/Source/RewindDebugger/Private/RewindDebuggerFallbackTrack.cpp
--- a/Engine/Plugins/Animation/GameplayInsights/Source/RewindDebugger/Private/RewindDebuggerFallbackTrack.cpp
+++ b/Engine/Plugins/Animation/GameplayInsights/Source/RewindDebugger/Private/RewindDebuggerFallbackTrack.cpp
@@ -4,7 +4,6 @@
 #include "IRewindDebugger.h"
 #include "IGameplayProvider.h"
 #include "IAnimationProvider.h"
-#include "SObjectTimelineView.h"
 #include "Styling/SlateIconFinder.h"
 #include "ObjectTrace.h"
 
diff --git a/Engine/Plugins/Animation/GameplayInsights/Source/RewindDebugger/Private/RewindDebuggerObjectTrack.cpp b/Engine/Plugins/Animation/GameplayInsights/Source/RewindDebugger/Private/RewindDebuggerObjectTrack.cpp
--- a/Engine/Plugins/Animation/GameplayInsights/Source/RewindDebugger/Private/RewindDebuggerObjectTrack.cpp
+++ b/Engine/Plugins/Animation/GameplayInsights/Source/RewindDebugger/Private/RewindDebuggerObjectTrack.cpp
@@ -7,7 +7,7 @@
 #include "IAnimationProvider.h"
 #include "RewindDebuggerViewCreators.h"
 #include "RewindDebuggerTrackCreators.h"
-#include "SObjectTimelineView.h"
+#include "SSegmentedTimelineView.h"
 #include "Styling/SlateIconFinder.h"
 #include "ObjectTrace.h"
 
@@ -16,10 +16,10 @@ namespace RewindDebugger
 
 TSharedPtr<SWidget> FRewindDebuggerObjectTrack::GetTimelineViewInternal()
 {
-	return SNew(SObjectTimelineView)
+	return SNew(SSegmentedTimelineView)
 		.ViewRange_Lambda([]() { return IRewindDebugger::Instance()->GetCurrentViewRange(); })
 		.FillColor(FLinearColor(0.02f, 0.02f, 0.02f, 0.5f))
-		.ExistenceRange_Raw(this, &FRewindDebuggerObjectTrack::GetExistenceRange);
+		.SegmentData_Raw(this, &FRewindDebuggerObjectTrack::GetExistenceRange);
 
 }
 
@@ -64,7 +64,8 @@ bool FRewindDebuggerObjectTrack::UpdateInternal()
 
 	bool bChanged = false;
 
-	ExistenceRange = GameplayProvider->GetObjectRecordingLifetime(ObjectId);
+	check(ExistenceRange->Segments.Num() == 1);
+	ExistenceRange->Segments[0] = GameplayProvider->GetObjectRecordingLifetime(ObjectId);
 
 	if (!Icon.IsSet())
 	{
diff --git a/Engine/Plugins/Animation/GameplayInsights/Source/RewindDebugger/Private/RewindDebuggerObjectTrack.h b/Engine/Plugins/Animation/GameplayInsights/Source/RewindDebugger/Private/RewindDebuggerObjectTrack.h
--- a/Engine/Plugins/Animation/GameplayInsights/Source/RewindDebugger/Private/RewindDebuggerObjectTrack.h
+++ b/Engine/Plugins/Animation/GameplayInsights/Source/RewindDebugger/Private/RewindDebuggerObjectTrack.h
@@ -2,6 +2,7 @@
 #pragma once
 
 #include "RewindDebuggerTrack.h"
+#include "SSegmentedTimelineView.h"
 
 namespace RewindDebugger
 {
@@ -15,9 +16,11 @@ public:
 		, ObjectId(InObjectId)
 		, bAddController(bInAddController)
 	{
+		ExistenceRange = MakeShared<SSegmentedTimelineView::FSegmentData>();
+		ExistenceRange->Segments.SetNumUninitialized(1);
 	}
 
-	TRange<double> GetExistenceRange() const { return ExistenceRange; }
+	TSharedPtr<SSegmentedTimelineView::FSegmentData> GetExistenceRange() const { return ExistenceRange; }
 
 private:
 	virtual TSharedPtr<SWidget> GetTimelineViewInternal() override;
@@ -32,7 +35,7 @@ private:
 	
 	FString ObjectName;
 	FSlateIcon Icon;
-	TRange<double> ExistenceRange;
+	TSharedPtr<SSegmentedTimelineView::FSegmentData> ExistenceRange;
 	uint64 ObjectId;
 	TArray<TSharedPtr<FRewindDebuggerTrack>> Children;
 
diff --git a/Engine/Plugins/Animation/GameplayInsights/Source/RewindDebugger/Private/SObjectTimelineView.cpp b/Engine/Plugins/Animation/GameplayInsights/Source/RewindDebugger/Private/SObjectTimelineView.cpp
deleted file mode 100644
diff --git a/Engine/Plugins/Animation/GameplayInsights/Source/RewindDebugger/Private/SObjectTimelineView.h b/Engine/Plugins/Animation/GameplayInsights/Source/RewindDebugger/Private/SObjectTimelineView.h
deleted file mode 100644
diff --git a/Engine/Source/Developer/ToolWidgets/Private/SSegmentedTimelineView.cpp b/Engine/Source/Developer/ToolWidgets/Private/SSegmentedTimelineView.cpp
new file mode 100644
--- /dev/null
+++ b/Engine/Source/Developer/ToolWidgets/Private/SSegmentedTimelineView.cpp
@@ -0,0 +1,61 @@
+// Copyright Epic Games, Inc. All Rights Reserved.
+
+#include "SSegmentedTimelineView.h"
+#include "Rendering/DrawElements.h"
+#include "SSimpleTimeSlider.h"
+
+#define LOCTEXT_NAMESPACE "SSegmentedTimelineView"
+
+void SSegmentedTimelineView::Construct(const SSegmentedTimelineView::FArguments& InArgs)
+{
+	ViewRange = InArgs._ViewRange;
+	FillColor = InArgs._FillColor;
+	DesiredSize = InArgs._DesiredSize;
+	SegmentData = InArgs._SegmentData;
+}
+
+int32 SSegmentedTimelineView::OnPaint(const FPaintArgs& Args, const FGeometry& AllottedGeometry, const FSlateRect& MyCullingRect, FSlateWindowElementList& OutDrawElements, int32 LayerId, const FWidgetStyle& InWidgetStyle, bool bParentEnabled) const
+{
+	int32 NewLayer = PaintBlock(AllottedGeometry, MyCullingRect, OutDrawElements, LayerId, InWidgetStyle, bParentEnabled);
+
+	return FMath::Max(NewLayer, SCompoundWidget::OnPaint(Args, AllottedGeometry, MyCullingRect, OutDrawElements, NewLayer, InWidgetStyle, ShouldBeEnabled(bParentEnabled)));
+}
+
+
+FVector2D SSegmentedTimelineView::ComputeDesiredSize(float) const
+{
+	return DesiredSize.Get();
+}
+
+int32 SSegmentedTimelineView::PaintBlock(const FGeometry& AllottedGeometry, const FSlateRect& MyCullingRect, FSlateWindowElementList& OutDrawElements, int32 LayerId, const FWidgetStyle& InWidgetStyle, bool bParentEnabled) const
+{
+	// convert time range to from rewind debugger times to profiler times
+	TRange<double> DebugTimeRange = ViewRange.Get();
+	FLinearColor Color = FillColor.Get();
+
+	SSimpleTimeSlider::FScrubRangeToScreen RangeToScreen(DebugTimeRange, AllottedGeometry.GetLocalSize());
+	FVector2D Size = AllottedGeometry.GetLocalSize();
+
+	if (TSharedPtr<FSegmentData> SegmentDataPtr = SegmentData.Get())
+	{
+		for (const TRange<double>& SegmentRange : SegmentDataPtr->Segments)
+		{
+			double LowerBound = SegmentRange.HasLowerBound() ? SegmentRange.GetLowerBoundValue() : 0;
+			double UpperBound = SegmentRange.HasUpperBound() ? SegmentRange.GetUpperBoundValue() : Size.X;
+			float BoxMin = FMath::Max(0, RangeToScreen.InputToLocalX(LowerBound));
+			float BoxMax = FMath::Min(Size.X, RangeToScreen.InputToLocalX(UpperBound));
+
+			FPaintGeometry BoxGeometry = AllottedGeometry.ToPaintGeometry(FVector2D(BoxMin, 1), FVector2D(BoxMax - BoxMin, Size.Y - 2), 1);
+
+			const FSlateBrush* Brush = FAppStyle::GetBrush("Sequencer.SectionArea.Background");
+
+			FSlateDrawElement::MakeBox(OutDrawElements, LayerId, BoxGeometry, Brush, ESlateDrawEffect::None, Color);
+		}
+		
+		++LayerId;
+	}
+
+	return LayerId;
+}
+
+#undef LOCTEXT_NAMESPACE
diff --git a/Engine/Source/Developer/ToolWidgets/Public/SSegmentedTimelineView.h b/Engine/Source/Developer/ToolWidgets/Public/SSegmentedTimelineView.h
new file mode 100644
--- /dev/null
+++ b/Engine/Source/Developer/ToolWidgets/Public/SSegmentedTimelineView.h
@@ -0,0 +1,64 @@
+// Copyright Epic Games, Inc. All Rights Reserved.
+
+#pragma once
+
+#include "CoreMinimal.h"
+#include "Widgets/SCompoundWidget.h"
+#include "Widgets/DeclarativeSyntaxSupport.h"
+
+class FPaintArgs;
+class FSlateWindowElementList;
+class FTimeSliderController;
+
+
+// a timeline view which renders a block for in the timeline where the object exists
+class TOOLWIDGETS_API SSegmentedTimelineView : public SCompoundWidget
+{
+public:
+	struct FSegmentData
+	{
+	public:
+		TArray<TRange<double>> Segments;
+	};
+
+	SLATE_BEGIN_ARGS(SSegmentedTimelineView)
+		: _ViewRange(TRange<double>(0,10))
+		, _DesiredSize(FVector2D(100.f,20.f))
+	{}
+    	/** View time range */
+    	SLATE_ATTRIBUTE(TRange<double>, ViewRange);
+	
+		/** Existence Time Range */
+		SLATE_ATTRIBUTE(TSharedPtr<FSegmentData>, SegmentData);
+
+		/** Desired widget size */
+		SLATE_ATTRIBUTE(FVector2D, DesiredSize);
+	
+		/** Fill Color */
+		SLATE_ATTRIBUTE(FLinearColor, FillColor);
+     
+	SLATE_END_ARGS()
+
+
+	/**
+	 * Construct the widget
+	 * 
+	 * @param InArgs   A declaration from which to construct the widget
+	 */
+	void Construct( const FArguments& InArgs );
+
+protected:
+	// SWidget interface
+	virtual FVector2D ComputeDesiredSize(float) const override;
+	virtual int32 OnPaint( const FPaintArgs& Args, const FGeometry& AllottedGeometry, const FSlateRect& MyCullingRect, FSlateWindowElementList& OutDrawElements, int32 LayerId, const FWidgetStyle& InWidgetStyle, bool bParentEnabled ) const override;
+	
+	int32 PaintBlock(const FGeometry& AllottedGeometry, const FSlateRect& MyCullingRect, FSlateWindowElementList& OutDrawElements, int32 LayerId, const FWidgetStyle& InWidgetStyle, bool bParentEnabled ) const;
+	
+	TAttribute<TRange<double>> ViewRange;
+	TAttribute<TSharedPtr<FSegmentData>> SegmentData;
+	TAttribute<FVector2D> DesiredSize;
+	TAttribute<FLinearColor> FillColor;
+	
+	const FSlateBrush* WhiteBrush;
+};
+
diff --git a/Engine/Source/Editor/Persona/Private/AnimationEditorViewportClient.cpp b/Engine/Source/Editor/Persona/Private/AnimationEditorViewportClient.cpp
--- a/Engine/Source/Editor/Persona/Private/AnimationEditorViewportClient.cpp
+++ b/Engine/Source/Editor/Persona/Private/AnimationEditorViewportClient.cpp
@@ -1580,17 +1580,21 @@ void FAnimationViewportClient::DrawMeshBonesSourceRawAnimation(UDebugSkelMeshCom
 
 void FAnimationViewportClient::DrawWatchedPoses(UDebugSkelMeshComponent * MeshComponent, FPrimitiveDrawInterface* PDI)
 {
-	if (UAnimBlueprintGeneratedClass* AnimBlueprintGeneratedClass = Cast<UAnimBlueprintGeneratedClass>(MeshComponent->AnimClass))
+	if (UAnimBlueprintGeneratedClass* AnimBPGenClass = Cast<UAnimBlueprintGeneratedClass>(MeshComponent->AnimClass))
 	{
-		if (UAnimBlueprint* Blueprint = Cast<UAnimBlueprint>(AnimBlueprintGeneratedClass->ClassGeneratedBy))
+		USkeleton* TargetSkeleton = AnimBPGenClass->GetTargetSkeleton();
+		if (TargetSkeleton)
 		{
-			if (Blueprint->GetObjectBeingDebugged())
+			if (UAnimBlueprint* Blueprint = Cast<UAnimBlueprint>(AnimBPGenClass->ClassGeneratedBy))
 			{
-				const FAnimBlueprintDebugData& DebugData = AnimBlueprintGeneratedClass->GetAnimBlueprintDebugData();
-				DebugData.ForEachActiveVisiblePoseWatchPoseElement([MeshComponent, PDI](const FPoseWatchDebugData& PoseWatchDebugData)
+				if (Blueprint->GetObjectBeingDebugged() && MeshComponent)
 				{
-					SkeletalDebugRendering::DrawBonesFromPoseWatch(*PoseWatchDebugData.PoseInfo.Get(), MeshComponent, PDI, PoseWatchDebugData.PoseWatchPoseElement);
-				});
+					const FAnimBlueprintDebugData& DebugData = AnimBPGenClass->GetAnimBlueprintDebugData();
+					DebugData.ForEachActiveVisiblePoseWatchPoseElement([PDI, MeshComponent](const FAnimNodePoseWatch& PoseWatch)
+					{
+						SkeletalDebugRendering::DrawBonesFromPoseWatch(PDI, PoseWatch, MeshComponent->GetReferenceSkeleton(), /*bUseWorldTransform*/false);
+					});
+				}
 			}
 		}
 	}
diff --git a/Engine/Source/Editor/RewindDebuggerInterface/RewindDebuggerInterface.Build.cs b/Engine/Source/Editor/RewindDebuggerInterface/RewindDebuggerInterface.Build.cs
--- a/Engine/Source/Editor/RewindDebuggerInterface/RewindDebuggerInterface.Build.cs
+++ b/Engine/Source/Editor/RewindDebuggerInterface/RewindDebuggerInterface.Build.cs
@@ -11,6 +11,7 @@ public class RewindDebuggerInterface : ModuleRules
 				"Core",
 				"CoreUObject",
 				"SlateCore",
+				"ToolWidgets"
 			}
 		);
 	}
diff --git a/Engine/Source/Runtime/Engine/Classes/Animation/AnimBlueprintGeneratedClass.h b/Engine/Source/Runtime/Engine/Classes/Animation/AnimBlueprintGeneratedClass.h
--- a/Engine/Source/Runtime/Engine/Classes/Animation/AnimBlueprintGeneratedClass.h
+++ b/Engine/Source/Runtime/Engine/Classes/Animation/AnimBlueprintGeneratedClass.h
@@ -143,11 +143,6 @@ struct FAnimBlueprintDebugData_NodeVisit
 	}
 };
 
-struct FPoseWatchDebugData
-{
-	const UPoseWatchPoseElement* PoseWatchPoseElement;
-	TSharedPtr<FCompactHeapPose> PoseInfo;
-};
 
 struct FAnimBlueprintDebugData_AttributeRecord
 {
@@ -319,7 +314,7 @@ public:
 
 	void AddPoseWatch(int32 NodeID, UPoseWatchPoseElement* const InPoseWatchPoseElement);
 	void RemovePoseWatch(int32 NodeID);
-	void ForEachActiveVisiblePoseWatchPoseElement(const TFunctionRef<void(FPoseWatchDebugData&)>& InFunction) const;
+	void ForEachActiveVisiblePoseWatchPoseElement(const TFunctionRef<void(const FAnimNodePoseWatch&)>& InFunction) const;
 
 	TArrayView<const FName> GetNodeAttributes(TWeakObjectPtr<UAnimGraphNode_Base> InAnimGraphNode) const;
 #endif
diff --git a/Engine/Source/Runtime/Engine/Classes/Engine/PoseWatch.h b/Engine/Source/Runtime/Engine/Classes/Engine/PoseWatch.h
--- a/Engine/Source/Runtime/Engine/Classes/Engine/PoseWatch.h
+++ b/Engine/Source/Runtime/Engine/Classes/Engine/PoseWatch.h
@@ -5,21 +5,61 @@
 #include "CoreMinimal.h"
 #include "UObject/ObjectMacros.h"
 #include "UObject/Object.h"
+#include "BoneIndices.h"
 #include "PoseWatch.generated.h"
 
 
 struct FCompactHeapPose;
 class UAnimBlueprint;
 class UBlendProfile;
+struct FSlateIcon;
 
 struct FAnimNodePoseWatch
 {
+public:
 	// Object (anim instance) that this pose came from
 	TWeakObjectPtr<const UObject>			Object;
 	TWeakObjectPtr<UPoseWatch>				PoseWatch;
 	TWeakObjectPtr<UPoseWatchPoseElement>	PoseWatchPoseElement;
-	TSharedPtr<FCompactHeapPose>			PoseInfo;
 	int32									NodeID;
+
+	template<class TAllocator>
+	bool SetPose(const TArray<FBoneIndexType>& InRequiredBones, const TArray<FTransform, TAllocator>& InBoneTransforms)
+	{
+		RequiredBones = InRequiredBones;
+		BoneTransforms = InBoneTransforms;
+		return true;
+	}
+
+	void SetWorldTransform(const FTransform& InWorldTransform)
+	{
+		WorldTransform = InWorldTransform;
+	}
+
+	const TArray<FBoneIndexType>& GetRequiredBones() const
+	{
+		return RequiredBones;
+	}
+
+	const TArray<FTransform>& GetBoneTransforms() const
+	{
+		return BoneTransforms;
+	}
+
+	const FTransform& GetWorldTransform() const
+	{
+		return WorldTransform;
+	}
+
+	bool IsValid() const
+	{
+		return Object.IsValid() && PoseWatch.IsValid();
+	}
+
+private:
+	FTransform WorldTransform;
+	TArray<FBoneIndexType> RequiredBones;
+	TArray<FTransform> BoneTransforms;
 };
 
 namespace PoseWatchUtil
@@ -247,6 +287,8 @@ public:
 
 	virtual bool GetIsEnabled() const;
 
+	static FSlateIcon StaticGetIcon();
+
 #endif // WITH_EDITOR
 
 #if WITH_EDITORONLY_DATA
@@ -295,7 +337,7 @@ public:
 	bool GetIsNodeEnabled() const;
 
 	/** Returns the color to display the pose watch using */
-	UE_DEPRECATED(5.1, "Node watchs no longer have colors, use the color of its elements instead.")
+	UE_DEPRECATED(5.1, "Node watches no longer have colors, use the color of its elements instead.")
 	FColor GetColor() const;
 
 	/** Set whether this should display its children in the pose watch manager window */
diff --git a/Engine/Source/Runtime/Engine/Private/Animation/AnimBlueprintGeneratedClass.cpp b/Engine/Source/Runtime/Engine/Private/Animation/AnimBlueprintGeneratedClass.cpp
--- a/Engine/Source/Runtime/Engine/Private/Animation/AnimBlueprintGeneratedClass.cpp
+++ b/Engine/Source/Runtime/Engine/Private/Animation/AnimBlueprintGeneratedClass.cpp
@@ -247,7 +247,6 @@ void FAnimBlueprintDebugData::AddPoseWatch(int32 NodeID, UPoseWatchPoseElement*
 	NewAnimNodePoseWatch.NodeID = NodeID;
 	NewAnimNodePoseWatch.PoseWatch = InPoseWatchPoseElement->GetParent();
 	NewAnimNodePoseWatch.PoseWatchPoseElement = InPoseWatchPoseElement;
-	NewAnimNodePoseWatch.PoseInfo = MakeShareable(new FCompactHeapPose());
 	NewAnimNodePoseWatch.Object = nullptr;
 }
 
@@ -263,7 +262,7 @@ void FAnimBlueprintDebugData::RemovePoseWatch(int32 NodeID)
 	}
 }
 
-void FAnimBlueprintDebugData::ForEachActiveVisiblePoseWatchPoseElement(const TFunctionRef<void(FPoseWatchDebugData&)>& InFunction) const
+void FAnimBlueprintDebugData::ForEachActiveVisiblePoseWatchPoseElement(const TFunctionRef<void(const FAnimNodePoseWatch&)>& InFunction) const
 {
 	for (const FAnimNodePoseWatch& PoseWatchNode : AnimNodePoseWatch)
 	{
@@ -271,11 +270,7 @@ void FAnimBlueprintDebugData::ForEachActiveVisiblePoseWatchPoseElement(const TFu
 		{
 			if (PoseWatchPoseElement->GetIsEnabled() && PoseWatchPoseElement->GetIsVisible())
 			{
-				FPoseWatchDebugData DrawData;
-				DrawData.PoseWatchPoseElement = PoseWatchPoseElement;
-				DrawData.PoseInfo = PoseWatchNode.PoseInfo;
-
-				InFunction(DrawData);
+				InFunction(PoseWatchNode);
 			}
 		}
 	}
diff --git a/Engine/Source/Runtime/Engine/Private/Animation/AnimInstanceProxy.cpp b/Engine/Source/Runtime/Engine/Private/Animation/AnimInstanceProxy.cpp
--- a/Engine/Source/Runtime/Engine/Private/Animation/AnimInstanceProxy.cpp
+++ b/Engine/Source/Runtime/Engine/Private/Animation/AnimInstanceProxy.cpp
@@ -27,6 +27,7 @@
 #include "Animation/MirrorDataTable.h"
 #include "Animation/AnimBlueprintGeneratedClass.h"
 #include "Animation/AnimStateMachineTypes.h"
+#include "Animation/AnimTrace.h"
 
 #define DO_ANIMSTAT_PROCESSING(StatName) DEFINE_STAT(STAT_ ## StatName)
 #include "Animation/AnimMTStats.h"
@@ -579,6 +580,24 @@ void FAnimInstanceProxy::PostUpdate(UAnimInstance* InAnimInstance) const
 	}
 #endif
 
+#if WITH_EDITORONLY_DATA
+	UAnimBlueprint* AnimBP = GetAnimBlueprint();
+	if (bIsBeingDebugged)
+	{
+		FAnimBlueprintDebugData* DebugData = AnimBP->GetDebugData();
+		if (DebugData)
+		{
+			for (const FAnimNodePoseWatch& PoseWatch : PoseWatchEntriesForThisFrame)
+			{
+				if (!PoseWatch.PoseWatchPoseElement->GetIsEnabled())
+				{
+					TRACE_ANIM_POSE_WATCH(*this, PoseWatch.NodeID, TArray<FTransform>(), TArray<uint16>(), FTransform::Identity, false);
+				}
+			}
+		}
+	}
+#endif
+
 	InAnimInstance->NotifyQueue.Append(NotifyQueue);
 	InAnimInstance->NotifyQueue.ApplyMontageNotifies(*this);
 
@@ -3142,16 +3161,25 @@ void FAnimInstanceProxy::RecordNodeAttribute(const FAnimInstanceProxy& InSourceP
 
 void FAnimInstanceProxy::RegisterWatchedPose(const FCompactPose& Pose, int32 LinkID)
 {
-	if(bIsBeingDebugged)
+	if (bIsBeingDebugged)
 	{
-		for (FAnimNodePoseWatch& PoseWatch : PoseWatchEntriesForThisFrame)
+		if (USkeletalMeshComponent* SkelMeshComponent = GetSkelMeshComponent())
 		{
-			if (PoseWatch.NodeID == LinkID)
+			for (FAnimNodePoseWatch& PoseWatch : PoseWatchEntriesForThisFrame)
 			{
-				PoseWatch.PoseInfo->CopyBonesFrom(Pose);
-				PoseWatch.Object = GetAnimInstanceObject();
-				PoseWatch.PoseWatch->SetIsNodeEnabled(true);
-				break;
+				if (PoseWatch.PoseWatch.IsValid() && PoseWatch.NodeID == LinkID)
+				{
+					PoseWatch.Object = GetAnimInstanceObject();
+					PoseWatch.PoseWatch->SetIsNodeEnabled(true);
+
+					const TArray<FTransform, FAnimStackAllocator>& BoneTransforms = Pose.GetBones();
+					const TArray<FBoneIndexType>& TmpRequiredBones = Pose.GetBoneContainer().GetBoneIndicesArray();
+					PoseWatch.SetPose(TmpRequiredBones, BoneTransforms);
+					PoseWatch.SetWorldTransform(SkelMeshComponent->GetComponentTransform());
+
+					TRACE_ANIM_POSE_WATCH(*this, PoseWatch.NodeID, PoseWatch.GetBoneTransforms(), PoseWatch.GetRequiredBones(), PoseWatch.GetWorldTransform(), true);
+					break;
+				}
 			}
 		}
 	}
@@ -3161,16 +3189,25 @@ void FAnimInstanceProxy::RegisterWatchedPose(const FCSPose<FCompactPose>& Pose,
 {
 	if (bIsBeingDebugged)
 	{
-		for (FAnimNodePoseWatch& PoseWatch : PoseWatchEntriesForThisFrame)
+		if (USkeletalMeshComponent* SkelMeshComponent = GetSkelMeshComponent())
 		{
-			if (PoseWatch.NodeID == LinkID)
+			for (FAnimNodePoseWatch& PoseWatch : PoseWatchEntriesForThisFrame)
 			{
-				FCompactPose TempPose;
-				FCSPose<FCompactPose>::ConvertComponentPosesToLocalPoses(Pose, TempPose);
-				PoseWatch.PoseInfo->CopyBonesFrom(TempPose);
-				PoseWatch.Object = GetAnimInstanceObject();
-				PoseWatch.PoseWatch->SetIsNodeEnabled(true);
-				break;
+				if (PoseWatch.PoseWatch.IsValid() && PoseWatch.NodeID == LinkID)
+				{
+					FCompactPose TempPose;
+					FCSPose<FCompactPose>::ConvertComponentPosesToLocalPoses(Pose, TempPose);
+					PoseWatch.Object = GetAnimInstanceObject();
+					PoseWatch.PoseWatch->SetIsNodeEnabled(true);
+
+					const TArray<FTransform, FAnimStackAllocator>& BoneTransforms = TempPose.GetBones();
+					const TArray<FBoneIndexType>& TmpRequiredBones = TempPose.GetBoneContainer().GetBoneIndicesArray();
+					PoseWatch.SetPose(TmpRequiredBones, BoneTransforms);
+					PoseWatch.SetWorldTransform(SkelMeshComponent->GetComponentTransform());
+
+					TRACE_ANIM_POSE_WATCH(*this, PoseWatch.NodeID, PoseWatch.GetBoneTransforms(), PoseWatch.GetRequiredBones(), PoseWatch.GetWorldTransform(), true);
+					break;
+				}
 			}
 		}
 	}
diff --git a/Engine/Source/Runtime/Engine/Private/Animation/AnimTrace.cpp b/Engine/Source/Runtime/Engine/Private/Animation/AnimTrace.cpp
--- a/Engine/Source/Runtime/Engine/Private/Animation/AnimTrace.cpp
+++ b/Engine/Source/Runtime/Engine/Private/Animation/AnimTrace.cpp
@@ -241,6 +241,17 @@ UE_TRACE_EVENT_BEGIN(Animation, Sync)
 	UE_TRACE_EVENT_FIELD(uint32, GroupNameId)
 UE_TRACE_EVENT_END()
 
+UE_TRACE_EVENT_BEGIN(Animation, PoseWatch)
+	UE_TRACE_EVENT_FIELD(uint64, Cycle)
+	UE_TRACE_EVENT_FIELD(double, RecordingTime)
+	UE_TRACE_EVENT_FIELD(uint64, AnimInstanceId)
+	UE_TRACE_EVENT_FIELD(uint64, PoseWatchId)
+	UE_TRACE_EVENT_FIELD(float[], WorldTransform)
+	UE_TRACE_EVENT_FIELD(float[], BoneTransforms)
+	UE_TRACE_EVENT_FIELD(uint16[], RequiredBones)
+	UE_TRACE_EVENT_FIELD(bool, bIsEnabled)
+UE_TRACE_EVENT_END()
+
 // Object annotations used for tracing
 FUObjectAnnotationSparseBool GSkeletalMeshTraceAnnotations;
 
@@ -1169,4 +1180,32 @@ void FAnimTrace::OutputSync(const FAnimInstanceProxy& InSourceProxy, int32 InSou
 		<< Sync.GroupNameId(GroupNameId);
 }
 
+void FAnimTrace::OutputPoseWatch(const FAnimInstanceProxy& InSourceProxy, int32 InPoseWatchId, const TArray<FTransform>& BoneTransforms, const TArray<FBoneIndexType>& RequiredBones, const FTransform& WorldTransform, const bool bIsEnabled)
+{
+	bool bChannelEnabled = UE_TRACE_CHANNELEXPR_IS_ENABLED(AnimationChannel);
+	if (!bChannelEnabled)
+	{
+		return;
+	}
+
+	if (CANNOT_TRACE_OBJECT(InSourceProxy.GetSkelMeshComponent()))
+	{
+		return;
+	}
+
+	const UAnimInstance* AnimInstance = CastChecked<UAnimInstance>(InSourceProxy.GetAnimInstanceObject());
+
+	TRACE_OBJECT(InSourceProxy.GetAnimInstanceObject());
+
+	UE_TRACE_LOG(Animation, PoseWatch, AnimationChannel)
+		<< PoseWatch.Cycle(FPlatformTime::Cycles64())
+		<< PoseWatch.RecordingTime(FObjectTrace::GetWorldElapsedTime(AnimInstance->GetWorld()))
+		<< PoseWatch.AnimInstanceId(FObjectTrace::GetObjectId(InSourceProxy.GetAnimInstanceObject()))
+		<< PoseWatch.PoseWatchId(InPoseWatchId)
+		<< PoseWatch.WorldTransform(reinterpret_cast<const float*>(&WorldTransform), sizeof(FTransform) / sizeof(float))
+		<< PoseWatch.BoneTransforms(reinterpret_cast<const float*>(BoneTransforms.GetData()), BoneTransforms.Num() * (sizeof(FTransform) / sizeof(float)))
+		<< PoseWatch.RequiredBones(RequiredBones.GetData(), RequiredBones.Num())
+		<< PoseWatch.bIsEnabled(bIsEnabled);
+}
+
 #endif
diff --git a/Engine/Source/Runtime/Engine/Private/Components/SkeletalMeshComponent.cpp b/Engine/Source/Runtime/Engine/Private/Components/SkeletalMeshComponent.cpp
--- a/Engine/Source/Runtime/Engine/Private/Components/SkeletalMeshComponent.cpp
+++ b/Engine/Source/Runtime/Engine/Private/Components/SkeletalMeshComponent.cpp
@@ -4578,11 +4578,14 @@ void USkeletalMeshComponent::DebugDrawPoseWatches(FPrimitiveDrawInterface* PDI)
 			UAnimBlueprintGeneratedClass* AnimBPGenClass = CastChecked<UAnimBlueprintGeneratedClass>(AnimInstance->GetClass());
 			if (AnimBPGenClass)
 			{
-				const FAnimBlueprintDebugData& DebugData = AnimBPGenClass->GetAnimBlueprintDebugData();
-				DebugData.ForEachActiveVisiblePoseWatchPoseElement([this, PDI](const FPoseWatchDebugData& PoseWatchDebugData)
+				if (USkeletalMesh* TmpSkeletalMesh = GetSkeletalMeshAsset())
 				{
-					SkeletalDebugRendering::DrawBonesFromPoseWatch(*PoseWatchDebugData.PoseInfo.Get(), this, PDI, PoseWatchDebugData.PoseWatchPoseElement);
-				});
+					const FAnimBlueprintDebugData& DebugData = AnimBPGenClass->GetAnimBlueprintDebugData();
+					DebugData.ForEachActiveVisiblePoseWatchPoseElement([PDI, TmpSkeletalMesh](const FAnimNodePoseWatch& PoseWatch)
+					{
+						SkeletalDebugRendering::DrawBonesFromPoseWatch(PDI, PoseWatch, TmpSkeletalMesh->GetRefSkeleton(), /*bUseWorldTransform*/true);
+					});
+				}
 			}
 		}
 	}
diff --git a/Engine/Source/Runtime/Engine/Private/Engine/PoseWatch.cpp b/Engine/Source/Runtime/Engine/Private/Engine/PoseWatch.cpp
--- a/Engine/Source/Runtime/Engine/Private/Engine/PoseWatch.cpp
+++ b/Engine/Source/Runtime/Engine/Private/Engine/PoseWatch.cpp
@@ -488,6 +488,12 @@ bool UPoseWatchPoseElement::GetIsEnabled() const
 	return Parent.Get() && Parent->GetIsNodeEnabled();
 }
 
+FSlateIcon UPoseWatchPoseElement::StaticGetIcon()
+{
+	static FSlateIcon Icon(FAppStyle::GetAppStyleSetName(), "ClassIcon.PoseAsset");
+	return Icon;
+}
+
 #endif // WITH_EDITOR
 
 // UPoseWatch
diff --git a/Engine/Source/Runtime/Engine/Private/SkeletalDebugRendering.cpp b/Engine/Source/Runtime/Engine/Private/SkeletalDebugRendering.cpp
--- a/Engine/Source/Runtime/Engine/Private/SkeletalDebugRendering.cpp
+++ b/Engine/Source/Runtime/Engine/Private/SkeletalDebugRendering.cpp
@@ -4,7 +4,6 @@
 #include "DrawDebugHelpers.h"
 #include "SceneManagement.h"
 #include "Engine/PoseWatch.h"
-#include "Components/SkeletalMeshComponent.h"
 #include "BonePose.h"
 #include "Animation/BlendProfile.h"
 #include "HitProxies.h"
@@ -172,63 +171,67 @@ void DrawRootCone(
 }
 
 void DrawBonesFromPoseWatch(
-	const FCompactHeapPose& Pose,
-	USkeletalMeshComponent* MeshComponent,
 	FPrimitiveDrawInterface* PDI,
-	const UPoseWatchPoseElement* PoseWatchPoseElement)
+	const FAnimNodePoseWatch& PoseWatch,
+	const FReferenceSkeleton& RefSkeleton,
+	const bool bUseWorldTransform
+)
 {
 #if WITH_EDITOR
-	if (Pose.GetNumBones() == 0 ||
-		Pose.GetBoneContainer().GetCompactPoseNumBones() == 0 ||
-		!MeshComponent ||
-		!MeshComponent->GetSkeletalMeshAsset() ||
-		!PoseWatchPoseElement
-		)
+	if (!PoseWatch.PoseWatchPoseElement.IsValid())
 	{
 		return;
 	}
+	const FLinearColor BoneColor = FLinearColor(PoseWatch.PoseWatchPoseElement->GetColor());
 
-	// optionally override draw color
-	const FLinearColor BoneColor = FLinearColor(PoseWatchPoseElement->GetColor());
-
-	TArray<FTransform> WorldTransforms;
-	WorldTransforms.AddUninitialized(Pose.GetBoneContainer().GetNumBones());
-
-	TArray<uint16> RequiredBones;
-	const UBlendProfile* ViewportMask = PoseWatchPoseElement->ViewportMask;
-
-	const FVector RelativeOffset = MeshComponent->GetComponentRotation().RotateVector(PoseWatchPoseElement->ViewportOffset);
-
-	// we could cache parent bones as we calculate, but right now I'm not worried about perf issue of this
-	for (FCompactPoseBoneIndex BoneIndex : Pose.ForEachBoneIndex())
+	const TArray<FTransform>& InBoneTransforms = PoseWatch.GetBoneTransforms();
+	const TArray<FBoneIndexType>& InRequiredBones = PoseWatch.GetRequiredBones();
+	if (InRequiredBones.Num() == 0 || InBoneTransforms.Num() < InRequiredBones.Num())
 	{
-		FMeshPoseBoneIndex MeshBoneIndex = Pose.GetBoneContainer().MakeMeshPoseIndex(BoneIndex);
+		return;
+	}
 
-		int32 ParentIndex = Pose.GetBoneContainer().GetParentBoneIndex(MeshBoneIndex.GetInt());
+	const UBlendProfile* ViewportMask = PoseWatch.PoseWatchPoseElement->ViewportMask;
+	const FTransform WorldTransform = bUseWorldTransform ? PoseWatch.GetWorldTransform() : FTransform::Identity;
+	const FVector RelativeOffset = WorldTransform.GetRotation().RotateVector(PoseWatch.PoseWatchPoseElement->ViewportOffset);
+
+	TArray<FTransform> UseWorldTransforms;
+	UseWorldTransforms.AddDefaulted(RefSkeleton.GetNum());
+	
+	TArray<FBoneIndexType> UseRequiredBones;
+	UseRequiredBones.Reserve(InRequiredBones.Num());
+	
+	for (int32 Index = 0; Index < InRequiredBones.Num(); ++Index)
+	{
+		const FBoneIndexType& BoneIndex = InRequiredBones[Index];
+		if (!RefSkeleton.IsValidIndex(BoneIndex))
+		{
+			continue;
+		}
 
-		if (ParentIndex == INDEX_NONE)
+		const int32 ParentIndex = RefSkeleton.GetParentIndex(BoneIndex);
+		if (!RefSkeleton.IsValidIndex(ParentIndex))
 		{
-			WorldTransforms[MeshBoneIndex.GetInt()] = Pose[BoneIndex] * MeshComponent->GetComponentTransform();
-			WorldTransforms[MeshBoneIndex.GetInt()].AddToTranslation(RelativeOffset);
+			UseWorldTransforms[BoneIndex] = InBoneTransforms[Index] * WorldTransform;
+			UseWorldTransforms[BoneIndex].AddToTranslation(RelativeOffset);
 		}
 		else
 		{
-			WorldTransforms[MeshBoneIndex.GetInt()] = Pose[BoneIndex] * WorldTransforms[ParentIndex];
+			UseWorldTransforms[BoneIndex] = InBoneTransforms[Index] * UseWorldTransforms[ParentIndex];
 		}
 
-		if (ViewportMask == nullptr)
+		if (ViewportMask)
 		{
-			RequiredBones.Add(MeshBoneIndex.GetInt());
-		}
-		else
-		{
-			const FName& BoneName = MeshComponent->GetBoneName(MeshBoneIndex.GetInt());
-			const bool bBoneBlendScaleMeetsThreshold = ViewportMask->GetBoneBlendScale(BoneName) > PoseWatchPoseElement->BlendScaleThreshold;
-			if ((!PoseWatchPoseElement->bInvertViewportMask && bBoneBlendScaleMeetsThreshold) || (PoseWatchPoseElement->bInvertViewportMask && !bBoneBlendScaleMeetsThreshold))
+			const FName& BoneName = RefSkeleton.GetBoneName(BoneIndex);
+			const bool bBoneBlendScaleMeetsThreshold = ViewportMask->GetBoneBlendScale(BoneName) > PoseWatch.PoseWatchPoseElement->BlendScaleThreshold;
+			const bool bBoneRequired = (!PoseWatch.PoseWatchPoseElement->bInvertViewportMask && bBoneBlendScaleMeetsThreshold) || (PoseWatch.PoseWatchPoseElement->bInvertViewportMask && !bBoneBlendScaleMeetsThreshold);
+			if (!bBoneRequired)
 			{
-				RequiredBones.Add(MeshBoneIndex.GetInt());
+				continue;
 			}
 		}
+
+		UseRequiredBones.Add(BoneIndex);
 	}
 
 	FSkelDebugDrawConfig DrawConfig;
@@ -243,10 +246,10 @@ void DrawBonesFromPoseWatch(
 
 	SkeletalDebugRendering::DrawBones(
 		PDI,
-		MeshComponent->GetComponentLocation() + RelativeOffset,
-		RequiredBones,
-		MeshComponent->GetSkeletalMeshAsset()->GetRefSkeleton(),
-		WorldTransforms,
+		WorldTransform.GetLocation() + RelativeOffset,
+		UseRequiredBones,
+		RefSkeleton,
+		UseWorldTransforms,
 		/*SelectedBones*/TArray<int32>(),
 		/*BoneColors*/TArray<FLinearColor>(),
 		/*HitProxies*/TArray<HHitProxy*>(),
diff --git a/Engine/Source/Runtime/Engine/Public/Animation/AnimTrace.h b/Engine/Source/Runtime/Engine/Public/Animation/AnimTrace.h
--- a/Engine/Source/Runtime/Engine/Public/Animation/AnimTrace.h
+++ b/Engine/Source/Runtime/Engine/Public/Animation/AnimTrace.h
@@ -186,6 +186,9 @@ struct FAnimTrace
 
 	/** Helper function to output a sync record */
 	ENGINE_API static void OutputSync(const FAnimInstanceProxy& InSourceProxy, int32 InSourceNodeId, FName InGroupName);
+
+	/** Helper function to output a pose watch record */
+	ENGINE_API static void OutputPoseWatch(const FAnimInstanceProxy& InSourceProxy, int32 InPoseWatchId, const TArray<FTransform>& BoneTransforms, const TArray<FBoneIndexType>& RequiredBones, const FTransform& WorldTransform, const bool bIsEnabled);
 };
 
 #define TRACE_ANIM_TICK_RECORD(Context, TickRecord) \
@@ -239,6 +242,9 @@ struct FAnimTrace
 #define TRACE_ANIM_NODE_SYNC(SourceProxy, SourceNodeId, GroupName) \
 	FAnimTrace::OutputSync(SourceProxy, SourceNodeId, GroupName);
 
+#define TRACE_ANIM_POSE_WATCH(SourceProxy, PoseWatchId, BoneTransforms, RequiredBones, WorldTransform, bIsEnabled) \
+	FAnimTrace::OutputPoseWatch(SourceProxy, PoseWatchId, BoneTransforms, RequiredBones, WorldTransform, bIsEnabled);
+
 #else
 
 #define TRACE_ANIM_TICK_RECORD(Context, TickRecord)
@@ -258,5 +264,6 @@ struct FAnimTrace
 #define TRACE_ANIM_SYNC_MARKER(AnimInstance, SyncMarker)
 #define TRACE_ANIM_MONTAGE(AnimInstance, MontageInstance)
 #define TRACE_ANIM_NODE_SYNC(SourceProxy, SourceNodeId, GroupName)
+#define TRACE_ANIM_POSE_WATCH(SourceProxy, PoseWatchId, BoneTransforms, RequiredBones, WorldTransform, bIsEnabled)
 
 #endif
diff --git a/Engine/Source/Runtime/Engine/Public/SkeletalDebugRendering.h b/Engine/Source/Runtime/Engine/Public/SkeletalDebugRendering.h
--- a/Engine/Source/Runtime/Engine/Public/SkeletalDebugRendering.h
+++ b/Engine/Source/Runtime/Engine/Public/SkeletalDebugRendering.h
@@ -6,12 +6,10 @@
 #include "Engine/EngineTypes.h"
 #include "BoneIndices.h"
 
-class UPoseWatch;
-class UPoseWatchPoseElement;
+struct FAnimNodePoseWatch;
 class FPrimitiveDrawInterface;
 struct FCompactHeapPose;
 struct FReferenceSkeleton;
-class USkeletalMeshComponent;
 class HHitProxy;
 
 namespace EBoneDrawMode
@@ -114,10 +112,10 @@ ENGINE_API	void DrawRootCone(
 
 
 ENGINE_API void DrawBonesFromPoseWatch(
-	const FCompactHeapPose& Pose,
-	USkeletalMeshComponent* MeshComponent,
 	FPrimitiveDrawInterface* PDI,
-	const UPoseWatchPoseElement* PoseWatch);
+	const FAnimNodePoseWatch& PoseWatch,
+	const FReferenceSkeleton& RefSkeleton,
+	const bool bUseWorldTransform);
 
 
 /**
