Added include guards for headers from modules not loaded in non-editor builds

diff --git a/Engine/Plugins/Animation/GameplayInsights/Source/GameplayInsights/Private/PoseWatchTrack.cpp b/Engine/Plugins/Animation/GameplayInsights/Source/GameplayInsights/Private/PoseWatchTrack.cpp
--- a/Engine/Plugins/Animation/GameplayInsights/Source/GameplayInsights/Private/PoseWatchTrack.cpp
+++ b/Engine/Plugins/Animation/GameplayInsights/Source/GameplayInsights/Private/PoseWatchTrack.cpp
@@ -6,11 +6,12 @@
 #include "GameplayProvider.h"
 #include "AnimationProvider.h"
 #include "SNotifiesView.h"
-#include "Animation/AnimBlueprint.h"
+#include "ObjectTrace.h"
+#if WITH_EDITOR
 #include "Animation/AnimBlueprintGeneratedClass.h"
 #include "Animation/AnimInstance.h"
 #include "Engine/PoseWatch.h"
-#include "ObjectTrace.h"
+#endif
 
 #define LOCTEXT_NAMESPACE "PoseWatchesTrack"
 
@@ -23,6 +24,8 @@ FPoseWatchTrack::FPoseWatchTrack(uint64 InObjectId, const FPoseWatchTrack::FPose
 	, PoseWatchOwner(nullptr)
 {
 	EnabledSegments = MakeShared<SSegmentedTimelineView::FSegmentData>();
+
+#if WITH_EDITOR
 	Icon = UPoseWatchPoseElement::StaticGetIcon();
 
 	if (UObject* ObjectInstance = FObjectTrace::GetObjectFromId(ObjectId))
@@ -43,6 +46,7 @@ FPoseWatchTrack::FPoseWatchTrack(uint64 InObjectId, const FPoseWatchTrack::FPose
 			}
 		}
 	}
+#endif
 }
 
 FText FPoseWatchTrack::GetDisplayNameInternal() const
@@ -57,6 +61,7 @@ TSharedPtr<SSegmentedTimelineView::FSegmentData> FPoseWatchTrack::GetSegmentData
 
 bool FPoseWatchTrack::UpdateInternal()
 {
+#if WITH_EDITOR
 	IRewindDebugger* RewindDebugger = IRewindDebugger::Instance();
 	
 	TRange<double> TraceTimeRange = RewindDebugger->GetCurrentTraceRange();
@@ -142,10 +147,14 @@ bool FPoseWatchTrack::UpdateInternal()
 	}
 
 	return bChanged;
+#else
+	return false;
+#endif
 }
 
 TSharedPtr<SWidget> FPoseWatchTrack::GetTimelineViewInternal()
 {
+#if WITH_EDITOR
 	const auto GetPoseWatchColorLambda = [this]() -> FLinearColor { return PoseWatchOwner ? FLinearColor(PoseWatchOwner->GetColor()) : FLinearColor::White; };
 
 	const auto TimelineView = SNew(SSegmentedTimelineView)
@@ -154,6 +163,9 @@ TSharedPtr<SWidget> FPoseWatchTrack::GetTimelineViewInternal()
 		.SegmentData_Raw(this, &FPoseWatchTrack::GetSegmentData);
 
 	return TimelineView;
+#else
+	return TSharedPtr<SWidget>();
+#endif
 }
 
 FName FPoseWatchesTrackCreator::GetTargetTypeNameInternal() const
@@ -177,6 +189,7 @@ TSharedPtr<RewindDebugger::FRewindDebuggerTrack> FPoseWatchesTrackCreator::Creat
 FPoseWatchesTrack::FPoseWatchesTrack(uint64 InObjectId)
 	: ObjectId(InObjectId)
 {
+#if WITH_EDITOR
 	Icon = UPoseWatchPoseElement::StaticGetIcon();
 
 	if (UObject* ObjectInstance = FObjectTrace::GetObjectFromId(ObjectId))
@@ -189,11 +202,13 @@ FPoseWatchesTrack::FPoseWatchesTrack(uint64 InObjectId)
 			}
 		}
 	}
+#endif
 }
 
 
 bool FPoseWatchesTrack::UpdateInternal()
 {
+#if WITH_EDITOR
 	IRewindDebugger* RewindDebugger = IRewindDebugger::Instance();
 	
 	TRange<double> TraceTimeRange = RewindDebugger->GetCurrentTraceRange();
@@ -252,6 +267,9 @@ bool FPoseWatchesTrack::UpdateInternal()
 	}
 	
 	return bChanged;
+#else
+	return false;
+#endif
 }
 	
 void FPoseWatchesTrack::IterateSubTracksInternal(TFunction<void(TSharedPtr<FRewindDebuggerTrack> SubTrack)> IteratorFunction)
@@ -264,6 +282,7 @@ void FPoseWatchesTrack::IterateSubTracksInternal(TFunction<void(TSharedPtr<FRewi
 
 bool FPoseWatchesTrackCreator::HasDebugInfoInternal(uint64 ObjectId) const
 {
+#if WITH_EDITOR
 	const TraceServices::IAnalysisSession* AnalysisSession = IRewindDebugger::Instance()->GetAnalysisSession();
 	
 	TraceServices::FAnalysisSessionReadScope SessionReadScope(*AnalysisSession);
@@ -276,6 +295,9 @@ bool FPoseWatchesTrackCreator::HasDebugInfoInternal(uint64 ObjectId) const
 		});
 	}
 	return bHasData;
+#else
+	return false;
+#endif
 }
 
 	
