build fix for PoseWatchTrack
Disable entirely exept WITH_EDITOR

diff --git a/Engine/Plugins/Animation/GameplayInsights/Source/GameplayInsights/Private/PoseWatchTrack.cpp b/Engine/Plugins/Animation/GameplayInsights/Source/GameplayInsights/Private/PoseWatchTrack.cpp
--- a/Engine/Plugins/Animation/GameplayInsights/Source/GameplayInsights/Private/PoseWatchTrack.cpp
+++ b/Engine/Plugins/Animation/GameplayInsights/Source/GameplayInsights/Private/PoseWatchTrack.cpp
@@ -1,17 +1,20 @@
 // Copyright Epic Games, Inc. All Rights Reserved.
 
+
+#if WITH_EDITOR
+
+
 #include "PoseWatchTrack.h"
 #include "SCurveTimelineView.h"
 #include "IRewindDebugger.h"
 #include "GameplayProvider.h"
 #include "AnimationProvider.h"
 #include "SNotifiesView.h"
-#include "ObjectTrace.h"
-#if WITH_EDITOR
+#include "Animation/AnimBlueprint.h"
 #include "Animation/AnimBlueprintGeneratedClass.h"
 #include "Animation/AnimInstance.h"
 #include "Engine/PoseWatch.h"
-#endif
+#include "ObjectTrace.h"
 
 #define LOCTEXT_NAMESPACE "PoseWatchesTrack"
 
@@ -24,8 +27,6 @@ FPoseWatchTrack::FPoseWatchTrack(uint64 InObjectId, const FPoseWatchTrack::FPose
 	, PoseWatchOwner(nullptr)
 {
 	EnabledSegments = MakeShared<SSegmentedTimelineView::FSegmentData>();
-
-#if WITH_EDITOR
 	Icon = UPoseWatchPoseElement::StaticGetIcon();
 
 	if (UObject* ObjectInstance = FObjectTrace::GetObjectFromId(ObjectId))
@@ -46,7 +47,6 @@ FPoseWatchTrack::FPoseWatchTrack(uint64 InObjectId, const FPoseWatchTrack::FPose
 			}
 		}
 	}
-#endif
 }
 
 FText FPoseWatchTrack::GetDisplayNameInternal() const
@@ -61,7 +61,6 @@ TSharedPtr<SSegmentedTimelineView::FSegmentData> FPoseWatchTrack::GetSegmentData
 
 bool FPoseWatchTrack::UpdateInternal()
 {
-#if WITH_EDITOR
 	IRewindDebugger* RewindDebugger = IRewindDebugger::Instance();
 	
 	TRange<double> TraceTimeRange = RewindDebugger->GetCurrentTraceRange();
@@ -147,14 +146,10 @@ bool FPoseWatchTrack::UpdateInternal()
 	}
 
 	return bChanged;
-#else
-	return false;
-#endif
 }
 
 TSharedPtr<SWidget> FPoseWatchTrack::GetTimelineViewInternal()
 {
-#if WITH_EDITOR
 	const auto GetPoseWatchColorLambda = [this]() -> FLinearColor { return PoseWatchOwner ? FLinearColor(PoseWatchOwner->GetColor()) : FLinearColor::White; };
 
 	const auto TimelineView = SNew(SSegmentedTimelineView)
@@ -163,9 +158,6 @@ TSharedPtr<SWidget> FPoseWatchTrack::GetTimelineViewInternal()
 		.SegmentData_Raw(this, &FPoseWatchTrack::GetSegmentData);
 
 	return TimelineView;
-#else
-	return TSharedPtr<SWidget>();
-#endif
 }
 
 FName FPoseWatchesTrackCreator::GetTargetTypeNameInternal() const
@@ -189,7 +181,6 @@ TSharedPtr<RewindDebugger::FRewindDebuggerTrack> FPoseWatchesTrackCreator::Creat
 FPoseWatchesTrack::FPoseWatchesTrack(uint64 InObjectId)
 	: ObjectId(InObjectId)
 {
-#if WITH_EDITOR
 	Icon = UPoseWatchPoseElement::StaticGetIcon();
 
 	if (UObject* ObjectInstance = FObjectTrace::GetObjectFromId(ObjectId))
@@ -202,13 +193,11 @@ FPoseWatchesTrack::FPoseWatchesTrack(uint64 InObjectId)
 			}
 		}
 	}
-#endif
 }
 
 
 bool FPoseWatchesTrack::UpdateInternal()
 {
-#if WITH_EDITOR
 	IRewindDebugger* RewindDebugger = IRewindDebugger::Instance();
 	
 	TRange<double> TraceTimeRange = RewindDebugger->GetCurrentTraceRange();
@@ -267,9 +256,6 @@ bool FPoseWatchesTrack::UpdateInternal()
 	}
 	
 	return bChanged;
-#else
-	return false;
-#endif
 }
 	
 void FPoseWatchesTrack::IterateSubTracksInternal(TFunction<void(TSharedPtr<FRewindDebuggerTrack> SubTrack)> IteratorFunction)
@@ -282,7 +268,6 @@ void FPoseWatchesTrack::IterateSubTracksInternal(TFunction<void(TSharedPtr<FRewi
 
 bool FPoseWatchesTrackCreator::HasDebugInfoInternal(uint64 ObjectId) const
 {
-#if WITH_EDITOR
 	const TraceServices::IAnalysisSession* AnalysisSession = IRewindDebugger::Instance()->GetAnalysisSession();
 	
 	TraceServices::FAnalysisSessionReadScope SessionReadScope(*AnalysisSession);
@@ -295,12 +280,11 @@ bool FPoseWatchesTrackCreator::HasDebugInfoInternal(uint64 ObjectId) const
 		});
 	}
 	return bHasData;
-#else
-	return false;
-#endif
 }
 
 	
 }
 
-#undef LOCTEXT_NAMESPACE
\ No newline at end of file
+#undef LOCTEXT_NAMESPACE
+
+#endif
\ No newline at end of file
diff --git a/Engine/Plugins/Animation/GameplayInsights/Source/GameplayInsights/Private/PoseWatchTrack.h b/Engine/Plugins/Animation/GameplayInsights/Source/GameplayInsights/Private/PoseWatchTrack.h
--- a/Engine/Plugins/Animation/GameplayInsights/Source/GameplayInsights/Private/PoseWatchTrack.h
+++ b/Engine/Plugins/Animation/GameplayInsights/Source/GameplayInsights/Private/PoseWatchTrack.h
@@ -1,6 +1,8 @@
 // Copyright Epic Games, Inc. All Rights Reserved.
 #pragma once
 
+#if WITH_EDITOR
+
 #include "IAnimationProvider.h"
 #include "IRewindDebuggerTrackCreator.h"
 #include "RewindDebuggerTrack.h"
@@ -87,4 +89,6 @@ private:
 	virtual bool HasDebugInfoInternal(uint64 ObjectId) const override;
 };
 
-}
\ No newline at end of file
+}
+
+#endif
\ No newline at end of file
