Rewind debugger object hierarchy optimization
- store parent->child hierarchy relationships to avoid iterating over all objects in the world frequently

diff --git a/Engine/Plugins/Animation/GameplayInsights/Source/GameplayInsights/Private/GameplayProvider.cpp b/Engine/Plugins/Animation/GameplayInsights/Source/GameplayInsights/Private/GameplayProvider.cpp
--- a/Engine/Plugins/Animation/GameplayInsights/Source/GameplayInsights/Private/GameplayProvider.cpp
+++ b/Engine/Plugins/Animation/GameplayInsights/Source/GameplayInsights/Private/GameplayProvider.cpp
@@ -179,7 +179,6 @@ void FGameplayProvider::EnumerateObjectPropertyValues(uint64 InObjectId, const F
 		}
 	}
 }
-
 void FGameplayProvider::EnumerateObjects(TFunctionRef<void(const FObjectInfo&)> Callback) const
 {
 	Session.ReadAccessCheck();
@@ -203,6 +202,18 @@ void FGameplayProvider::EnumerateObjects(double StartTime, double EndTime, TFunc
 		});
 }
 
+void FGameplayProvider::EnumerateSubobjects(uint64 ObjectId, TFunctionRef<void(uint64 SubobjectId)> Callback) const
+{
+	TArray<uint64, TInlineAllocator<32>> SubobjectIds;
+	ObjectHierarchy.MultiFind(ObjectId, SubobjectIds);
+
+	for (auto SubObjectId : SubobjectIds)
+	{
+		Callback(SubObjectId);
+	}
+}
+
+
 const FClassInfo* FGameplayProvider::FindClassInfo(uint64 InClassId) const
 {
 	Session.ReadAccessCheck();
@@ -398,6 +409,8 @@ void FGameplayProvider::AppendObject(uint64 InObjectId, uint64 InOuterId, uint64
 
 		int32 NewObjectInfoIndex = ObjectInfos.Add(NewObjectInfo);
 		ObjectIdToIndexMap.Add(InObjectId, NewObjectInfoIndex);
+
+		ObjectHierarchy.AddUnique(NewObjectInfo.OuterId, NewObjectInfo.Id);
 	}
 }
 
diff --git a/Engine/Plugins/Animation/GameplayInsights/Source/GameplayInsights/Private/GameplayProvider.h b/Engine/Plugins/Animation/GameplayInsights/Source/GameplayInsights/Private/GameplayProvider.h
--- a/Engine/Plugins/Animation/GameplayInsights/Source/GameplayInsights/Private/GameplayProvider.h
+++ b/Engine/Plugins/Animation/GameplayInsights/Source/GameplayInsights/Private/GameplayProvider.h
@@ -25,6 +25,7 @@ public:
 	virtual void EnumerateObjectPropertyValues(uint64 InObjectId, const FObjectPropertiesMessage& InMessage, TFunctionRef<void(const FObjectPropertyValue&)> Callback) const override;
 	virtual void EnumerateObjects(TFunctionRef<void(const FObjectInfo&)> Callback) const override;
 	virtual void EnumerateObjects(double StartTime, double EndTime, TFunctionRef<void(const FObjectInfo&)> Callback) const override;
+	virtual void EnumerateSubobjects(uint64 ObjectId, TFunctionRef<void(uint64 SubobjectId)> Callback) const override;
 	virtual const FClassInfo* FindClassInfo(uint64 InClassId) const override;
 	virtual const UClass* FindClass(uint64 InClassId) const override;
 	virtual const FClassInfo* FindClassInfo(const TCHAR* InClassPath) const override;
@@ -181,4 +182,7 @@ private:
 
 	/** Whether we have any object properties */
 	bool bHasObjectProperties;
+
+	/** MultiMap from parent object to sub objects, to avoid slow reverse lookups */
+	TMultiMap<uint64, uint64> ObjectHierarchy;
 };
diff --git a/Engine/Plugins/Animation/GameplayInsights/Source/GameplayInsights/Public/IGameplayProvider.h b/Engine/Plugins/Animation/GameplayInsights/Source/GameplayInsights/Public/IGameplayProvider.h
--- a/Engine/Plugins/Animation/GameplayInsights/Source/GameplayInsights/Public/IGameplayProvider.h
+++ b/Engine/Plugins/Animation/GameplayInsights/Source/GameplayInsights/Public/IGameplayProvider.h
@@ -112,6 +112,7 @@ public:
 	virtual void EnumerateObjectPropertyValues(uint64 InObjectId, const FObjectPropertiesMessage& InMessage, TFunctionRef<void(const FObjectPropertyValue&)> Callback) const = 0;
 	virtual void EnumerateObjects(TFunctionRef<void(const FObjectInfo&)> Callback) const = 0;
 	virtual void EnumerateObjects(double StartTime, double EndTime, TFunctionRef<void(const FObjectInfo&)> Callback) const = 0;
+	virtual void EnumerateSubobjects(uint64 ObjectId, TFunctionRef<void(uint64 SubobjectId)> Callback) const = 0;
 	virtual const FClassInfo* FindClassInfo(uint64 InClassId) const = 0;
 	virtual const UClass* FindClass(uint64 InClassId) const = 0;
 	virtual const FClassInfo* FindClassInfo(const TCHAR* InClassPath) const = 0;
diff --git a/Engine/Plugins/Animation/GameplayInsights/Source/RewindDebugger/Private/RewindDebugger.cpp b/Engine/Plugins/Animation/GameplayInsights/Source/RewindDebugger/Private/RewindDebugger.cpp
--- a/Engine/Plugins/Animation/GameplayInsights/Source/RewindDebugger/Private/RewindDebugger.cpp
+++ b/Engine/Plugins/Animation/GameplayInsights/Source/RewindDebugger/Private/RewindDebugger.cpp
@@ -52,6 +52,9 @@ FRewindDebugger::FRewindDebugger()  :
 	bTargetActorPositionValid(false),
 	bIsDetailsPanelOpen(true)
 {
+	UE::Trace::ToggleChannel(TEXT("Cpu"), false);
+	UE::Trace::ToggleChannel(TEXT("Gpu"), false);
+
 	RecordingDuration.Set(0);
 
 	if (GEditor->bIsSimulatingInEditor || GEditor->PlayWorld)
@@ -226,9 +229,11 @@ uint64 FRewindDebugger::GetTargetActorId() const
 
 void FRewindDebugger::RefreshDebugTracks()
 {
+	TRACE_CPUPROFILER_EVENT_SCOPE(FRewindDebugger::RefreshDebugTracks);
 	if (const TraceServices::IAnalysisSession* Session = GetAnalysisSession())
 	{
 		TraceServices::FAnalysisSessionReadScope SessionReadScope(*Session);
+
 		UWorld* World = GetWorldToVisualize();
 
 		if (const IGameplayProvider* GameplayProvider = Session->ReadProvider<IGameplayProvider>("GameplayProvider"))
@@ -671,6 +676,7 @@ void FRewindDebugger::Tick(float DeltaTime)
 			{
 				if (bRecording)
 				{
+					TRACE_CPUPROFILER_EVENT_SCOPE(FRewindDebugger::Tick_UpdateSimulating);
 					RecordingDuration.Set(FObjectTrace::GetWorldElapsedTime(World));
 					SetCurrentScrubTime(RecordingDuration.Get());
 					TrackCursorDelegate.ExecuteIfBound(false);
@@ -682,6 +688,7 @@ void FRewindDebugger::Tick(float DeltaTime)
 				{
 					if (ControlState == EControlState::Play || ControlState == EControlState::PlayReverse)
 					{
+						TRACE_CPUPROFILER_EVENT_SCOPE(FRewindDebugger::Tick_UpdatePlayback);
 						float Rate = PlaybackRate * (ControlState == EControlState::Play ? 1 : -1);
 						SetCurrentScrubTime(FMath::Clamp(CurrentScrubTime + Rate * DeltaTime, 0.0f, RecordingDuration.Get()));
 						TrackCursorDelegate.ExecuteIfBound(Rate<0);
diff --git a/Engine/Plugins/Animation/GameplayInsights/Source/RewindDebugger/Private/RewindDebuggerObjectTrack.cpp b/Engine/Plugins/Animation/GameplayInsights/Source/RewindDebugger/Private/RewindDebuggerObjectTrack.cpp
--- a/Engine/Plugins/Animation/GameplayInsights/Source/RewindDebugger/Private/RewindDebuggerObjectTrack.cpp
+++ b/Engine/Plugins/Animation/GameplayInsights/Source/RewindDebugger/Private/RewindDebuggerObjectTrack.cpp
@@ -54,9 +54,40 @@ void FRewindDebuggerObjectTrack::IterateSubTracksInternal(TFunction<void(TShared
 	}
 };
 
+FText FRewindDebuggerObjectTrack::GetDisplayNameInternal() const
+{
+	if (!bDisplayNameValid)
+	{
+		IRewindDebugger* RewindDebugger = IRewindDebugger::Instance();
+		const TraceServices::IAnalysisSession* Session = RewindDebugger->GetAnalysisSession();
+		TraceServices::FAnalysisSessionReadScope SessionReadScope(*Session);
+
+		const IGameplayProvider* GameplayProvider = Session->ReadProvider<IGameplayProvider>("GameplayProvider");
+
+		const FObjectInfo& ObjectInfo = GameplayProvider->GetObjectInfo(ObjectId);
+
+		DisplayName = FText::FromString(ObjectInfo.Name);
+
+		if (const FWorldInfo* WorldInfo = GameplayProvider->FindWorldInfoFromObject(ObjectId))
+		{
+			bool bIsServer = WorldInfo->NetMode == FWorldInfo::ENetMode::DedicatedServer;
+
+			if (bIsServer)
+			{
+				DisplayName = FText::Format(NSLOCTEXT("RewindDebuggerTrack", " (Server)", "{0} (Server)"), FText::FromString(ObjectInfo.Name));
+			}
+		}
+
+		bDisplayNameValid = true;
+	}
+
+	return DisplayName;
+}
+
 bool FRewindDebuggerObjectTrack::UpdateInternal()
 {
-	IRewindDebugger *RewindDebugger = IRewindDebugger::Instance();
+	TRACE_CPUPROFILER_EVENT_SCOPE(FRewindDebuggerObjectTrack::UpdateInternal);
+	IRewindDebugger* RewindDebugger = IRewindDebugger::Instance();
 	const TraceServices::IAnalysisSession* Session = RewindDebugger->GetAnalysisSession();
 	TraceServices::FAnalysisSessionReadScope SessionReadScope(*Session);
 
@@ -85,31 +116,31 @@ bool FRewindDebuggerObjectTrack::UpdateInternal()
 
 	// add debug views as children
 	FRewindDebuggerTrackCreators::EnumerateCreators([this, &Session, &bChanged, RewindDebugger](const IRewindDebuggerTrackCreator* Creator)
-	{
-		int32 FoundIndex = Children.FindLastByPredicate([&bChanged, this, Creator, RewindDebugger](const TSharedPtr<FRewindDebuggerTrack>& Track) { return Track->GetName() == Creator->GetName(); });
+		{
+			int32 FoundIndex = Children.FindLastByPredicate([&bChanged, this, Creator, RewindDebugger](const TSharedPtr<FRewindDebuggerTrack>& Track) { return Track->GetName() == Creator->GetName(); });
 
-		const bool bHasDebugInfo = Creator->HasDebugInfo(ObjectId) && IsTargetType(ObjectId, Creator->GetTargetTypeName(), *Session);
+			const bool bHasDebugInfo = Creator->HasDebugInfo(ObjectId) && IsTargetType(ObjectId, Creator->GetTargetTypeName(), *Session);
 
-		if (FoundIndex >= 0)
-		{
-			if (!bHasDebugInfo)
+			if (FoundIndex >= 0)
 			{
-				bChanged = true;
-				Children.RemoveAt(FoundIndex);
+				if (!bHasDebugInfo)
+				{
+					bChanged = true;
+					Children.RemoveAt(FoundIndex);
+				}
 			}
-		}
-		else
-		{
-			if (bHasDebugInfo)
+			else
 			{
-				bChanged = true;
-				if (TSharedPtr<FRewindDebuggerTrack> Track = Creator->CreateTrack(ObjectId))
+				if (bHasDebugInfo)
 				{
-					Children.Add(Track);
+					bChanged = true;
+					if (TSharedPtr<FRewindDebuggerTrack> Track = Creator->CreateTrack(ObjectId))
+					{
+						Children.Add(Track);
+					}
 				}
 			}
-		}
-	});
+		});
 
 	// Fallback codepath to add views with no track implementation
 	FRewindDebuggerViewCreators::EnumerateCreators([this, &Session, &bChanged, RewindDebugger](const IRewindDebuggerViewCreator* Creator)
@@ -140,24 +171,33 @@ bool FRewindDebuggerObjectTrack::UpdateInternal()
 	);
 
 	// add child objects/components
-	GameplayProvider->EnumerateObjects(RewindDebugger->GetCurrentTraceRange().GetLowerBoundValue(), RewindDebugger->GetCurrentTraceRange().GetUpperBoundValue(),
-		[this, &bChanged, GameplayProvider, &FoundObjects, RewindDebugger](const FObjectInfo& InObjectInfo)
-		{
-			if (InObjectInfo.OuterId == ObjectId)
-			{
-				int32 FoundIndex = Children.FindLastByPredicate([&InObjectInfo](const TSharedPtr<FRewindDebuggerTrack>& Info) { return Info->GetObjectId() == InObjectInfo.Id; });
-				bool bExisted = FoundIndex >= 0;
+	{
+		TRACE_CPUPROFILER_EVENT_SCOPE(FRewindDebuggerObjectTrack::UpdateInternal_AddChildComponents);
 
-				if (!bExisted)
+		TRange<double> ViewRange = RewindDebugger->GetCurrentViewRange();
+
+		GameplayProvider->EnumerateSubobjects(ObjectId, [this, &FoundObjects, &bChanged,&ViewRange, GameplayProvider](uint64 SubobjectId)
+			{
+				TRange<double> Lifetime = GameplayProvider->GetObjectRecordingLifetime(SubobjectId);
+				TRange<double> Overlap = TRange<double>::Intersection(Lifetime, ViewRange);
+				// only display the track if the lifetime of the object and the view range overlap
+				if (!Overlap.IsEmpty())
 				{
-					FoundIndex = Children.Num();
-					Children.Add(MakeShared<FRewindDebuggerObjectTrack>(InObjectInfo.Id, InObjectInfo.Name));
-				}
+					int32 FoundIndex = Children.FindLastByPredicate([SubobjectId](const TSharedPtr<FRewindDebuggerTrack>& Info) { return Info->GetObjectId() == SubobjectId; });
+					bool bExisted = FoundIndex >= 0;
 
-				bChanged = bChanged | !bExisted;
-				FoundObjects.Add(InObjectInfo.Id);
-			}
-		});
+					if (!bExisted)
+					{
+						FoundIndex = Children.Num();
+
+						Children.Add(MakeShared<FRewindDebuggerObjectTrack>(SubobjectId, ""));
+					}
+
+					bChanged = bChanged || !bExisted;
+					FoundObjects.Add(SubobjectId);
+				}
+			});
+	}
 
 	// add controller and it's component hierarchy if one is attached
 	if (bAddController)
diff --git a/Engine/Plugins/Animation/GameplayInsights/Source/RewindDebugger/Private/RewindDebuggerObjectTrack.h b/Engine/Plugins/Animation/GameplayInsights/Source/RewindDebugger/Private/RewindDebuggerObjectTrack.h
--- a/Engine/Plugins/Animation/GameplayInsights/Source/RewindDebugger/Private/RewindDebuggerObjectTrack.h
+++ b/Engine/Plugins/Animation/GameplayInsights/Source/RewindDebugger/Private/RewindDebuggerObjectTrack.h
@@ -15,6 +15,7 @@ public:
 		: ObjectName(InObjectName)
 		, ObjectId(InObjectId)
 		, bAddController(bInAddController)
+		, bDisplayNameValid(false)
 	{
 		ExistenceRange = MakeShared<SSegmentedTimelineView::FSegmentData>();
 		ExistenceRange->Segments.SetNumUninitialized(1);
@@ -29,10 +30,11 @@ private:
 	
 	virtual FName GetNameInternal() const override { return ""; }
 	virtual FSlateIcon GetIconInternal() override { return Icon; }
-	virtual FText GetDisplayNameInternal() const override { return FText::FromString(ObjectName); }
+	virtual FText GetDisplayNameInternal() const override;
 	virtual uint64 GetObjectIdInternal() const override { return ObjectId; }
 	virtual bool HasDebugDataInternal() const override { return false; }
 	
+	mutable FText DisplayName;
 	FString ObjectName;
 	FSlateIcon Icon;
 	TSharedPtr<SSegmentedTimelineView::FSegmentData> ExistenceRange;
@@ -40,6 +42,7 @@ private:
 	TArray<TSharedPtr<FRewindDebuggerTrack>> Children;
 
 	bool bAddController;
+	mutable bool bDisplayNameValid;
 };
 
 }
\ No newline at end of file
