Fix issue introduced by object hierarchy optimization:
  Using Recording Time to filter out objects doesn't work correctly if they were created in a different recorded session

diff --git a/Engine/Plugins/Animation/GameplayInsights/Source/GameplayInsights/Private/GameplayProvider.cpp b/Engine/Plugins/Animation/GameplayInsights/Source/GameplayInsights/Private/GameplayProvider.cpp
--- a/Engine/Plugins/Animation/GameplayInsights/Source/GameplayInsights/Private/GameplayProvider.cpp
+++ b/Engine/Plugins/Animation/GameplayInsights/Source/GameplayInsights/Private/GameplayProvider.cpp
@@ -529,6 +529,18 @@ uint64 FGameplayProvider::FindPossessingController(uint64 PawnId, double Time) c
 	return ControllerId;
 }
 
+TRange<double> FGameplayProvider::GetObjectTraceLifetime(uint64 ObjectId) const
+{
+	if (const uint64* FoundIndex = ActiveObjectLifetimes.Find(ObjectId))
+	{
+		return TRange<double>(ObjectLifetimes.GetEventStartTime(*FoundIndex), ObjectLifetimes.GetEventEndTime(*FoundIndex));
+	}
+	else
+	{
+		return TRange<double>(0, 0);
+	}
+}
+
 TRange<double> FGameplayProvider::GetObjectRecordingLifetime(uint64 ObjectId) const
 {
 	if (const uint64 *FoundIndex = ActiveObjectRecordingLifetimes.Find(ObjectId))
diff --git a/Engine/Plugins/Animation/GameplayInsights/Source/GameplayInsights/Private/GameplayProvider.h b/Engine/Plugins/Animation/GameplayInsights/Source/GameplayInsights/Private/GameplayProvider.h
--- a/Engine/Plugins/Animation/GameplayInsights/Source/GameplayInsights/Private/GameplayProvider.h
+++ b/Engine/Plugins/Animation/GameplayInsights/Source/GameplayInsights/Private/GameplayProvider.h
@@ -90,7 +90,11 @@ public:
 	/** Search PawnPossession timeline to find the Controller Object Id for a Pawn (or 0 if none)*/
 	virtual uint64 FindPossessingController(uint64 Pawn, double Time) const override;
 	
-	/** find the time range for which an object existed */
+
+	/** find the Trace time range for which an object existed */
+	virtual TRange<double> GetObjectTraceLifetime(uint64 ObjectId) const override;
+
+	/** find the Recording time range for which an object existed */
 	virtual TRange<double> GetObjectRecordingLifetime(uint64 ObjectId) const override;
 
 private:
diff --git a/Engine/Plugins/Animation/GameplayInsights/Source/GameplayInsights/Public/IGameplayProvider.h b/Engine/Plugins/Animation/GameplayInsights/Source/GameplayInsights/Public/IGameplayProvider.h
--- a/Engine/Plugins/Animation/GameplayInsights/Source/GameplayInsights/Public/IGameplayProvider.h
+++ b/Engine/Plugins/Animation/GameplayInsights/Source/GameplayInsights/Public/IGameplayProvider.h
@@ -128,5 +128,6 @@ public:
 	virtual const RecordingInfoTimeline* GetRecordingInfo(uint32 RecordingId) const = 0; 
 	virtual void ReadViewTimeline(TFunctionRef<void(const ViewTimeline&)> Callback) const = 0;
 	virtual uint64 FindPossessingController(uint64 Pawn, double Time) const = 0;
+	virtual TRange<double> GetObjectTraceLifetime(uint64 ObjectId) const = 0;
 	virtual TRange<double> GetObjectRecordingLifetime(uint64 ObjectId) const = 0;
 };
\ No newline at end of file
diff --git a/Engine/Plugins/Animation/GameplayInsights/Source/RewindDebugger/Private/RewindDebuggerObjectTrack.cpp b/Engine/Plugins/Animation/GameplayInsights/Source/RewindDebugger/Private/RewindDebuggerObjectTrack.cpp
--- a/Engine/Plugins/Animation/GameplayInsights/Source/RewindDebugger/Private/RewindDebuggerObjectTrack.cpp
+++ b/Engine/Plugins/Animation/GameplayInsights/Source/RewindDebugger/Private/RewindDebuggerObjectTrack.cpp
@@ -174,12 +174,12 @@ bool FRewindDebuggerObjectTrack::UpdateInternal()
 	{
 		TRACE_CPUPROFILER_EVENT_SCOPE(FRewindDebuggerObjectTrack::UpdateInternal_AddChildComponents);
 
-		TRange<double> ViewRange = RewindDebugger->GetCurrentViewRange();
+		TRange<double> TraceRange = RewindDebugger->GetCurrentTraceRange();
 
-		GameplayProvider->EnumerateSubobjects(ObjectId, [this, &FoundObjects, &bChanged,&ViewRange, GameplayProvider](uint64 SubobjectId)
+		GameplayProvider->EnumerateSubobjects(ObjectId, [this, &FoundObjects, &bChanged,&TraceRange, GameplayProvider](uint64 SubobjectId)
 			{
 				TRange<double> Lifetime = GameplayProvider->GetObjectRecordingLifetime(SubobjectId);
-				TRange<double> Overlap = TRange<double>::Intersection(Lifetime, ViewRange);
+				TRange<double> Overlap = TRange<double>::Intersection(Lifetime, TraceRange);
 				// only display the track if the lifetime of the object and the view range overlap
 				if (!Overlap.IsEmpty())
 				{
