Add double click handler support for rewind debugger tracks, and add a double click handler for blend curve tracks

diff --git a/Engine/Plugins/Animation/GameplayInsights/Source/GameplayInsights/Private/BlendWeightsTrack.cpp b/Engine/Plugins/Animation/GameplayInsights/Source/GameplayInsights/Private/BlendWeightsTrack.cpp
--- a/Engine/Plugins/Animation/GameplayInsights/Source/GameplayInsights/Private/BlendWeightsTrack.cpp
+++ b/Engine/Plugins/Animation/GameplayInsights/Source/GameplayInsights/Private/BlendWeightsTrack.cpp
@@ -12,6 +12,8 @@
 #include "Animation/AnimInstance.h"
 #include "Animation/AnimSequence.h"
 #include "Animation/BlendSpace.h"
+#include "Editor.h"
+#include "Subsystems/AssetEditorSubsystem.h"
 #endif
 
 #define LOCTEXT_NAMESPACE "BlendWeightsTrack"
@@ -335,6 +337,25 @@ TSharedPtr<SWidget> FBlendWeightTrack::GetDetailsViewInternal()
 	return BlendWeightsView;
 }
 
+bool FBlendWeightTrack::HandleDoubleClickInternal()
+{
+#if WITH_EDITOR
+	IRewindDebugger* RewindDebugger = IRewindDebugger::Instance();
+	if (const TraceServices::IAnalysisSession* AnalysisSession = RewindDebugger->GetAnalysisSession())
+	{
+		TraceServices::FAnalysisSessionReadScope SessionReadScope(*AnalysisSession);
+
+		const FGameplayProvider* GameplayProvider = AnalysisSession->ReadProvider<FGameplayProvider>(FGameplayProvider::ProviderName);
+
+		const FObjectInfo& AssetInfo = GameplayProvider->GetObjectInfo(GetAssetId());
+
+		GEditor->GetEditorSubsystem<UAssetEditorSubsystem>()->OpenEditorForAsset(AssetInfo.PathName);
+
+		return true;
+	}
+#endif
+	return false;
+}
 
 FName FBlendWeightsTrackCreator::GetTargetTypeNameInternal() const
 {
diff --git a/Engine/Plugins/Animation/GameplayInsights/Source/GameplayInsights/Private/BlendWeightsTrack.h b/Engine/Plugins/Animation/GameplayInsights/Source/GameplayInsights/Private/BlendWeightsTrack.h
--- a/Engine/Plugins/Animation/GameplayInsights/Source/GameplayInsights/Private/BlendWeightsTrack.h
+++ b/Engine/Plugins/Animation/GameplayInsights/Source/GameplayInsights/Private/BlendWeightsTrack.h
@@ -41,6 +41,7 @@ private:
 	virtual FName GetNameInternal() const override { return "BlendWeight"; }
 	virtual FText GetDisplayNameInternal() const override { return CurveName; }
 	virtual uint64 GetObjectIdInternal() const override { return ObjectId; }
+	virtual bool HandleDoubleClickInternal() override;
 
 	TSharedPtr<SCurveTimelineView::FTimelineCurveData> GetCurveData() const;
 	
diff --git a/Engine/Plugins/Animation/GameplayInsights/Source/RewindDebugger/Private/RewindDebugger.cpp b/Engine/Plugins/Animation/GameplayInsights/Source/RewindDebugger/Private/RewindDebugger.cpp
--- a/Engine/Plugins/Animation/GameplayInsights/Source/RewindDebugger/Private/RewindDebugger.cpp
+++ b/Engine/Plugins/Animation/GameplayInsights/Source/RewindDebugger/Private/RewindDebugger.cpp
@@ -1028,6 +1028,11 @@ void FRewindDebugger::ComponentDoubleClicked(TSharedPtr<RewindDebugger::FRewindD
 	}
 	
 	SelectedTrack = SelectedObject;
+
+	if (SelectedTrack->HandleDoubleClick())
+	{
+		return;
+	}
 	
 	IModularFeatures& ModularFeatures = IModularFeatures::Get();
 	static const FName HandlerFeatureName = IRewindDebuggerDoubleClickHandler::ModularFeatureName;
diff --git a/Engine/Source/Editor/RewindDebuggerInterface/Public/RewindDebuggerTrack.h b/Engine/Source/Editor/RewindDebuggerInterface/Public/RewindDebuggerTrack.h
--- a/Engine/Source/Editor/RewindDebuggerInterface/Public/RewindDebuggerTrack.h
+++ b/Engine/Source/Editor/RewindDebuggerInterface/Public/RewindDebuggerTrack.h
@@ -87,6 +87,12 @@ public:
 		return HasDebugDataInternal();
 	}
 
+	// Called when a track is double clicked.  Returns true if the track handled the double click
+	bool HandleDoubleClick()
+	{
+		return HandleDoubleClickInternal();
+	}
+
 	bool IsVisible() const { return bVisible; }
 	void SetIsVisible(bool bInIsVisible) { bVisible = bInIsVisible; }
 
@@ -101,6 +107,7 @@ private:
 	virtual uint64 GetObjectIdInternal() const { return 0; }
 	virtual bool HasDebugDataInternal() const { return true; }
 	virtual void IterateSubTracksInternal(TFunction<void(TSharedPtr<FRewindDebuggerTrack> SubTrack)> IteratorFunction) { }
+	virtual bool HandleDoubleClickInternal() { return false; }
 
 	bool bExpanded;
 	bool bVisible;
