Fixed compile error when building without OBJECT_TRACE_ENABLED

diff --git a/Engine/Plugins/Animation/GameplayInsights/Source/GameplayInsights/Private/PoseWatchTrack.cpp b/Engine/Plugins/Animation/GameplayInsights/Source/GameplayInsights/Private/PoseWatchTrack.cpp
--- a/Engine/Plugins/Animation/GameplayInsights/Source/GameplayInsights/Private/PoseWatchTrack.cpp
+++ b/Engine/Plugins/Animation/GameplayInsights/Source/GameplayInsights/Private/PoseWatchTrack.cpp
@@ -29,6 +29,7 @@ FPoseWatchTrack::FPoseWatchTrack(uint64 InObjectId, const FPoseWatchTrack::FPose
 	EnabledSegments = MakeShared<SSegmentedTimelineView::FSegmentData>();
 	Icon = UPoseWatchPoseElement::StaticGetIcon();
 
+#if OBJECT_TRACE_ENABLED
 	if (UObject* ObjectInstance = FObjectTrace::GetObjectFromId(ObjectId))
 	{
 		if (UAnimInstance* AnimInstance = Cast<UAnimInstance>(ObjectInstance))
@@ -47,6 +48,7 @@ FPoseWatchTrack::FPoseWatchTrack(uint64 InObjectId, const FPoseWatchTrack::FPose
 			}
 		}
 	}
+#endif
 }
 
 FText FPoseWatchTrack::GetDisplayNameInternal() const
@@ -183,6 +185,7 @@ FPoseWatchesTrack::FPoseWatchesTrack(uint64 InObjectId)
 {
 	Icon = UPoseWatchPoseElement::StaticGetIcon();
 
+#if OBJECT_TRACE_ENABLED
 	if (UObject* ObjectInstance = FObjectTrace::GetObjectFromId(ObjectId))
 	{
 		if (UAnimInstance* AnimInstance = Cast<UAnimInstance>(ObjectInstance))
@@ -193,6 +196,7 @@ FPoseWatchesTrack::FPoseWatchesTrack(uint64 InObjectId)
 			}
 		}
 	}
+#endif
 }
 
 
