Added tooltip for SCurveTimelineView that allows to preview curve value at cursor location

diff --git a/Engine/Plugins/Animation/GameplayInsights/Source/GameplayInsights/Private/SCurveTimelineView.cpp b/Engine/Plugins/Animation/GameplayInsights/Source/GameplayInsights/Private/SCurveTimelineView.cpp
--- a/Engine/Plugins/Animation/GameplayInsights/Source/GameplayInsights/Private/SCurveTimelineView.cpp
+++ b/Engine/Plugins/Animation/GameplayInsights/Source/GameplayInsights/Private/SCurveTimelineView.cpp
@@ -4,6 +4,8 @@
 #include "Framework/Application/SlateApplication.h"
 #include "Rendering/DrawElements.h"
 #include "SSimpleTimeSlider.h"
+#include "Widgets/Text/STextBlock.h"
+#include "Widgets/SToolTip.h"
 
 #define LOCTEXT_NAMESPACE "SCurveTimelineView"
 
@@ -28,6 +30,87 @@ int32 SCurveTimelineView::OnPaint( const FPaintArgs& Args, const FGeometry& Allo
 	return FMath::Max( NewLayer, SCompoundWidget::OnPaint(Args, AllottedGeometry, MyCullingRect, OutDrawElements, NewLayer, InWidgetStyle, ShouldBeEnabled( bParentEnabled ) ) );
 }
 
+FReply SCurveTimelineView::OnMouseMove(const FGeometry& InMyGeometry, const FPointerEvent& InMouseEvent)
+{
+	UpdateCurveToolTip(InMyGeometry, InMouseEvent);
+	return FReply::Unhandled();
+}
+
+void SCurveTimelineView::UpdateCurveToolTip(const FGeometry& InMyGeometry, const FPointerEvent& InMouseEvent)
+{
+	if (InMyGeometry.IsUnderLocation(InMouseEvent.GetScreenSpacePosition()))
+	{
+		if (const TSharedPtr<FTimelineCurveData> Curve = CurveData.Get())
+		{
+			// Mouse position in widget space
+			const FVector2D HitPosition = InMyGeometry.AbsoluteToLocal(InMouseEvent.GetScreenSpacePosition());
+
+			// Range helper struct
+			const TRange<double> DebugTimeRange = ViewRange.Get();
+			const SSimpleTimeSlider::FScrubRangeToScreen RangeToScreen( DebugTimeRange, InMyGeometry.GetLocalSize() );
+		
+			// Mouse position from widget space to curve input space
+			const double TargetTime = RangeToScreen.LocalXToInput(HitPosition.X);
+
+			// Get curve value at given time
+			TArray<FTimelineCurveData::CurvePoint> & CurvePoints = Curve->Points;
+			const int NumPoints = CurvePoints.Num();
+		
+			if (NumPoints > 0)
+			{
+				for (int i = 1; i < NumPoints; ++i)
+				{
+					const FTimelineCurveData::CurvePoint& Point1 = CurvePoints[i-1];
+					const FTimelineCurveData::CurvePoint& Point2 = CurvePoints[i];
+
+					// Find points that contain mouse hit-point time
+					if ( Point1.Time >= TargetTime && TargetTime <= Point2.Time)
+					{
+						// Choose point with the smallest delta
+						const float Delta1 = abs(TargetTime - Point1.Time);
+						const float Delta2 = abs(TargetTime - Point2.Time);
+
+						// Get closest point
+						const FTimelineCurveData::CurvePoint & TargetPoint = Delta1 < Delta2 ? Point1 : Point2;
+
+						// Tooltip text formatting
+						FNumberFormattingOptions FormattingOptions;
+						FormattingOptions.MaximumFractionalDigits = 3;
+						CurveToolTipOutputText = FText::Format(LOCTEXT("CurveToolTipValueFormat", "Value: {0}"), FText::AsNumber(TargetPoint.Value, &FormattingOptions));
+						CurveToolTipInputText = FText::Format(LOCTEXT("CurveToolTipTimeFormat", "Time: {0}"), FText::AsNumber(TargetPoint.Time, &FormattingOptions));
+				
+						// Update tooltip info
+						if (CurveToolTip.IsValid() == false)
+						{
+							SetToolTip(
+								SAssignNew(CurveToolTip, SToolTip)
+								.BorderImage( FCoreStyle::Get().GetBrush( "ToolTip.BrightBackground" ) )
+								[
+									SNew(SVerticalBox)
+									+ SVerticalBox::Slot()
+									[
+										SNew(STextBlock)
+										.Text(this, &SCurveTimelineView::GetCurveToolTipInputText)
+										.Font(FCoreStyle::Get().GetFontStyle("ToolTip.LargerFont"))
+										.ColorAndOpacity(FLinearColor::Black)
+									]
+									+ SVerticalBox::Slot()
+									[
+										SNew(STextBlock)
+										.Text(this, &SCurveTimelineView::GetCurveToolTipOutputText)
+										.Font(FCoreStyle::Get().GetFontStyle("ToolTip.LargerFont"))
+										.ColorAndOpacity(FLinearColor::Black)
+									]
+								]);
+						}
+				
+						break;
+					}
+				}
+			}
+		}
+	}
+}
 
 FVector2D SCurveTimelineView::ComputeDesiredSize( float ) const
 {
@@ -174,13 +257,7 @@ int32 SCurveTimelineView::PaintCurve(const FGeometry& AllottedGeometry, const FS
 								LineColor,
 								false
 								);
-
-
-								
 		}
-	
-
-
 	}
 
 	return LayerId;
diff --git a/Engine/Plugins/Animation/GameplayInsights/Source/GameplayInsights/Private/SCurveTimelineView.h b/Engine/Plugins/Animation/GameplayInsights/Source/GameplayInsights/Private/SCurveTimelineView.h
--- a/Engine/Plugins/Animation/GameplayInsights/Source/GameplayInsights/Private/SCurveTimelineView.h
+++ b/Engine/Plugins/Animation/GameplayInsights/Source/GameplayInsights/Private/SCurveTimelineView.h
@@ -9,7 +9,7 @@
 class FPaintArgs;
 class FSlateWindowElementList;
 class FTimeSliderController;
-
+class SToolTip;
 
 class SCurveTimelineView : public SCompoundWidget
 {
@@ -67,12 +67,17 @@ public:
 		FixedRangeMax = Max;
 	}
 
+	FText GetCurveToolTipInputText() const { return CurveToolTipInputText; }
+	FText GetCurveToolTipOutputText() const { return CurveToolTipOutputText; }
+	
 protected:
 	// SWidget interface
 	virtual FVector2D ComputeDesiredSize(float) const override;
 	virtual int32 OnPaint( const FPaintArgs& Args, const FGeometry& AllottedGeometry, const FSlateRect& MyCullingRect, FSlateWindowElementList& OutDrawElements, int32 LayerId, const FWidgetStyle& InWidgetStyle, bool bParentEnabled ) const override;
+	virtual FReply OnMouseMove(const FGeometry& MyGeometry, const FPointerEvent& MouseEvent) override;
 	
 	int32 PaintCurve(const FGeometry& AllottedGeometry, const FSlateRect& MyCullingRect, FSlateWindowElementList& OutDrawElements, int32 LayerId, const FWidgetStyle& InWidgetStyle, bool bParentEnabled ) const;
+	void UpdateCurveToolTip(const FGeometry& MyGeometry, const FPointerEvent& MouseEvent);
 	
 	TAttribute<TRange<double>> ViewRange;
 	TAttribute<TSharedPtr<FTimelineCurveData>> CurveData;
@@ -81,6 +86,10 @@ protected:
 	TAttribute<FLinearColor> FillColor;
 	TAttribute<bool> RenderFill;
 	
+	TSharedPtr<SToolTip> CurveToolTip;	/** The tooltip control for the curve. */
+	FText CurveToolTipInputText;		
+	FText CurveToolTipOutputText;
+	
 	const FSlateBrush* WhiteBrush;
 
 	bool bUseFixedRange;
