Fix weapons and other attached actors while scrubbing/replaying in Rewind Debugger

diff --git a/Engine/Plugins/Animation/GameplayInsights/Source/RewindDebugger/Private/RewindDebugger.cpp b/Engine/Plugins/Animation/GameplayInsights/Source/RewindDebugger/Private/RewindDebugger.cpp
--- a/Engine/Plugins/Animation/GameplayInsights/Source/RewindDebugger/Private/RewindDebugger.cpp
+++ b/Engine/Plugins/Animation/GameplayInsights/Source/RewindDebugger/Private/RewindDebugger.cpp
@@ -773,9 +773,10 @@ void FRewindDebugger::Tick(float DeltaTime)
 															MeshComponentsToReset.Add(ObjectId, ResetData);
 														}
 
-														// todo: we need to take into account tick order requirements for attached objects here
 														MeshComponent->SetWorldTransform(ComponentWorldTransform, false, nullptr, ETeleportType::TeleportPhysics);
 														MeshComponent->SetForcedLOD(PoseMessage.LodIndex + 1);
+														MeshComponent->UpdateChildTransforms(EUpdateTransformFlags::None, ETeleportType::TeleportPhysics);
+
 														return TraceServices::EEventEnumerate::Stop;
 													});
 											});
