fix code which disables cpu profiling to only disable if not already connected.

diff --git a/Engine/Plugins/Animation/GameplayInsights/Source/GameplayInsights/Private/GameplayInsightsModule.cpp b/Engine/Plugins/Animation/GameplayInsights/Source/GameplayInsights/Private/GameplayInsightsModule.cpp
--- a/Engine/Plugins/Animation/GameplayInsights/Source/GameplayInsights/Private/GameplayInsightsModule.cpp
+++ b/Engine/Plugins/Animation/GameplayInsights/Source/GameplayInsights/Private/GameplayInsightsModule.cpp
@@ -424,6 +424,23 @@ void FGameplayInsightsModule::StartTrace()
 {
 	if (!bTraceStarted)
 	{
+		if (!FTraceAuxiliary::IsConnected())
+		{
+			// cpu tracing is enabled by default, but not connected.
+			// when not connected, disable all channels initially to avoid a whole bunch of extra cpu, memory, and disk overhead from processing all the extra default trace channels
+			////////////////////////////////////////////////////////////////////////////////
+			UE::Trace::EnumerateChannels([](const ANSICHAR* ChannelName, bool bEnabled, void*)
+				{
+					if (bEnabled)
+					{
+						FString ChannelNameFString(ChannelName);
+						UE::Trace::ToggleChannel(ChannelNameFString.GetCharArray().GetData(), false);
+					}
+				}
+				, nullptr);
+		}
+
+
 		bTraceStarted = FTraceAuxiliary::Start(
 			FTraceAuxiliary::EConnectionType::Network,
 			TEXT("127.0.0.1"),
diff --git a/Engine/Plugins/Animation/GameplayInsights/Source/RewindDebugger/Private/RewindDebugger.cpp b/Engine/Plugins/Animation/GameplayInsights/Source/RewindDebugger/Private/RewindDebugger.cpp
--- a/Engine/Plugins/Animation/GameplayInsights/Source/RewindDebugger/Private/RewindDebugger.cpp
+++ b/Engine/Plugins/Animation/GameplayInsights/Source/RewindDebugger/Private/RewindDebugger.cpp
@@ -24,6 +24,7 @@
 #include "LevelEditor.h"
 #include "RewindDebuggerModule.h"
 #include "Engine/PoseWatch.h"
+#include "ProfilingDebugging/TraceAuxiliary.h"
 
 static void IterateExtensions(TFunction<void(IRewindDebuggerExtension* Extension)> IteratorFunction)
 {
@@ -52,9 +53,6 @@ FRewindDebugger::FRewindDebugger()  :
 	bTargetActorPositionValid(false),
 	bIsDetailsPanelOpen(true)
 {
-	UE::Trace::ToggleChannel(TEXT("Cpu"), false);
-	UE::Trace::ToggleChannel(TEXT("Gpu"), false);
-
 	RecordingDuration.Set(0);
 
 	if (GEditor->bIsSimulatingInEditor || GEditor->PlayWorld)
