Disable eyedropper actor selector in Rewind Debugger when player control is not ejected, and add a tooltip

diff --git a/Engine/Plugins/Animation/GameplayInsights/Source/RewindDebugger/Private/SRewindDebugger.cpp b/Engine/Plugins/Animation/GameplayInsights/Source/RewindDebugger/Private/SRewindDebugger.cpp
--- a/Engine/Plugins/Animation/GameplayInsights/Source/RewindDebugger/Private/SRewindDebugger.cpp
+++ b/Engine/Plugins/Animation/GameplayInsights/Source/RewindDebugger/Private/SRewindDebugger.cpp
@@ -25,6 +25,7 @@
 #include "Widgets/Input/SSearchBox.h"
 #include "Widgets/Layout/SExpandableArea.h"
 #include "Widgets/Layout/SScrollBox.h"
+#include "Kismet2/DebuggerCommands.h"
 #include "RewindDebuggerSettings.h"
 
 #define LOCTEXT_NAMESPACE "SRewindDebugger"
@@ -401,6 +402,12 @@ void SRewindDebugger::Construct(const FArguments& InArgs, TSharedRef<FUICommandL
 						SNew(SButton)
 							.ButtonStyle(FAppStyle::Get(), "SimpleButton")
 							.OnClicked(this, &SRewindDebugger::OnSelectActorClicked)
+							.ToolTipText(LOCTEXT("SelectActorTooltip", "Select Target Actor in Scene (Eject player control first)"))
+							.IsEnabled_Lambda([]()
+								{
+									return !FPlayWorldCommandCallbacks::IsInPIE();
+								}
+							)
 						[
 							SNew(SImage)
 							.Image(FRewindDebuggerStyle::Get().GetBrush("RewindDebugger.SelectActor"))
