Add fallback to OpenEditorForAsset (String path version) to find unsaved assets

diff --git a/Engine/Plugins/Animation/GameplayInsights/Source/GameplayInsights/Private/SVariantValueView.cpp b/Engine/Plugins/Animation/GameplayInsights/Source/GameplayInsights/Private/SVariantValueView.cpp
--- a/Engine/Plugins/Animation/GameplayInsights/Source/GameplayInsights/Private/SVariantValueView.cpp
+++ b/Engine/Plugins/Animation/GameplayInsights/Source/GameplayInsights/Private/SVariantValueView.cpp
@@ -192,17 +192,17 @@ static TSharedRef<SWidget> MakeVariantValueWidget(const TraceServices::IAnalysis
 				.ToolTipText(FText::Format(LOCTEXT("ClassHyperlinkTooltipFormat", "Open class '{0}'"), FText::FromString(ClassInfo.PathName)))
 				.OnNavigate_Lambda([ClassInfo]()
 				{
-					if (LoadPackage(NULL, ClassInfo.PathName, LOAD_NoRedirects)) // check if it's found as an asset
+					if (UClass* FoundClass = UClass::TryFindTypeSlow<UClass>(ClassInfo.Name))
 					{
-						GEditor->GetEditorSubsystem<UAssetEditorSubsystem>()->OpenEditorForAsset(ClassInfo.PathName);
-					}
-					else // it must be a native class
-					{
-						if (UClass* FoundClass = UClass::TryFindTypeSlow<UClass>(ClassInfo.Name))
+						if (FoundClass->IsNative())
 						{
+							// for native classes navigatte to source code
 							FSourceCodeNavigation::NavigateToClass(FoundClass);
 						}
 					}
+
+					// for non-native classes, try opening the asset
+					GEditor->GetEditorSubsystem<UAssetEditorSubsystem>()->OpenEditorForAsset(ClassInfo.PathName);
 				});
 #else
 			return 
diff --git a/Engine/Source/Editor/UnrealEd/Private/Subsystems/AssetEditorSubsystem.cpp b/Engine/Source/Editor/UnrealEd/Private/Subsystems/AssetEditorSubsystem.cpp
--- a/Engine/Source/Editor/UnrealEd/Private/Subsystems/AssetEditorSubsystem.cpp
+++ b/Engine/Source/Editor/UnrealEd/Private/Subsystems/AssetEditorSubsystem.cpp
@@ -603,8 +603,8 @@ void UAssetEditorSubsystem::OpenEditorForAsset(const FSoftObjectPath& AssetPath)
 void UAssetEditorSubsystem::OpenEditorForAsset(const FString& AssetPathName)
 {
 	// An asset needs loading
-	UPackage* Package = LoadPackage(NULL, *AssetPathName, LOAD_NoRedirects);
 
+	UPackage* Package = LoadPackage(NULL, *AssetPathName, LOAD_NoRedirects);
 	if (Package)
 	{
 		Package->FullyLoad();
@@ -612,7 +612,17 @@ void UAssetEditorSubsystem::OpenEditorForAsset(const FString& AssetPathName)
 		FString AssetName = FPaths::GetBaseFilename(AssetPathName);
 		UObject* Object = FindObject<UObject>(Package, *AssetName);
 
-		if (Object != NULL)
+		if (Object != nullptr)
+		{
+			OpenEditorForAsset(Object);
+		}
+	}
+	else
+	{
+		// fallback for unsaved assets
+		UObject* Object = FindObject<UObject>(nullptr, *AssetPathName);
+		
+		if (Object != nullptr)
 		{
 			OpenEditorForAsset(Object);
 		}
