Null check fix for difficult to reproduce rewind debugger crash

diff --git a/Engine/Plugins/Animation/GameplayInsights/Source/GameplayInsights/Private/GameplayProvider.cpp b/Engine/Plugins/Animation/GameplayInsights/Source/GameplayInsights/Private/GameplayProvider.cpp
--- a/Engine/Plugins/Animation/GameplayInsights/Source/GameplayInsights/Private/GameplayProvider.cpp
+++ b/Engine/Plugins/Animation/GameplayInsights/Source/GameplayInsights/Private/GameplayProvider.cpp
@@ -197,7 +197,10 @@ void FGameplayProvider::EnumerateObjects(double StartTime, double EndTime, TFunc
 		[this, Callback](double InStartTime, double InEndTime, uint32 InDepth, const FObjectExistsMessage& ExistsMessage)
 		{
 			checkSlow(ObjectIdToIndexMap.Contains(ExistsMessage.ObjectId));
-			Callback(ObjectInfos[ObjectIdToIndexMap[ExistsMessage.ObjectId]]);
+			if (const int32* ObjectInfoIndex = ObjectIdToIndexMap.Find(ExistsMessage.ObjectId))
+			{
+				Callback(ObjectInfos[*ObjectInfoIndex]);
+			}
 			return TraceServices::EEventEnumerate::Continue;
 		});
 }
@@ -531,6 +534,8 @@ uint64 FGameplayProvider::FindPossessingController(uint64 PawnId, double Time) c
 
 TRange<double> FGameplayProvider::GetObjectTraceLifetime(uint64 ObjectId) const
 {
+	Session.ReadAccessCheck();
+
 	if (const uint64* FoundIndex = ActiveObjectLifetimes.Find(ObjectId))
 	{
 		return TRange<double>(ObjectLifetimes.GetEventStartTime(*FoundIndex), ObjectLifetimes.GetEventEndTime(*FoundIndex));
@@ -543,6 +548,8 @@ TRange<double> FGameplayProvider::GetObjectTraceLifetime(uint64 ObjectId) const
 
 TRange<double> FGameplayProvider::GetObjectRecordingLifetime(uint64 ObjectId) const
 {
+	Session.ReadAccessCheck();
+
 	if (const uint64 *FoundIndex = ActiveObjectRecordingLifetimes.Find(ObjectId))
 	{
 		return TRange<double>(ObjectRecordingLifetimes.GetEventStartTime(*FoundIndex), ObjectRecordingLifetimes.GetEventEndTime(*FoundIndex));
