Fix Rewind Debugger issue with selecting debug target before recording

diff --git a/Engine/Plugins/Animation/GameplayInsights/Source/RewindDebugger/Private/RewindDebugger.cpp b/Engine/Plugins/Animation/GameplayInsights/Source/RewindDebugger/Private/RewindDebugger.cpp
--- a/Engine/Plugins/Animation/GameplayInsights/Source/RewindDebugger/Private/RewindDebugger.cpp
+++ b/Engine/Plugins/Animation/GameplayInsights/Source/RewindDebugger/Private/RewindDebugger.cpp
@@ -238,11 +238,6 @@ void FRewindDebugger::RefreshDebugTracks()
 		{
 			uint64 TargetActorId = GetTargetActorId();
 
-			if (TargetActorId == 0)
-			{
-				return;
-			}
-
 			bool bChanged = false;
 
 			// add actor (even if it isn't found in the gameplay provider)
@@ -254,10 +249,12 @@ void FRewindDebugger::RefreshDebugTracks()
 			}
 			else
 			{
-				if (DebugTracks[0]->GetDisplayName().ToString() != DebugTargetActor.Get() || DebugTracks[0]->GetObjectId() != TargetActorId)
+				FString DebugTargetActorName = DebugTargetActor.Get();
+
+				if (DebugTracks[0]->GetDisplayName().ToString() != DebugTargetActorName || DebugTracks[0]->GetObjectId() != TargetActorId)
 				{
 					bChanged = true;
-					DebugTracks[0] = MakeShared<RewindDebugger::FRewindDebuggerObjectTrack>(TargetActorId, DebugTargetActor.Get());
+					DebugTracks[0] = MakeShared<RewindDebugger::FRewindDebuggerObjectTrack>(TargetActorId, DebugTargetActorName);
 				}
 			}
 
diff --git a/Engine/Plugins/Animation/GameplayInsights/Source/RewindDebugger/Private/RewindDebuggerObjectTrack.cpp b/Engine/Plugins/Animation/GameplayInsights/Source/RewindDebugger/Private/RewindDebuggerObjectTrack.cpp
--- a/Engine/Plugins/Animation/GameplayInsights/Source/RewindDebugger/Private/RewindDebuggerObjectTrack.cpp
+++ b/Engine/Plugins/Animation/GameplayInsights/Source/RewindDebugger/Private/RewindDebuggerObjectTrack.cpp
@@ -58,23 +58,30 @@ FText FRewindDebuggerObjectTrack::GetDisplayNameInternal() const
 {
 	if (!bDisplayNameValid)
 	{
-		IRewindDebugger* RewindDebugger = IRewindDebugger::Instance();
-		const TraceServices::IAnalysisSession* Session = RewindDebugger->GetAnalysisSession();
-		TraceServices::FAnalysisSessionReadScope SessionReadScope(*Session);
-
-		const IGameplayProvider* GameplayProvider = Session->ReadProvider<IGameplayProvider>("GameplayProvider");
+		if (ObjectId == 0)
+		{
+			DisplayName = FText::FromString(ObjectName);
+		}
+		else
+		{
+			IRewindDebugger* RewindDebugger = IRewindDebugger::Instance();
+			const TraceServices::IAnalysisSession* Session = RewindDebugger->GetAnalysisSession();
+			TraceServices::FAnalysisSessionReadScope SessionReadScope(*Session);
 
-		const FObjectInfo& ObjectInfo = GameplayProvider->GetObjectInfo(ObjectId);
+			const IGameplayProvider* GameplayProvider = Session->ReadProvider<IGameplayProvider>("GameplayProvider");
 
-		DisplayName = FText::FromString(ObjectInfo.Name);
+			const FObjectInfo& ObjectInfo = GameplayProvider->GetObjectInfo(ObjectId);
 
-		if (const FWorldInfo* WorldInfo = GameplayProvider->FindWorldInfoFromObject(ObjectId))
-		{
-			bool bIsServer = WorldInfo->NetMode == FWorldInfo::ENetMode::DedicatedServer;
+			DisplayName = FText::FromString(ObjectInfo.Name);
 
-			if (bIsServer)
+			if (const FWorldInfo* WorldInfo = GameplayProvider->FindWorldInfoFromObject(ObjectId))
 			{
-				DisplayName = FText::Format(NSLOCTEXT("RewindDebuggerTrack", " (Server)", "{0} (Server)"), FText::FromString(ObjectInfo.Name));
+				bool bIsServer = WorldInfo->NetMode == FWorldInfo::ENetMode::DedicatedServer;
+
+				if (bIsServer)
+				{
+					DisplayName = FText::Format(NSLOCTEXT("RewindDebuggerTrack", " (Server)", "{0} (Server)"), FText::FromString(ObjectInfo.Name));
+				}
 			}
 		}
 
