Add teleport flags when rewind debugger moves skeletal mesh components back to their original locations when you resume playing

diff --git a/Engine/Plugins/Animation/GameplayInsights/Source/RewindDebugger/Private/RewindDebugger.cpp b/Engine/Plugins/Animation/GameplayInsights/Source/RewindDebugger/Private/RewindDebugger.cpp
--- a/Engine/Plugins/Animation/GameplayInsights/Source/RewindDebugger/Private/RewindDebugger.cpp
+++ b/Engine/Plugins/Animation/GameplayInsights/Source/RewindDebugger/Private/RewindDebugger.cpp
@@ -149,7 +149,7 @@ void FRewindDebugger::OnPIEResumed(bool bSimulating)
 	{
 		if (USkeletalMeshComponent* MeshComponent = MeshData.Value.Component.Get())
 		{
-			MeshComponent->SetRelativeTransform(MeshData.Value.RelativeTransform);
+			MeshComponent->SetRelativeTransform(MeshData.Value.RelativeTransform, false, nullptr, ETeleportType::TeleportPhysics);
 		}
 	}
 
@@ -163,7 +163,7 @@ void FRewindDebugger::OnPIESingleStepped(bool bSimulating)
 	{
 		if (USkeletalMeshComponent* MeshComponent = MeshData.Value.Component.Get())
 		{
-			MeshComponent->SetRelativeTransform(MeshData.Value.RelativeTransform);
+			MeshComponent->SetRelativeTransform(MeshData.Value.RelativeTransform, false, nullptr, ETeleportType::TeleportPhysics);
 		}
 	}
 
