Update icon for Rewind Debugger menu to a camera icon since it currently only shows camera options.

diff --git a/Engine/Plugins/Animation/GameplayInsights/Source/RewindDebugger/Private/SRewindDebugger.cpp b/Engine/Plugins/Animation/GameplayInsights/Source/RewindDebugger/Private/SRewindDebugger.cpp
--- a/Engine/Plugins/Animation/GameplayInsights/Source/RewindDebugger/Private/SRewindDebugger.cpp
+++ b/Engine/Plugins/Animation/GameplayInsights/Source/RewindDebugger/Private/SRewindDebugger.cpp
@@ -330,7 +330,7 @@ void SRewindDebugger::Construct(const FArguments& InArgs, TSharedRef<FUICommandL
 						.ButtonContent()
 					[
 						SNew(SImage)
-						.Image(FRewindDebuggerStyle::Get().GetBrush("RewindDebugger.MenuIcon"))
+						.Image(FAppStyle::Get().GetBrush("ClassIcon.CameraComponent"))
 					]
 				]
 				+SHorizontalBox::Slot().FillWidth(1.0)
