Rewind Debugger Track now support an overlay icon. Padding and Vertical centering adjustments.

diff --git a/Engine/Plugins/Animation/GameplayInsights/Source/RewindDebugger/Private/SRewindDebuggerComponentTree.cpp b/Engine/Plugins/Animation/GameplayInsights/Source/RewindDebugger/Private/SRewindDebuggerComponentTree.cpp
--- a/Engine/Plugins/Animation/GameplayInsights/Source/RewindDebugger/Private/SRewindDebuggerComponentTree.cpp
+++ b/Engine/Plugins/Animation/GameplayInsights/Source/RewindDebugger/Private/SRewindDebuggerComponentTree.cpp
@@ -8,6 +8,7 @@
 #include "Widgets/Images/SImage.h"
 
 #include "RewindDebugger.h"
+#include "Widgets/Images/SLayeredImage.h"
 
 #define LOCTEXT_NAMESPACE "SAnimationInsights"
 
@@ -24,7 +25,16 @@ SRewindDebuggerComponentTree::~SRewindDebuggerComponentTree()
 
 TSharedRef<ITableRow> SRewindDebuggerComponentTree::ComponentTreeViewGenerateRow(TSharedPtr<RewindDebugger::FRewindDebuggerTrack> InItem, const TSharedRef<STableViewBase>& OwnerTable)
 {
-	FSlateIcon ObjectIcon = InItem->GetIcon();
+	const FSlateIcon ObjectIcon = InItem->GetIcon();
+
+	const TSharedRef<SLayeredImage> LayeredIcons = SNew(SLayeredImage)
+				.DesiredSizeOverride(FVector2D(16, 16))
+				.Image(ObjectIcon.GetIcon());
+
+	if (ObjectIcon.GetOverlayIcon())
+	{
+		LayeredIcons->AddLayer(ObjectIcon.GetOverlayIcon());
+	}
 	
 	return 
 		SNew(STableRow<TSharedPtr<RewindDebugger::FRewindDebuggerTrack>>, OwnerTable)
@@ -32,10 +42,12 @@ TSharedRef<ITableRow> SRewindDebuggerComponentTree::ComponentTreeViewGenerateRow
 			SNew(SHorizontalBox)
 			+SHorizontalBox::Slot().AutoWidth().Padding(2)
 			[
-				SNew(SImage).DesiredSizeOverride(FVector2D(16,16))
-				.Image(ObjectIcon.GetIcon())
+				LayeredIcons
 			]
 			+SHorizontalBox::Slot()
+			.FillWidth(1.0f)
+			.Padding(1.0f, 0.0f)
+			.VAlign(VAlign_Center)
 			[
 				SNew(STextBlock)
 				.Text(InItem->GetDisplayName())
