Fixed single frame discrepancy between Rewind Debugger Timeline / View information and the previewed SkeletalMesh on the viewport.

diff --git a/Engine/Plugins/Animation/GameplayInsights/Source/RewindDebugger/Private/RewindDebugger.cpp b/Engine/Plugins/Animation/GameplayInsights/Source/RewindDebugger/Private/RewindDebugger.cpp
--- a/Engine/Plugins/Animation/GameplayInsights/Source/RewindDebugger/Private/RewindDebugger.cpp
+++ b/Engine/Plugins/Animation/GameplayInsights/Source/RewindDebugger/Private/RewindDebugger.cpp
@@ -723,14 +723,22 @@ void FRewindDebugger::Tick(float DeltaTime)
 
 												AnimationProvider->ReadSkeletalMeshPoseTimeline(ObjectId, [this, Frame, ObjectId, AnimationProvider](const IAnimationProvider::SkeletalMeshPoseTimeline& TimelineData, bool bHasCurves)
 												{
+													const FSkeletalMeshPoseMessage * PoseMessage = nullptr;
+
+													// Get last pose in frame
 													TimelineData.EnumerateEvents(Frame.StartTime, Frame.EndTime,
-														[this](double InStartTime, double InEndTime, uint32 InDepth, const FSkeletalMeshPoseMessage& PoseMessage)
+														[&PoseMessage](double InStartTime, double InEndTime, uint32 InDepth, const FSkeletalMeshPoseMessage& InPoseMessage)
 														{
-															bTargetActorPositionValid = true;
-															TargetActorPosition = PoseMessage.ComponentToWorld.GetTranslation();
-													
-															return TraceServices::EEventEnumerate::Stop;
+															PoseMessage = &InPoseMessage;
+															return TraceServices::EEventEnumerate::Continue;
 														});
+
+													// Update position based on pose
+													if (PoseMessage)
+													{
+														bTargetActorPositionValid = true;
+														TargetActorPosition = PoseMessage->ComponentToWorld.GetTranslation();
+													}
 												});
 											}
 										}
@@ -752,28 +760,36 @@ void FRewindDebugger::Tick(float DeltaTime)
 										{
 											AnimationProvider->ReadSkeletalMeshPoseTimeline(ObjectId, [this, &Frame, ObjectId, MeshComponent, AnimationProvider](const IAnimationProvider::SkeletalMeshPoseTimeline& TimelineData, bool bHasCurves)
 											{
+												const FSkeletalMeshPoseMessage * PoseMessage = nullptr;
+
+												// Get last pose in frame
 												TimelineData.EnumerateEvents(Frame.StartTime, Frame.EndTime,
-													[this, ObjectId, MeshComponent, AnimationProvider](double InStartTime, double InEndTime, uint32 InDepth, const FSkeletalMeshPoseMessage& PoseMessage)
+													[&PoseMessage](double InStartTime, double InEndTime, uint32 InDepth, const FSkeletalMeshPoseMessage& InPoseMessage)
 													{
-														FTransform ComponentWorldTransform;
-														const FSkeletalMeshInfo* SkeletalMeshInfo = AnimationProvider->FindSkeletalMeshInfo(PoseMessage.MeshId);
-														AnimationProvider->GetSkeletalMeshComponentSpacePose(PoseMessage, *SkeletalMeshInfo, ComponentWorldTransform, MeshComponent->GetEditableComponentSpaceTransforms());
-														MeshComponent->ApplyEditedComponentSpaceTransforms();
+														PoseMessage = &InPoseMessage;
+														return TraceServices::EEventEnumerate::Continue;
+													});
 
-														if (MeshComponentsToReset.Find(ObjectId) == nullptr)
-														{
-															FMeshComponentResetData ResetData;
-															ResetData.Component = MeshComponent;
-															ResetData.RelativeTransform = MeshComponent->GetRelativeTransform();
-															MeshComponentsToReset.Add(ObjectId, ResetData);
-														}
+												// Update mesh based on pose
+												if (PoseMessage)
+												{
+													FTransform ComponentWorldTransform;
+													const FSkeletalMeshInfo* SkeletalMeshInfo = AnimationProvider->FindSkeletalMeshInfo(PoseMessage->MeshId);
+													AnimationProvider->GetSkeletalMeshComponentSpacePose(*PoseMessage, *SkeletalMeshInfo, ComponentWorldTransform, MeshComponent->GetEditableComponentSpaceTransforms());
+													MeshComponent->ApplyEditedComponentSpaceTransforms();
 
-														MeshComponent->SetWorldTransform(ComponentWorldTransform, false, nullptr, ETeleportType::TeleportPhysics);
-														MeshComponent->SetForcedLOD(PoseMessage.LodIndex + 1);
-														MeshComponent->UpdateChildTransforms(EUpdateTransformFlags::None, ETeleportType::TeleportPhysics);
+													if (MeshComponentsToReset.Find(ObjectId) == nullptr)
+													{
+														FMeshComponentResetData ResetData;
+														ResetData.Component = MeshComponent;
+														ResetData.RelativeTransform = MeshComponent->GetRelativeTransform();
+														MeshComponentsToReset.Add(ObjectId, ResetData);
+													}
 
-														return TraceServices::EEventEnumerate::Stop;
-													});
+													MeshComponent->SetWorldTransform(ComponentWorldTransform, false, nullptr, ETeleportType::TeleportPhysics);
+													MeshComponent->SetForcedLOD(PoseMessage->LodIndex + 1);
+													MeshComponent->UpdateChildTransforms(EUpdateTransformFlags::None, ETeleportType::TeleportPhysics);
+												}
 											});
 										}
 									}
