Fixes for double click behavior in Rewind Debugger

diff --git a/Engine/Plugins/Animation/GameplayInsights/Source/RewindDebugger/Private/RewindDebugger.cpp b/Engine/Plugins/Animation/GameplayInsights/Source/RewindDebugger/Private/RewindDebugger.cpp
--- a/Engine/Plugins/Animation/GameplayInsights/Source/RewindDebugger/Private/RewindDebugger.cpp
+++ b/Engine/Plugins/Animation/GameplayInsights/Source/RewindDebugger/Private/RewindDebugger.cpp
@@ -984,6 +984,12 @@ void FRewindDebugger::Tick(float DeltaTime)
 	}
 }
 
+void FRewindDebugger::OpenDetailsPanel()
+{
+	bIsDetailsPanelOpen = true;
+	ComponentSelectionChanged(SelectedTrack);
+}
+
 void FRewindDebugger::ComponentSelectionChanged(TSharedPtr<RewindDebugger::FRewindDebuggerTrack> SelectedObject)
 {
 	SelectedTrack = SelectedObject;
@@ -1036,47 +1042,7 @@ void FRewindDebugger::ComponentDoubleClicked(TSharedPtr<RewindDebugger::FRewindD
 	}
 	
 	SelectedTrack = SelectedObject;
-
-	if (SelectedTrack->HandleDoubleClick())
-	{
-		return;
-	}
-	
-	IModularFeatures& ModularFeatures = IModularFeatures::Get();
-	static const FName HandlerFeatureName = IRewindDebuggerDoubleClickHandler::ModularFeatureName;
-		
-	if (const TraceServices::IAnalysisSession* Session = GetAnalysisSession())
-	{
-		TraceServices::FAnalysisSessionReadScope SessionReadScope(*Session);
-	
-		const IGameplayProvider* GameplayProvider = Session->ReadProvider<IGameplayProvider>("GameplayProvider");
-		const FObjectInfo& ObjectInfo = GameplayProvider->GetObjectInfo(SelectedObject->GetObjectId());
-		uint64 ClassId = ObjectInfo.ClassId;
-		bool bHandled = false;
-		
-		const int32 NumExtensions = ModularFeatures.GetModularFeatureImplementationCount(HandlerFeatureName);
-
-		// iterate up the class hierarchy, looking for a registered double click handler, until we find the one that succeeeds that is most specific to the type of this object
-		while (ClassId != 0 && !bHandled)
-		{
-			const FClassInfo& ClassInfo = GameplayProvider->GetClassInfo(ClassId);
-		
-			for (int32 ExtensionIndex = 0; ExtensionIndex < NumExtensions; ++ExtensionIndex)
-			{
-				IRewindDebuggerDoubleClickHandler* Handler = static_cast<IRewindDebuggerDoubleClickHandler*>(ModularFeatures.GetModularFeatureImplementation(HandlerFeatureName, ExtensionIndex));
-				if (Handler->GetTargetTypeName() == ClassInfo.Name)
-				{
-					if (Handler->HandleDoubleClick(this))
-					{
-						bHandled = true;
-						break;
-					}
-				}
-			}
-		
-			ClassId = ClassInfo.SuperId;
-		}
-	}
+	SelectedTrack->HandleDoubleClick();
 }
 
 TSharedPtr<SWidget> FRewindDebugger::BuildComponentContextMenu()
diff --git a/Engine/Plugins/Animation/GameplayInsights/Source/RewindDebugger/Private/RewindDebugger.h b/Engine/Plugins/Animation/GameplayInsights/Source/RewindDebugger/Private/RewindDebugger.h
--- a/Engine/Plugins/Animation/GameplayInsights/Source/RewindDebugger/Private/RewindDebugger.h
+++ b/Engine/Plugins/Animation/GameplayInsights/Source/RewindDebugger/Private/RewindDebugger.h
@@ -114,6 +114,7 @@ public:
 	TBindableProperty<float>* GetRecordingDurationProperty() { return &RecordingDuration; }
 	TBindableProperty<FString, BindingType_Out>* GetDebugTargetActorProperty() { return &DebugTargetActor; }
 
+	virtual void OpenDetailsPanel() override;
 	void SetIsDetailsPanelOpen(bool bIsOpen) { bIsDetailsPanelOpen = bIsOpen; }
 	bool IsDetailsPanelOpen(bool bIsOpen) { return bIsDetailsPanelOpen; }
 
diff --git a/Engine/Plugins/Animation/GameplayInsights/Source/RewindDebugger/Private/RewindDebuggerObjectTrack.cpp b/Engine/Plugins/Animation/GameplayInsights/Source/RewindDebugger/Private/RewindDebuggerObjectTrack.cpp
--- a/Engine/Plugins/Animation/GameplayInsights/Source/RewindDebugger/Private/RewindDebuggerObjectTrack.cpp
+++ b/Engine/Plugins/Animation/GameplayInsights/Source/RewindDebugger/Private/RewindDebuggerObjectTrack.cpp
@@ -10,6 +10,7 @@
 #include "SSegmentedTimelineView.h"
 #include "Styling/SlateIconFinder.h"
 #include "ObjectTrace.h"
+#include "IRewindDebuggerDoubleClickHandler.h"
 
 namespace RewindDebugger
 {
@@ -23,6 +24,48 @@ TSharedPtr<SWidget> FRewindDebuggerObjectTrack::GetTimelineViewInternal()
 
 }
 
+bool FRewindDebuggerObjectTrack::HandleDoubleClickInternal()
+{
+	IModularFeatures& ModularFeatures = IModularFeatures::Get();
+	static const FName HandlerFeatureName = IRewindDebuggerDoubleClickHandler::ModularFeatureName;
+
+	IRewindDebugger* RewindDebugger = IRewindDebugger::Instance();
+		
+	if (const TraceServices::IAnalysisSession* Session = RewindDebugger->GetAnalysisSession())
+	{
+		TraceServices::FAnalysisSessionReadScope SessionReadScope(*Session);
+	
+		const IGameplayProvider* GameplayProvider = Session->ReadProvider<IGameplayProvider>("GameplayProvider");
+		const FObjectInfo& ObjectInfo = GameplayProvider->GetObjectInfo(ObjectId);
+		uint64 ClassId = ObjectInfo.ClassId;
+		bool bHandled = false;
+		
+		const int32 NumExtensions = ModularFeatures.GetModularFeatureImplementationCount(HandlerFeatureName);
+
+		// iterate up the class hierarchy, looking for a registered double click handler, until we find the one that succeeeds that is most specific to the type of this object
+		while (ClassId != 0 && !bHandled)
+		{
+			const FClassInfo& ClassInfo = GameplayProvider->GetClassInfo(ClassId);
+		
+			for (int32 ExtensionIndex = 0; ExtensionIndex < NumExtensions; ++ExtensionIndex)
+			{
+				IRewindDebuggerDoubleClickHandler* Handler = static_cast<IRewindDebuggerDoubleClickHandler*>(ModularFeatures.GetModularFeatureImplementation(HandlerFeatureName, ExtensionIndex));
+				if (Handler->GetTargetTypeName() == ClassInfo.Name)
+				{
+					if (Handler->HandleDoubleClick(RewindDebugger))
+					{
+						bHandled = true;
+						break;
+					}
+				}
+			}
+		
+			ClassId = ClassInfo.SuperId;
+		}
+	}
+	
+	return true;
+}
 
 // check if an object is or is a subclass of a type by name, based on Insights traced type info
 static bool IsTargetType(uint64 ObjectId, FName TargetTypeName, const TraceServices::IAnalysisSession& Session)
diff --git a/Engine/Plugins/Animation/GameplayInsights/Source/RewindDebugger/Private/RewindDebuggerObjectTrack.h b/Engine/Plugins/Animation/GameplayInsights/Source/RewindDebugger/Private/RewindDebuggerObjectTrack.h
--- a/Engine/Plugins/Animation/GameplayInsights/Source/RewindDebugger/Private/RewindDebuggerObjectTrack.h
+++ b/Engine/Plugins/Animation/GameplayInsights/Source/RewindDebugger/Private/RewindDebuggerObjectTrack.h
@@ -33,6 +33,7 @@ private:
 	virtual FText GetDisplayNameInternal() const override;
 	virtual uint64 GetObjectIdInternal() const override { return ObjectId; }
 	virtual bool HasDebugDataInternal() const override { return false; }
+	virtual bool HandleDoubleClickInternal() override;
 	
 	mutable FText DisplayName;
 	FString ObjectName;
diff --git a/Engine/Source/Editor/RewindDebuggerInterface/Public/IRewindDebugger.h b/Engine/Source/Editor/RewindDebuggerInterface/Public/IRewindDebugger.h
--- a/Engine/Source/Editor/RewindDebuggerInterface/Public/IRewindDebugger.h
+++ b/Engine/Source/Editor/RewindDebuggerInterface/Public/IRewindDebugger.h
@@ -78,6 +78,9 @@ public:
 	// returns the length of the current recording
 	virtual double GetRecordingDuration() const = 0;
 
+	// opens the Rewind Debugger details panel tab
+	virtual void OpenDetailsPanel() = 0;
+
 	// get the current IRewindDebugger instance
 	static IRewindDebugger* Instance();
 	
diff --git a/Engine/Source/Editor/RewindDebuggerInterface/Public/RewindDebuggerTrack.h b/Engine/Source/Editor/RewindDebuggerInterface/Public/RewindDebuggerTrack.h
--- a/Engine/Source/Editor/RewindDebuggerInterface/Public/RewindDebuggerTrack.h
+++ b/Engine/Source/Editor/RewindDebuggerInterface/Public/RewindDebuggerTrack.h
@@ -3,6 +3,7 @@
 
 #include "CoreMinimal.h"
 #include "IRewindDebuggerViewCreator.h"
+#include "IRewindDebugger.h"
 
 namespace TraceServices
 {
@@ -107,7 +108,12 @@ private:
 	virtual uint64 GetObjectIdInternal() const { return 0; }
 	virtual bool HasDebugDataInternal() const { return true; }
 	virtual void IterateSubTracksInternal(TFunction<void(TSharedPtr<FRewindDebuggerTrack> SubTrack)> IteratorFunction) { }
-	virtual bool HandleDoubleClickInternal() { return false; }
+
+	virtual bool HandleDoubleClickInternal()
+	{
+		IRewindDebugger::Instance()->OpenDetailsPanel();
+		return true;
+	};
 
 	bool bExpanded;
 	bool bVisible;
