Fix for PIE losing focus when playing with rewind debugger recording

diff --git a/Engine/Plugins/Animation/GameplayInsights/Source/RewindDebugger/Private/RewindDebugger.cpp b/Engine/Plugins/Animation/GameplayInsights/Source/RewindDebugger/Private/RewindDebugger.cpp
--- a/Engine/Plugins/Animation/GameplayInsights/Source/RewindDebugger/Private/RewindDebugger.cpp
+++ b/Engine/Plugins/Animation/GameplayInsights/Source/RewindDebugger/Private/RewindDebugger.cpp
@@ -998,7 +998,10 @@ void FRewindDebugger::ComponentSelectionChanged(TSharedPtr<RewindDebugger::FRewi
 	{
 		FLevelEditorModule& LevelEditorModule = FModuleManager::GetModuleChecked<FLevelEditorModule>("LevelEditor");
 		TSharedPtr<FTabManager> LevelEditorTabManager = LevelEditorModule.GetLevelEditorTabManager();
-		TSharedPtr<SDockTab> DetailsTab = LevelEditorTabManager->TryInvokeTab(FRewindDebuggerModule::DetailsTabName);
+		
+		 // if we now have no selection, don't force the tab into focus - this happens when tracks disappear and can cause PIE to lose focus while playing
+		bool bInvokeAsInactive = !SelectedTrack.IsValid();
+		TSharedPtr<SDockTab> DetailsTab = LevelEditorTabManager->TryInvokeTab(FRewindDebuggerModule::DetailsTabName, bInvokeAsInactive);
 
 		if (DetailsTab.IsValid())
 		{
