Fix error with Rewind Debugger Timeline code which was incorrectly filtering visible tracks as soon as the component/object was deleted

diff --git a/Engine/Plugins/Animation/GameplayInsights/Source/RewindDebugger/Private/RewindDebuggerObjectTrack.cpp b/Engine/Plugins/Animation/GameplayInsights/Source/RewindDebugger/Private/RewindDebuggerObjectTrack.cpp
--- a/Engine/Plugins/Animation/GameplayInsights/Source/RewindDebugger/Private/RewindDebuggerObjectTrack.cpp
+++ b/Engine/Plugins/Animation/GameplayInsights/Source/RewindDebugger/Private/RewindDebuggerObjectTrack.cpp
@@ -225,11 +225,12 @@ bool FRewindDebuggerObjectTrack::UpdateInternal()
 		TRACE_CPUPROFILER_EVENT_SCOPE(FRewindDebuggerObjectTrack::UpdateInternal_AddChildComponents);
 
 		TRange<double> TraceRange = RewindDebugger->GetCurrentTraceRange();
+		TRange<double> ViewRange = RewindDebugger->GetCurrentViewRange();
 
-		GameplayProvider->EnumerateSubobjects(ObjectId, [this, &FoundObjects, &bChanged,&TraceRange, GameplayProvider](uint64 SubobjectId)
+		GameplayProvider->EnumerateSubobjects(ObjectId, [this, &FoundObjects, &bChanged,&TraceRange, &ViewRange,GameplayProvider](uint64 SubobjectId)
 			{
 				TRange<double> Lifetime = GameplayProvider->GetObjectRecordingLifetime(SubobjectId);
-				TRange<double> Overlap = TRange<double>::Intersection(Lifetime, TraceRange);
+				TRange<double> Overlap = TRange<double>::Intersection(Lifetime, ViewRange);
 				// only display the track if the lifetime of the object and the view range overlap
 				if (!Overlap.IsEmpty())
 				{
