Adds ability to watch properties in Rewind Debugger. (Property Watch Track)

**New**
- Added `FPropertiesTrack` as the root track that handles the watched properties via Rewind Debugger.
- Added `FPropertyTrack` used to for watched property tracks via Rewind Debugger.
- Added `FPropertyWatchManager` used to keep track of watched traced properties.
- Added `FPropertyTracks` type specializations:`FBoolPropertyTrack` and `FNumericPropertyTrack`.
- Added `ReadObjectPropertiesStorage` method in `IGameplayProvider`. Used to access container where traced object properties are stored.
- Added `FindPropertyValueFromStorageIndex` method in `IGameplayProvider`. Used to access a traced property given its id.
- Added context menu building support via `FPropertyTrack` and `SPropertiesDebugViewBase`.
- Added a delegate for mouse button down on an `SVariantValueView`.
- Added `SPropertyTrackView` to show a single `FPropertyTrack` variable information in the Rewind Debugger details panel.
- Improved traced property variables displaying. Friendly display names. Specialization for bools, vector, vector2d and quaternions.
- Added a log category for Rewind Debugger. `LogRewindDebugger`.

**Changed**
- Updated `FObjectPropertyValue` to keep track of its unique name id, and unique parent id.
- Updated `FRewindDebugger` to handle menu registration for `ComponentContextMenu`.
- Updated `FVariantTreeNode` to support storing a trace property id. Used for details context menu to support having a menu entry to watch property.
- Updated `SPropertiesDebugViewBase` to have the ability to create a context menu.
- Updated `SObjectPropertiesView` to better display traced properties.
- Updated object property tracing to handle updated `FObjectPropertyValue`. Create unique named ids for properties.
- Updated `FGameAnalyzer` to handle new version of `FObjectPropertyValue`.
- Updated `FRewindDebuggerTrack` to support a building context menu.
- Updated `UComponentContextMenuContext` to stored the current Rewind Debugger selected track.

**Removed**
- Removed usage of `FObjectPropertiesViewCreator`. Now using `FPropertiesTrackCreator`.

diff --git a/Engine/Plugins/Animation/GameplayInsights/Source/GameplayInsights/Private/BoolPropertyTrack.cpp b/Engine/Plugins/Animation/GameplayInsights/Source/GameplayInsights/Private/BoolPropertyTrack.cpp
new file mode 100644
--- /dev/null
+++ b/Engine/Plugins/Animation/GameplayInsights/Source/GameplayInsights/Private/BoolPropertyTrack.cpp
@@ -0,0 +1,101 @@
+﻿// Copyright Epic Games, Inc. All Rights Reserved.
+
+#include "BoolPropertyTrack.h"
+
+#include "GameplayProvider.h"
+#include "PropertyHelpers.h"
+
+namespace RewindDebugger
+{
+	///////////////////////////////////////////////////////////////////////////////
+	// Bool Property Track
+	///////////////////////////////////////////////////////////////////////////////
+
+	FBoolPropertyTrack::FBoolPropertyTrack(uint64 InObjectId, const TSharedPtr<FObjectPropertyInfo> & InObjectProperty, const TSharedPtr<FPropertyTrack> & InParentTrack) : FPropertyTrack(InObjectId, InObjectProperty, InParentTrack)
+	{
+		EnabledSegments = MakeShared<SSegmentedTimelineView::FSegmentData>();
+	}
+
+	bool FBoolPropertyTrack::CanBeCreated(const FObjectPropertyInfo& InObjectProperty)
+	{
+		const IRewindDebugger* RewindDebugger = IRewindDebugger::Instance();
+		check(RewindDebugger)
+
+		if (const TraceServices::IAnalysisSession* AnalysisSession = RewindDebugger->GetAnalysisSession())
+		{
+			TraceServices::FAnalysisSessionReadScope SessionReadScope(*AnalysisSession);
+		
+			if (const FGameplayProvider* GameplayProvider = AnalysisSession->ReadProvider<FGameplayProvider>(FGameplayProvider::ProviderName))
+			{
+				return FObjectPropertyHelpers::IsBoolProperty(InObjectProperty.Property.Value, GameplayProvider->GetPropertyName(InObjectProperty.Property.TypeStringId));
+			}
+		}
+		
+		return false;
+	}
+	
+	bool FBoolPropertyTrack::UpdateInternal()
+	{
+		const IRewindDebugger* RewindDebugger = IRewindDebugger::Instance();
+		check(RewindDebugger)
+		
+		// Convert time range to from rewind debugger times to profiler times
+		const TRange<double> TraceTimeRange = RewindDebugger->GetCurrentTraceRange();
+		const double StartTime = TraceTimeRange.GetLowerBoundValue();
+		const double EndTime = TraceTimeRange.GetUpperBoundValue();
+
+		// Clear previous data
+		TArray<TRange<double>> & Segments = EnabledSegments->Segments;
+		Segments.SetNum(0, false);
+
+		bool bPrevValue = false;
+		double SegmentStartTime = 0.0;
+		double PrevTime = 0.0f;
+		
+		ReadObjectPropertyValueOverTime(StartTime, EndTime, [&Segments, &SegmentStartTime, &bPrevValue, &PrevTime](const FObjectPropertyValue& InValue, uint32 InValueIndex, const FObjectPropertiesMessage& InMessage, const IGameplayProvider::ObjectPropertiesTimeline &  InTimeline , double InStart, double InEndTime)
+		{
+			PrevTime = InMessage.ElapsedTime;
+
+			const bool CurrentValue = FString(InValue.Value).ToBool();
+			
+			// From false to true (Start)
+			if (!bPrevValue && CurrentValue)
+			{
+				SegmentStartTime = InMessage.ElapsedTime;
+				bPrevValue = true;
+			}
+
+			// From true to false (End)
+			else if (bPrevValue && !CurrentValue)
+			{
+				bPrevValue = false;
+
+				Segments.Add(TRange<double>(SegmentStartTime, InMessage.ElapsedTime));
+			}
+		});
+
+		// Handle open segment.
+		if (bPrevValue)
+		{
+			Segments.Add(TRange<double>(SegmentStartTime, PrevTime));
+		}
+		
+		// Update child tracks.
+		return FPropertyTrack::UpdateInternal();
+	}
+
+	TSharedPtr<SWidget> FBoolPropertyTrack::GetTimelineViewInternal()
+	{
+		const auto TimelineView = SNew(SSegmentedTimelineView)
+			.FillColor(FObjectPropertyHelpers::GetPropertyColor(ObjectProperty->Property))
+			.ViewRange_Lambda([]() { return IRewindDebugger::Instance()->GetCurrentViewRange(); })
+			.SegmentData_Raw(this, &FBoolPropertyTrack::GetSegmentData);
+
+		return TimelineView;
+	}
+
+	TSharedPtr<SSegmentedTimelineView::FSegmentData> FBoolPropertyTrack::GetSegmentData() const
+	{
+		return EnabledSegments;
+	}
+}
\ No newline at end of file
diff --git a/Engine/Plugins/Animation/GameplayInsights/Source/GameplayInsights/Private/BoolPropertyTrack.h b/Engine/Plugins/Animation/GameplayInsights/Source/GameplayInsights/Private/BoolPropertyTrack.h
new file mode 100644
--- /dev/null
+++ b/Engine/Plugins/Animation/GameplayInsights/Source/GameplayInsights/Private/BoolPropertyTrack.h
@@ -0,0 +1,35 @@
+﻿// Copyright Epic Games, Inc. All Rights Reserved.
+
+#pragma once
+
+#include "PropertyTrack.h"
+#include "SSegmentedTimelineView.h"
+
+namespace RewindDebugger
+{
+	/**
+	 * Track that displays a object's boolean property value overtime
+	 */
+	class FBoolPropertyTrack : public FPropertyTrack
+	{
+	public:
+
+		BASE_PROPERTY_TRACK()
+		
+		/** Constructor */
+		FBoolPropertyTrack(uint64 InObjectId, const TSharedPtr<FObjectPropertyInfo> & InObjectProperty, const TSharedPtr<FPropertyTrack> & InParentTrack);
+
+	private:
+		
+		/** Begin IRewindDebuggerTrack interface */
+		virtual bool UpdateInternal() override;
+		virtual TSharedPtr<SWidget> GetTimelineViewInternal() override;
+		/** End IRewindDebuggerTrack interface */
+
+		/** @return Segmented boolean data for watched property */
+		TSharedPtr<SSegmentedTimelineView::FSegmentData> GetSegmentData() const;
+
+		/** Keep track of segments when the traced value was true */
+		TSharedPtr<SSegmentedTimelineView::FSegmentData> EnabledSegments;
+	};
+}
diff --git a/Engine/Plugins/Animation/GameplayInsights/Source/GameplayInsights/Private/GameplayAnalyzer.cpp b/Engine/Plugins/Animation/GameplayInsights/Source/GameplayInsights/Private/GameplayAnalyzer.cpp
--- a/Engine/Plugins/Animation/GameplayInsights/Source/GameplayInsights/Private/GameplayAnalyzer.cpp
+++ b/Engine/Plugins/Animation/GameplayInsights/Source/GameplayInsights/Private/GameplayAnalyzer.cpp
@@ -30,8 +30,10 @@ void FGameplayAnalyzer::OnAnalysisBegin(const FOnAnalysisContext& Context)
 	Builder.RouteEvent(RouteId_View, "Object", "View");
 	Builder.RouteEvent(RouteId_ClassPropertyStringId, "Object", "ClassPropertyStringId");
 	Builder.RouteEvent(RouteId_PropertiesStart, "Object", "PropertiesStart");
+	Builder.RouteEvent(RouteId_PropertiesStart2, "Object", "PropertiesStart2");
 	Builder.RouteEvent(RouteId_PropertiesEnd, "Object", "PropertiesEnd");
-	Builder.RouteEvent(RouteId_PropertyValue, "Object", "PropertyValue2");
+	Builder.RouteEvent(RouteId_PropertyValue2, "Object", "PropertyValue2");
+	Builder.RouteEvent(RouteId_PropertyValue3, "Object", "PropertyValue3");
 }
 
 bool FGameplayAnalyzer::OnEvent(uint16 RouteId, EStyle Style, const FOnEventContext& Context)
@@ -171,7 +173,15 @@ bool FGameplayAnalyzer::OnEvent(uint16 RouteId, EStyle Style, const FOnEventCont
 	{
 		uint64 Cycle = EventData.GetValue<uint64>("Cycle");
 		uint64 ObjectId = EventData.GetValue<uint64>("ObjectId");
-		GameplayProvider.AppendPropertiesStart(ObjectId, Context.EventTime.AsSeconds(Cycle), Cycle);
+		GameplayProvider.AppendPropertiesStart(ObjectId, Context.EventTime.AsSeconds(Cycle), Cycle, 0);
+		break;
+	}
+	case RouteId_PropertiesStart2:
+	{
+		uint64 Cycle = EventData.GetValue<uint64>("Cycle");
+		uint64 ObjectId = EventData.GetValue<uint64>("ObjectId");
+		double RecordingTime = EventData.GetValue<double>("RecordingTime");
+		GameplayProvider.AppendPropertiesStart(ObjectId, Context.EventTime.AsSeconds(Cycle), Cycle, RecordingTime);
 		break;
 	}
 	case RouteId_PropertiesEnd:
@@ -181,7 +191,7 @@ bool FGameplayAnalyzer::OnEvent(uint16 RouteId, EStyle Style, const FOnEventCont
 		GameplayProvider.AppendPropertiesEnd(ObjectId, Context.EventTime.AsSeconds(Cycle));
 		break;
 	}
-	case RouteId_PropertyValue:
+	case RouteId_PropertyValue2:
 	{
 		uint64 Cycle = EventData.GetValue<uint64>("Cycle");
 		uint64 ObjectId = EventData.GetValue<uint64>("ObjectId");
@@ -190,9 +200,22 @@ bool FGameplayAnalyzer::OnEvent(uint16 RouteId, EStyle Style, const FOnEventCont
 		uint32 KeyId = EventData.GetValue<uint32>("KeyId");
 
 		FStringView Value; EventData.GetString("Value", Value);
-		GameplayProvider.AppendPropertyValue(ObjectId, Context.EventTime.AsSeconds(Cycle), Cycle, ParentId, TypeId, KeyId, Value);
+		GameplayProvider.AppendPropertyValue(ObjectId, Context.EventTime.AsSeconds(Cycle), Cycle, ParentId, TypeId, KeyId, 0, Value);
 		break;
 	}
+	case RouteId_PropertyValue3:
+		{
+			uint64 Cycle = EventData.GetValue<uint64>("Cycle");
+			uint64 ObjectId = EventData.GetValue<uint64>("ObjectId");
+			int32 ParentId = EventData.GetValue<int32>("ParentId");
+			uint32 TypeId = EventData.GetValue<uint32>("TypeId");
+			uint32 NameId = EventData.GetValue<uint32>("NameId");
+			uint32 ParentNameId = EventData.GetValue<uint32>("ParentNameId");
+			
+			FStringView Value; EventData.GetString("Value", Value);
+			GameplayProvider.AppendPropertyValue(ObjectId, Context.EventTime.AsSeconds(Cycle), Cycle, ParentId, TypeId, NameId, ParentNameId, Value);
+			break;
+		}
 	}
 
 	return true;
diff --git a/Engine/Plugins/Animation/GameplayInsights/Source/GameplayInsights/Private/GameplayAnalyzer.h b/Engine/Plugins/Animation/GameplayInsights/Source/GameplayInsights/Private/GameplayAnalyzer.h
--- a/Engine/Plugins/Animation/GameplayInsights/Source/GameplayInsights/Private/GameplayAnalyzer.h
+++ b/Engine/Plugins/Animation/GameplayInsights/Source/GameplayInsights/Private/GameplayAnalyzer.h
@@ -32,8 +32,10 @@ private:
 		RouteId_ClassPropertyStringId,
 		RouteId_ClassProperty,
 		RouteId_PropertiesStart,
+		RouteId_PropertiesStart2,
 		RouteId_PropertiesEnd,
-		RouteId_PropertyValue,
+		RouteId_PropertyValue2,
+		RouteId_PropertyValue3,
 		RouteId_RecordingInfo,
 	};
 
diff --git a/Engine/Plugins/Animation/GameplayInsights/Source/GameplayInsights/Private/GameplayInsightsModule.cpp b/Engine/Plugins/Animation/GameplayInsights/Source/GameplayInsights/Private/GameplayInsightsModule.cpp
--- a/Engine/Plugins/Animation/GameplayInsights/Source/GameplayInsights/Private/GameplayInsightsModule.cpp
+++ b/Engine/Plugins/Animation/GameplayInsights/Source/GameplayInsights/Private/GameplayInsightsModule.cpp
@@ -17,12 +17,13 @@
 #include "Stats/Stats.h"
 #include "ObjectPropertyTrace.h"
 #include "SMontageView.h"
-#include "SObjectPropertiesView.h"
 #include "AnimCurvesTrack.h"
 #include "BlendWeightsTrack.h"
 #include "MontagesTrack.h"
 #include "NotifiesTrack.h"
 #include "PoseWatchTrack.h"
+#include "PropertiesTrack.h"
+#include "PropertyWatchManager.h"
 
 #if WITH_EDITOR
 #include "IAnimationBlueprintEditorModule.h"
@@ -73,8 +74,6 @@ void FGameplayInsightsModule::StartupModule()
 
 #if WITH_EDITOR
 	// register rewind debugger view creators
-	static FObjectPropertiesViewCreator ObjectPropertiesViewCreator;
-	IModularFeatures::Get().RegisterModularFeature(IRewindDebuggerViewCreator::ModularFeatureName, &ObjectPropertiesViewCreator);
 	static FAnimGraphSchematicViewCreator AnimGraphSchematicViewCreator;
 	IModularFeatures::Get().RegisterModularFeature(IRewindDebuggerViewCreator::ModularFeatureName, &AnimGraphSchematicViewCreator);
 	
@@ -88,8 +87,11 @@ void FGameplayInsightsModule::StartupModule()
 	IModularFeatures::Get().RegisterModularFeature(RewindDebugger::IRewindDebuggerTrackCreator::ModularFeatureName, &NotifiesTrackCreator);
 	static RewindDebugger::FPoseWatchesTrackCreator PoseWatchesTrackCreator;
 	IModularFeatures::Get().RegisterModularFeature(RewindDebugger::IRewindDebuggerTrackCreator::ModularFeatureName, &PoseWatchesTrackCreator);
+	static RewindDebugger::FPropertiesTrackCreator PropertyTrackCreator;
+	IModularFeatures::Get().RegisterModularFeature(RewindDebugger::IRewindDebuggerTrackCreator::ModularFeatureName, &PropertyTrackCreator);
 
-
+	FPropertyWatchManager::Initialize();
+	
 	if (!IsRunningCommandlet())
 	{
 		IAnimationBlueprintEditorModule& AnimationBlueprintEditorModule = FModuleManager::LoadModuleChecked<IAnimationBlueprintEditorModule>("AnimationBlueprintEditor");
@@ -240,6 +242,7 @@ void FGameplayInsightsModule::ShutdownModule()
 	{
 		AnimationBlueprintEditorModule->OnGetCustomDebugObjects().Remove(CustomDebugObjectHandle);
 	}
+	FPropertyWatchManager::Shutdown();
 #else
 	IUnrealInsightsModule& UnrealInsightsModule = FModuleManager::LoadModuleChecked<IUnrealInsightsModule>("TraceInsights");
 	FOnRegisterMajorTabExtensions& TimingProfilerLayoutExtension = UnrealInsightsModule.OnRegisterMajorTabExtension(FInsightsManagerTabs::TimingProfilerTabId);
diff --git a/Engine/Plugins/Animation/GameplayInsights/Source/GameplayInsights/Private/GameplayProvider.cpp b/Engine/Plugins/Animation/GameplayInsights/Source/GameplayInsights/Private/GameplayProvider.cpp
--- a/Engine/Plugins/Animation/GameplayInsights/Source/GameplayInsights/Private/GameplayProvider.cpp
+++ b/Engine/Plugins/Animation/GameplayInsights/Source/GameplayInsights/Private/GameplayProvider.cpp
@@ -162,6 +162,24 @@ bool FGameplayProvider::ReadObjectPropertiesTimeline(uint64 InObjectId, TFunctio
 	return false;
 }
 
+bool FGameplayProvider::ReadObjectPropertiesStorage(uint64 InObjectId, const FObjectPropertiesMessage& InMessage, TFunctionRef<void(const TConstArrayView<FObjectPropertyValue>&)> Callback) const
+{
+	Session.ReadAccessCheck();
+
+	const uint32* IndexPtr = ObjectIdToPropertiesStorage.Find(InObjectId);
+	if(IndexPtr != nullptr)
+	{
+		if (*IndexPtr < uint32(ObjectIdToPropertiesStorage.Num()))
+		{
+			const TSharedRef<FObjectPropertiesStorage> Storage = PropertiesStorage[*IndexPtr];
+			Callback(MakeArrayView(Storage->Values.GetData() + InMessage.PropertyValueStartIndex, InMessage.PropertyValueEndIndex - InMessage.PropertyValueStartIndex));
+			return true;
+		}
+	}
+
+	return false;
+}
+
 void FGameplayProvider::EnumerateObjectPropertyValues(uint64 InObjectId, const FObjectPropertiesMessage& InMessage, TFunctionRef<void(const FObjectPropertyValue&)> Callback) const
 {
 	Session.ReadAccessCheck();
@@ -179,6 +197,7 @@ void FGameplayProvider::EnumerateObjectPropertyValues(uint64 InObjectId, const F
 		}
 	}
 }
+
 void FGameplayProvider::EnumerateObjects(TFunctionRef<void(const FObjectInfo&)> Callback) const
 {
 	Session.ReadAccessCheck();
@@ -205,6 +224,27 @@ void FGameplayProvider::EnumerateObjects(double StartTime, double EndTime, TFunc
 		});
 }
 
+const FObjectPropertyValue* FGameplayProvider::FindPropertyValueFromStorageIndex(uint64 InObjectId, int64 InStorageIndex) const
+{
+	Session.ReadAccessCheck();
+
+	const uint32* IndexPtr = ObjectIdToPropertiesStorage.Find(InObjectId);
+	if(IndexPtr != nullptr)
+	{
+		if (*IndexPtr < static_cast<uint32>(ObjectIdToPropertiesStorage.Num()))
+		{
+			const TSharedRef<FObjectPropertiesStorage> Storage = PropertiesStorage[*IndexPtr];
+
+			if (Storage->Values.IsValidIndex(InStorageIndex))
+			{
+				return &Storage->Values[InStorageIndex];
+			}
+		}
+	}
+
+	return nullptr;
+}
+
 void FGameplayProvider::EnumerateSubobjects(uint64 ObjectId, TFunctionRef<void(uint64 SubobjectId)> Callback) const
 {
 	TArray<uint64, TInlineAllocator<32>> SubobjectIds;
@@ -660,7 +700,7 @@ void FGameplayProvider::AppendClassPropertyStringId(uint32 InStringId, const FSt
 	PropertyStrings.Add(InStringId, StoredString);
 }
 
-void FGameplayProvider::AppendPropertiesStart(uint64 InObjectId, double InTime, uint64 InEventId)
+void FGameplayProvider::AppendPropertiesStart(uint64 InObjectId, double InTime, uint64 InEventId, double InRecordingTime)
 {
 	Session.WriteAccessCheck();
 
@@ -687,6 +727,8 @@ void FGameplayProvider::AppendPropertiesStart(uint64 InObjectId, double InTime,
 	FObjectPropertiesMessage& Message = Storage->OpenEvent;
 	Message.PropertyValueStartIndex = Storage->Values.Num();
 	Message.PropertyValueEndIndex = Storage->Values.Num();
+	Message.ProfileTime = InTime;
+	Message.ElapsedTime = InRecordingTime;
 }
 
 void FGameplayProvider::AppendPropertiesEnd(uint64 InObjectId, double InTime)
@@ -719,7 +761,7 @@ void FGameplayProvider::AppendPropertiesEnd(uint64 InObjectId, double InTime)
 	}
 }
 
-void FGameplayProvider::AppendPropertyValue(uint64 InObjectId, double InTime, uint64 InEventId, int32 InParentId, uint32 InTypeStringId, uint32 InKeyStringId, const FStringView& InValue)
+void FGameplayProvider::AppendPropertyValue(uint64 InObjectId, double InTime, uint64 InEventId, int32 InParentId, uint32 InTypeStringId, uint32 InNameId, uint32 InParentNameId, const FStringView& InValue)
 {
 	Session.WriteAccessCheck();
 
@@ -743,12 +785,14 @@ void FGameplayProvider::AppendPropertyValue(uint64 InObjectId, double InTime, ui
 	if(Storage->OpenEventId == InEventId)
 	{
 		FObjectPropertyValue& Message = Storage->Values.AddDefaulted_GetRef();
+		
 		Message.Value = Session.StoreString(InValue);
 		Message.ValueAsFloat = FCString::Atof(Message.Value);
 		Message.ParentId = InParentId;
 		Message.TypeStringId = InTypeStringId;
-		Message.KeyStringId = InKeyStringId; 
-
+		Message.NameId = InNameId;
+		Message.ParentNameId = InParentNameId;
+		
 		Storage->OpenEvent.PropertyValueEndIndex = Storage->Values.Num();
 	}
 }
diff --git a/Engine/Plugins/Animation/GameplayInsights/Source/GameplayInsights/Private/GameplayProvider.h b/Engine/Plugins/Animation/GameplayInsights/Source/GameplayInsights/Private/GameplayProvider.h
--- a/Engine/Plugins/Animation/GameplayInsights/Source/GameplayInsights/Private/GameplayProvider.h
+++ b/Engine/Plugins/Animation/GameplayInsights/Source/GameplayInsights/Private/GameplayProvider.h
@@ -22,10 +22,12 @@ public:
 	virtual bool ReadObjectEventsTimeline(uint64 InObjectId, TFunctionRef<void(const ObjectEventsTimeline&)> Callback) const override;
 	virtual bool ReadObjectEvent(uint64 InObjectId, uint64 InMessageId, TFunctionRef<void(const FObjectEventMessage&)> Callback) const override;
 	virtual bool ReadObjectPropertiesTimeline(uint64 InObjectId, TFunctionRef<void(const ObjectPropertiesTimeline&)> Callback) const override;
+	virtual bool ReadObjectPropertiesStorage(uint64 InObjectId, const FObjectPropertiesMessage& InMessage, TFunctionRef<void(const TConstArrayView<FObjectPropertyValue>&)> Callback) const override;
 	virtual void EnumerateObjectPropertyValues(uint64 InObjectId, const FObjectPropertiesMessage& InMessage, TFunctionRef<void(const FObjectPropertyValue&)> Callback) const override;
 	virtual void EnumerateObjects(TFunctionRef<void(const FObjectInfo&)> Callback) const override;
 	virtual void EnumerateObjects(double StartTime, double EndTime, TFunctionRef<void(const FObjectInfo&)> Callback) const override;
 	virtual void EnumerateSubobjects(uint64 ObjectId, TFunctionRef<void(uint64 SubobjectId)> Callback) const override;
+	virtual const FObjectPropertyValue* FindPropertyValueFromStorageIndex(uint64 InObjectId, int64 InStorageIndex) const override;
 	virtual const FClassInfo* FindClassInfo(uint64 InClassId) const override;
 	virtual const UClass* FindClass(uint64 InClassId) const override;
 	virtual const FClassInfo* FindClassInfo(const TCHAR* InClassPath) const override;
@@ -73,13 +75,13 @@ public:
 	void AppendClassPropertyStringId(uint32 InStringId, const FStringView& InString);
 
 	/** Add a properties start message */
-	void AppendPropertiesStart(uint64 InObjectId, double InTime, uint64 InEventId);
+	void AppendPropertiesStart(uint64 InObjectId, double InTime, uint64 InEventId, double InRecordingTime);
 
 	/** Add a properties end message */
 	void AppendPropertiesEnd(uint64 InObjectId, double InTime);
 
 	/** Add a property value message */
-	void AppendPropertyValue(uint64 InObjectId, double InTime, uint64 InEventId, int32 InParentId, uint32 InTypeStringId, uint32 InKeyStringId, const FStringView& InValue);
+	void AppendPropertyValue(uint64 InObjectId, double InTime, uint64 InEventId, int32 InParentId, uint32 InTypeStringId, uint32 InNameId, uint32 InParentNameId, const FStringView& InValue);
 
 	/** Check whether we have any data */
 	bool HasAnyData() const;
diff --git a/Engine/Plugins/Animation/GameplayInsights/Source/GameplayInsights/Private/NumericPropertyTrack.cpp b/Engine/Plugins/Animation/GameplayInsights/Source/GameplayInsights/Private/NumericPropertyTrack.cpp
new file mode 100644
--- /dev/null
+++ b/Engine/Plugins/Animation/GameplayInsights/Source/GameplayInsights/Private/NumericPropertyTrack.cpp
@@ -0,0 +1,85 @@
+﻿// Copyright Epic Games, Inc. All Rights Reserved.
+
+#include "NumericPropertyTrack.h"
+#include "PropertyHelpers.h"
+
+namespace RewindDebugger
+{
+	///////////////////////////////////////////////////////////////////////////////
+	// Numeric Property Track
+	///////////////////////////////////////////////////////////////////////////////
+	
+	FNumericPropertyTrack::FNumericPropertyTrack(uint64 InObjectId, const TSharedPtr<FObjectPropertyInfo> & InObjectProperty, const TSharedPtr<FPropertyTrack> & InParentTrack) : FPropertyTrack(InObjectId, InObjectProperty, InParentTrack)
+	{
+	}
+
+	bool FNumericPropertyTrack::CanBeCreated(const FObjectPropertyInfo& InObjectProperty)
+	{
+		const FString ValueString = { InObjectProperty.Property.Value };
+		return !ValueString.IsEmpty() && ValueString.IsNumeric();
+	}
+	
+	TSharedPtr<SCurveTimelineView::FTimelineCurveData> FNumericPropertyTrack::GetCurveData() const
+	{
+		if (!CurveData.IsValid())
+		{
+			CurveData = MakeShared<SCurveTimelineView::FTimelineCurveData>();
+		}
+
+		CurvesUpdateRequested++;
+
+		return CurveData;
+	}
+	
+	bool FNumericPropertyTrack::UpdateInternal()
+	{
+		const IRewindDebugger* RewindDebugger = IRewindDebugger::Instance();
+		check(RewindDebugger)
+
+		// Convert time range to from rewind debugger times to profiler times
+		const TRange<double> TraceTimeRange = RewindDebugger->GetCurrentTraceRange();
+		const double StartTime = TraceTimeRange.GetLowerBoundValue();
+		const double EndTime = TraceTimeRange.GetUpperBoundValue();
+
+		// Root information
+		
+		// Get information providers)
+		if (CurvesUpdateRequested > 10)
+		{
+			// Clear curve
+			TArray<SCurveTimelineView::FTimelineCurveData::CurvePoint> & CurvePoints = CurveData->Points;
+			CurvePoints.SetNum(0, false);
+
+			// Draw curve overtime
+			ReadObjectPropertyValueOverTime(StartTime, EndTime, [&CurvePoints](const FObjectPropertyValue& InValue, uint32 InValueIndex, const FObjectPropertiesMessage& InMessage, const IGameplayProvider::ObjectPropertiesTimeline &  InTimeline , double InStart, double InEndTime)
+			{
+				CurvePoints.Add({InMessage.ElapsedTime, InValue.ValueAsFloat});
+			});
+
+			// Reset counter
+			CurvesUpdateRequested = 0;
+		}
+
+		// Update child tracks.
+		return FPropertyTrack::UpdateInternal();
+	}
+
+	TSharedPtr<SWidget> FNumericPropertyTrack::GetTimelineViewInternal()
+	{
+		const FLinearColor Color = FObjectPropertyHelpers::GetPropertyColor(ObjectProperty->Property);
+		
+		FLinearColor CurveColor = Color;
+		CurveColor.R *= 0.25;
+		CurveColor.G *= 0.25;
+		CurveColor.B *= 0.25;
+
+		TSharedPtr<SCurveTimelineView> CurveTimelineView = SNew(SCurveTimelineView)
+		.FillColor(Color)
+		.CurveColor(CurveColor)
+		.ViewRange_Lambda([]() { return IRewindDebugger::Instance()->GetCurrentViewRange(); })
+		.RenderFill(false)
+		.CurveData_Raw(this, &FNumericPropertyTrack::GetCurveData);
+
+		return CurveTimelineView;
+	}
+}
\ No newline at end of file
diff --git a/Engine/Plugins/Animation/GameplayInsights/Source/GameplayInsights/Private/NumericPropertyTrack.h b/Engine/Plugins/Animation/GameplayInsights/Source/GameplayInsights/Private/NumericPropertyTrack.h
new file mode 100644
--- /dev/null
+++ b/Engine/Plugins/Animation/GameplayInsights/Source/GameplayInsights/Private/NumericPropertyTrack.h
@@ -0,0 +1,39 @@
+﻿// Copyright Epic Games, Inc. All Rights Reserved.
+
+#pragma once
+
+#include "PropertyTrack.h"
+#include "SCurveTimelineView.h"
+
+namespace RewindDebugger
+{
+	/**
+	 * Track that displays a object's numeric property value overtime
+	 */
+	class FNumericPropertyTrack : public FPropertyTrack
+	{
+	public:
+
+		BASE_PROPERTY_TRACK()
+		
+		/** Constructor */
+		FNumericPropertyTrack(uint64 InObjectId, const TSharedPtr<FObjectPropertyInfo> & InObjectProperty, const TSharedPtr<FPropertyTrack> & InParentTrack);
+
+		/** @return Property's timeline view curve data */
+		TSharedPtr<SCurveTimelineView::FTimelineCurveData> GetCurveData() const;
+	
+	protected:
+		
+		/** Begin IRewindDebuggerTrack interface */
+		virtual bool UpdateInternal() override;
+		virtual TSharedPtr<SWidget> GetTimelineViewInternal() override;
+		/** End IRewindDebuggerTrack interface */
+		
+		/** Curve data used to display value of property overtime */
+		mutable TSharedPtr<SCurveTimelineView::FTimelineCurveData> CurveData;
+
+		/** Keeps track of how many times the curve view has been requested to be drawn/computed */
+		mutable int CurvesUpdateRequested = 0;
+	};
+
+}
diff --git a/Engine/Plugins/Animation/GameplayInsights/Source/GameplayInsights/Private/ObjectPropertiesTrack.cpp b/Engine/Plugins/Animation/GameplayInsights/Source/GameplayInsights/Private/ObjectPropertiesTrack.cpp
--- a/Engine/Plugins/Animation/GameplayInsights/Source/GameplayInsights/Private/ObjectPropertiesTrack.cpp
+++ b/Engine/Plugins/Animation/GameplayInsights/Source/GameplayInsights/Private/ObjectPropertiesTrack.cpp
@@ -4,9 +4,9 @@
 #include "GameplayProvider.h"
 #include "Insights/ViewModels/TimingTrackViewport.h"
 #include "Insights/ViewModels/TimingEvent.h"
-#include "Algo/Sort.h"
 #include "Insights/ViewModels/TooltipDrawState.h"
 #include "GameplaySharedData.h"
+#include "PropertyHelpers.h"
 #include "Insights/ViewModels/TimingEventSearch.h"
 #include "TraceServices/Model/Frames.h"
 #include "VariantTreeNode.h"
@@ -146,8 +146,7 @@ FText FObjectPropertiesTrack::MakeTrackName(const FGameplaySharedData& InSharedD
 
 void FObjectPropertiesTrack::GetVariantsAtFrame(const TraceServices::FFrame& InFrame, TArray<TSharedRef<FVariantTreeNode>>& OutVariants) const
 {
-	const FGameplayProvider* GameplayProvider = SharedData.GetAnalysisSession().ReadProvider<FGameplayProvider>(FGameplayProvider::ProviderName);
-	if(GameplayProvider)
+	if(const FGameplayProvider* GameplayProvider = SharedData.GetAnalysisSession().ReadProvider<FGameplayProvider>(FGameplayProvider::ProviderName))
 	{
 		TraceServices::FAnalysisSessionReadScope SessionReadScope(SharedData.GetAnalysisSession());
 
@@ -158,21 +157,24 @@ void FObjectPropertiesTrack::GetVariantsAtFrame(const TraceServices::FFrame& InF
 		{
 			InTimeline.EnumerateEvents(InFrame.StartTime, InFrame.EndTime, [this, GameplayProvider, &Header, &PropertyVariants](double InStartTime, double InEndTime, uint32 InDepth, const FObjectPropertiesMessage& InMessage)
 			{
-				GameplayProvider->EnumerateObjectPropertyValues(GetGameplayTrack().GetObjectId(), InMessage, [GameplayProvider, &Header, &PropertyVariants](const FObjectPropertyValue& InValue)
+				GameplayProvider->ReadObjectPropertiesStorage(GetGameplayTrack().GetObjectId(), InMessage, [GameplayProvider, &Header, &PropertyVariants](const TConstArrayView<FObjectPropertyValue> & InStorage)
 				{
-					const TCHAR* Key = GameplayProvider->GetPropertyName(InValue.KeyStringId);
-					PropertyVariants.Add(FVariantTreeNode::MakeString(FText::FromString(Key), InValue.Value));
-
-					// note assumes that order is parent->child in the properties array
-					if(InValue.ParentId != INDEX_NONE)
-					{
-						PropertyVariants[InValue.ParentId]->AddChild(PropertyVariants.Last().ToSharedRef());
-					}
-					else
+					for (int32 i = 0; i < InStorage.Num(); ++i)
 					{
-						Header->AddChild(PropertyVariants.Last().ToSharedRef());
+						PropertyVariants.Add(FObjectPropertyHelpers::GetVariantNodeFromProperty(i, *GameplayProvider, InStorage));
+
+						// note assumes that order is parent->child in the properties array
+						if(InStorage[i].ParentId != INDEX_NONE)
+						{
+							PropertyVariants[InStorage[i].ParentId]->AddChild(PropertyVariants.Last().ToSharedRef());
+						}
+						else
+						{
+							Header->AddChild(PropertyVariants.Last().ToSharedRef());
+						}
 					}
 				});
+				
 				return TraceServices::EEventEnumerate::Stop;
 			});
 		});
diff --git a/Engine/Plugins/Animation/GameplayInsights/Source/GameplayInsights/Private/ObjectPropertyTrace.cpp b/Engine/Plugins/Animation/GameplayInsights/Source/GameplayInsights/Private/ObjectPropertyTrace.cpp
--- a/Engine/Plugins/Animation/GameplayInsights/Source/GameplayInsights/Private/ObjectPropertyTrace.cpp
+++ b/Engine/Plugins/Animation/GameplayInsights/Source/GameplayInsights/Private/ObjectPropertyTrace.cpp
@@ -1,6 +1,7 @@
 // Copyright Epic Games, Inc. All Rights Reserved.
 
 #include "ObjectPropertyTrace.h"
+#include "IGameplayProvider.h"
 #include "Trace/Trace.h"
 #include "Trace/Trace.inl"
 #include "Containers/Ticker.h"
@@ -15,9 +16,10 @@ UE_TRACE_EVENT_BEGIN(Object, ClassPropertyStringId, NoSync|Important)
 	UE_TRACE_EVENT_FIELD(UE::Trace::WideString, Value)
 UE_TRACE_EVENT_END()
 
-UE_TRACE_EVENT_BEGIN(Object, PropertiesStart)
+UE_TRACE_EVENT_BEGIN(Object, PropertiesStart2)
 	UE_TRACE_EVENT_FIELD(uint64, Cycle)
 	UE_TRACE_EVENT_FIELD(uint64, ObjectId)
+	UE_TRACE_EVENT_FIELD(double, RecordingTime)
 UE_TRACE_EVENT_END()
 
 UE_TRACE_EVENT_BEGIN(Object, PropertiesEnd)
@@ -25,12 +27,13 @@ UE_TRACE_EVENT_BEGIN(Object, PropertiesEnd)
 	UE_TRACE_EVENT_FIELD(uint64, ObjectId)
 UE_TRACE_EVENT_END()
 
-UE_TRACE_EVENT_BEGIN(Object, PropertyValue2)
+UE_TRACE_EVENT_BEGIN(Object, PropertyValue3)
 	UE_TRACE_EVENT_FIELD(uint64, Cycle)
 	UE_TRACE_EVENT_FIELD(uint64, ObjectId)
 	UE_TRACE_EVENT_FIELD(int32, ParentId)
 	UE_TRACE_EVENT_FIELD(uint32, TypeId)
-	UE_TRACE_EVENT_FIELD(uint32, KeyId)
+	UE_TRACE_EVENT_FIELD(uint32, NameId)
+	UE_TRACE_EVENT_FIELD(uint32, ParentNameId)
 	UE_TRACE_EVENT_FIELD(UE::Trace::WideString, Value)
 UE_TRACE_EVENT_END()
 
@@ -41,26 +44,52 @@ namespace ObjectPropertyTrace
 
 	static uint32 CurrentClassPropertyStringId = 0;
 	static TMap<FString, uint32> StringIdMap;
-
-	typedef TFunctionRef<void(const FString&, const FString&, const FString&, int32, int32)> IterateFunction;
+	typedef TFunctionRef<void(const FString& /*InTypeString*/, const FString& /*InValue*/, int32 /*InId*/, int32 /*InParentId*/, FStringBuilderBase & /*InPropertyNameBuilder*/)> IterateFunction;
 
 	static uint32 TraceStringId(const FString& InString)
 	{
-		if(uint32* ExistingIdPtr = StringIdMap.Find(InString))
+		if(const uint32* ExistingIdPtr = StringIdMap.Find(InString))
 		{
 			return *ExistingIdPtr;
 		}
 
-		uint32 NewId = StringIdMap.Add(InString, ++CurrentClassPropertyStringId);
+		if (InString.IsEmpty())
+		{
+			return 0;	// Return invalid id.
+		}
+		else
+		{
+			const uint32 NewId = StringIdMap.Add(InString, ++CurrentClassPropertyStringId);
 
-		UE_TRACE_LOG(Object, ClassPropertyStringId, ObjectProperties, InString.Len() * sizeof(TCHAR))
-			<< ClassPropertyStringId.Id(NewId)
-			<< ClassPropertyStringId.Value(*InString, InString.Len());
+			UE_TRACE_LOG(Object, ClassPropertyStringId, ObjectProperties, InString.Len() * sizeof(TCHAR))
+				<< ClassPropertyStringId.Id(NewId)
+				<< ClassPropertyStringId.Value(*InString, InString.Len());
 
-		return NewId;
+			return NewId;
+		}
 	}
 
-	static void IteratePropertiesRecursive(FProperty* InProperty, const void* InContainer, const FString& InKey, IterateFunction InFunction, int32& InId, int32 InParentId)
+	/**
+	 * Helper to ensure FStringBuilderBase clears any added strings at the end of a scope.
+	 * 
+	 * Used to support creating one single FStringBuilderBase when tracing an object's properties recursively.
+	 */
+	struct FScopedPropertyName
+	{
+		FScopedPropertyName(FStringBuilderBase & InBuilder) : StartSize(InBuilder.Len()), Builder(InBuilder)
+		{
+		}
+
+		~FScopedPropertyName()
+		{
+			Builder.RemoveSuffix(FMath::Max(Builder.Len() - StartSize, 0));
+		}
+		
+		int32 StartSize;
+		FStringBuilderBase & Builder;
+	};
+	
+	static void IteratePropertiesRecursive(FProperty* InProperty, const void* InContainer, const FString& InKey, IterateFunction InFunction, int32& InId, int32 InParentId, FStringBuilderBase & InPropertyNameBuilder)
 	{
 		// Handle container properties
 		if (const FArrayProperty* ArrayProperty = CastField<FArrayProperty>(InProperty))
@@ -69,14 +98,22 @@ namespace ObjectPropertyTrace
 
 			int32 ArrayRootId = ++InId;
 			FString SizeString = FString::Printf(TEXT("{Num = %d}"), Helper.Num());
-			InFunction(InProperty->GetCPPType(), InKey, SizeString, ArrayRootId, InParentId);
 			
+			// Build property name id
+			FScopedPropertyName ScopedPropertyName(InPropertyNameBuilder);
+			ScopedPropertyName.Builder.AppendChar(TEXT('.'));
+			ScopedPropertyName.Builder.Append(InProperty->GetName());
+
+			// Trace array property
+			InFunction(InProperty->GetCPPType(), SizeString, ArrayRootId, InParentId, InPropertyNameBuilder);
+
+			// Trace container items
 			for (int32 DynamicIndex = 0; DynamicIndex < Helper.Num(); ++DynamicIndex)
 			{
 				const void* ValuePtr = Helper.GetRawPtr(DynamicIndex);
 				FString KeyString = FString::Printf(TEXT("[%d]"), DynamicIndex);
-
-				IteratePropertiesRecursive(ArrayProperty->Inner, ValuePtr, KeyString, InFunction, InId, ArrayRootId);
+				
+				IteratePropertiesRecursive(ArrayProperty->Inner, ValuePtr, KeyString, InFunction, InId, ArrayRootId, InPropertyNameBuilder);
 			}
 		}
 		else if (const FMapProperty* MapProperty = CastField<FMapProperty>(InProperty))
@@ -85,8 +122,14 @@ namespace ObjectPropertyTrace
 
 			int32 MapRootId = ++InId;
 			FString SizeString = FString::Printf(TEXT("{Num = %d}"), Helper.Num());
-			InFunction(InProperty->GetCPPType(), InKey, SizeString, MapRootId, InParentId);
 
+			// Build property name id
+			FScopedPropertyName ScopedPropertyName(InPropertyNameBuilder);
+			ScopedPropertyName.Builder.AppendChar(TEXT('.'));
+			ScopedPropertyName.Builder.Append(InProperty->GetName());
+			
+			InFunction(InProperty->GetCPPType(), SizeString, MapRootId, InParentId, InPropertyNameBuilder);
+			
 			int32 Num = Helper.Num();
 			int32 MapIndex = 0;
 			for (int32 DynamicIndex = 0; Num; ++DynamicIndex)
@@ -96,13 +139,19 @@ namespace ObjectPropertyTrace
 					int32 MapEntryId = ++InId;
 					FString KeyString = FString::Printf(TEXT("[%d]"), MapIndex++);
 					FString TypeString = FString::Printf(TEXT("{%s, %s}"), *MapProperty->KeyProp->GetCPPType(), *MapProperty->ValueProp->GetCPPType());
-					InFunction(TypeString, KeyString, TEXT("{...}"), MapEntryId, MapRootId);
 
-					const void* PairPtr = Helper.GetPairPtr(DynamicIndex);
+					// Build pair parent name id
+					FScopedPropertyName ScopedParentPairName(InPropertyNameBuilder);
+					ScopedParentPairName.Builder.AppendChar(TEXT('.'));
+					ScopedParentPairName.Builder.Append(KeyString);
 					
-					IteratePropertiesRecursive(MapProperty->KeyProp, PairPtr, MapProperty->KeyProp->GetName(), InFunction, InId, MapEntryId);
+					InFunction(TypeString, TEXT("{...}"), MapEntryId, MapRootId, InPropertyNameBuilder);
+					
+					const void* PairPtr = Helper.GetPairPtr(DynamicIndex);
+
+					IteratePropertiesRecursive(MapProperty->KeyProp, PairPtr, MapProperty->KeyProp->GetName(), InFunction, InId, MapEntryId, InPropertyNameBuilder);
 
-					IteratePropertiesRecursive(MapProperty->ValueProp, PairPtr, MapProperty->ValueProp->GetName(), InFunction, InId, MapEntryId);
+					IteratePropertiesRecursive(MapProperty->ValueProp, PairPtr, MapProperty->ValueProp->GetName(), InFunction, InId, MapEntryId, InPropertyNameBuilder);
 
 					--Num;
 				}
@@ -114,8 +163,16 @@ namespace ObjectPropertyTrace
 
 			int32 SetRootId = ++InId;
 			FString SizeString = FString::Printf(TEXT("{Num = %d}"), Helper.Num());
-			InFunction(InProperty->GetCPPType(), InKey, SizeString, SetRootId, InParentId);
 
+			// Build property name id
+			FScopedPropertyName ScopedPropertyName(InPropertyNameBuilder);
+			ScopedPropertyName.Builder.AppendChar(TEXT('.'));
+			ScopedPropertyName.Builder.Append(InProperty->GetName());
+
+			// Trace set property
+			InFunction(InProperty->GetCPPType(), SizeString, SetRootId, InParentId, InPropertyNameBuilder);
+
+			// Trace container items
 			int32 Num = Helper.Num();
 			int32 SetIndex = 0;
 			for (int32 DynamicIndex = 0; Num; ++DynamicIndex)
@@ -124,8 +181,8 @@ namespace ObjectPropertyTrace
 				{
 					const void* ValuePtr = Helper.GetElementPtr(DynamicIndex);
 					FString KeyString = FString::Printf(TEXT("[%d]"), SetIndex++);
-
-					IteratePropertiesRecursive(SetProperty->ElementProp, ValuePtr, KeyString, InFunction, InId, SetRootId);
+					
+					IteratePropertiesRecursive(SetProperty->ElementProp, ValuePtr, KeyString, InFunction, InId, SetRootId, InPropertyNameBuilder);
 
 					--Num;
 				}
@@ -134,13 +191,20 @@ namespace ObjectPropertyTrace
 		else if (const FStructProperty* StructProperty = CastField<FStructProperty>(InProperty))
 		{
 			int32 StructRootId = ++InId;
-			InFunction(InProperty->GetCPPType(), InKey, TEXT("{...}"), StructRootId, InParentId);
 
+			// Build property name id
+			FScopedPropertyName ScopedPropertyName(InPropertyNameBuilder);
+			ScopedPropertyName.Builder.AppendChar(TEXT('.'));
+			ScopedPropertyName.Builder.Append(InProperty->GetName());
+
+			// Trace struct property
+			InFunction(InProperty->GetCPPType(), TEXT("{...}"), StructRootId, InParentId, InPropertyNameBuilder);
+			
 			// Recurse
 			const void* StructContainer = StructProperty->ContainerPtrToValuePtr<const void>(InContainer);
 			for(TFieldIterator<FProperty> It(StructProperty->Struct); It; ++It)
 			{
-				IteratePropertiesRecursive(*It, StructContainer, It->GetName(), InFunction, InId, StructRootId);
+				IteratePropertiesRecursive(*It, StructContainer, It->GetName(), InFunction, InId, StructRootId, InPropertyNameBuilder);
 			}
 		}
 		else
@@ -152,7 +216,14 @@ namespace ObjectPropertyTrace
 				// Handle static array header
 				PropertyParentId = ++InId;
 				FString SizeString = FString::Printf(TEXT("{Num = %d}"), InProperty->ArrayDim);
-				InFunction(InProperty->GetCPPType(), InKey, SizeString, PropertyParentId, InParentId);
+
+				// Build property name id
+				FScopedPropertyName ScopedPropertyName(InPropertyNameBuilder);
+				ScopedPropertyName.Builder.AppendChar(TEXT('.'));
+				ScopedPropertyName.Builder.Append(InProperty->GetName());
+
+				// Trace property
+				InFunction(InProperty->GetCPPType(), SizeString, PropertyParentId, InParentId,  InPropertyNameBuilder);
 			}
 
 			for (int32 StaticIndex = 0; StaticIndex != InProperty->ArrayDim; ++StaticIndex)
@@ -160,20 +231,28 @@ namespace ObjectPropertyTrace
 				const void* ValuePtr = InProperty->ContainerPtrToValuePtr<const void>(InContainer, StaticIndex);
 
 				FString KeyString = InProperty->ArrayDim == 1 ? InKey : FString::Printf(TEXT("[%d]"), StaticIndex);
+
+				// Build property name id
+				FScopedPropertyName ScopedPropertyName(InPropertyNameBuilder);
+				ScopedPropertyName.Builder.AppendChar(TEXT('.'));
+				ScopedPropertyName.Builder.Append(KeyString);
+
+				// Trace property
 				FString ValueString;
 				InProperty->ExportText_Direct(ValueString, ValuePtr, ValuePtr, nullptr, PPF_None);
-
-				InFunction(InProperty->GetCPPType(), KeyString, ValueString, ++InId, PropertyParentId);
+				InFunction(InProperty->GetCPPType(), ValueString, ++InId, PropertyParentId,  InPropertyNameBuilder);
 			}
 		}
 	}
 
 	static void IterateProperties(UStruct* InStruct, const void* InContainer, IterateFunction InFunction)
 	{
+		FStringBuilderBase ParentNameBuilder;
+		
 		int32 Id = INDEX_NONE;
 		for(TFieldIterator<FProperty> It(InStruct); It; ++It)
 		{
-			IteratePropertiesRecursive(*It, InContainer, It->GetName(), InFunction, Id, INDEX_NONE);
+			IteratePropertiesRecursive(*It, InContainer, It->GetName(), InFunction, Id, INDEX_NONE, ParentNameBuilder);
 		}
 	}
 
@@ -186,24 +265,34 @@ namespace ObjectPropertyTrace
 				uint64 StartCycle = FPlatformTime::Cycles64();
 				uint64 ObjectId = FObjectTrace::GetObjectId(TracedObject);
 
-				UE_TRACE_LOG(Object, PropertiesStart, ObjectProperties)
-					<< PropertiesStart.Cycle(StartCycle)
-					<< PropertiesStart.ObjectId(ObjectId);
-				
+				UE_TRACE_LOG(Object, PropertiesStart2, ObjectProperties)
+					<< PropertiesStart2.Cycle(StartCycle)
+					<< PropertiesStart2.ObjectId(ObjectId)
+					<< PropertiesStart2.RecordingTime(FObjectTrace::GetWorldElapsedTime(TracedObject->GetWorld()));
+
 				uint64 ClassId = FObjectTrace::GetObjectId(TracedObject->GetClass());
 
-				IterateProperties(TracedObject->GetClass(), TracedObject, [StartCycle, ClassId, ObjectId](const FString& InType, const FString& InKey, const FString& InValue, int32 InId, int32 InParentId)
+				IterateProperties(TracedObject->GetClass(), TracedObject, [StartCycle, ClassId, ObjectId](const FString& InTypeString, const FString& InValue, int32 InId, int32 InParentId, FStringBuilderBase & InPropertyNameBuilder)
 				{
-					uint32 TypeId = TraceStringId(InType);
-					uint32 KeyId = TraceStringId(InKey);
-
-					UE_TRACE_LOG(Object, PropertyValue2, ObjectProperties)
-						<< PropertyValue2.Cycle(StartCycle)
-						<< PropertyValue2.ObjectId(ObjectId)
-						<< PropertyValue2.Value(*InValue)
-						<< PropertyValue2.ParentId(InParentId)
-						<< PropertyValue2.TypeId(TypeId)
-						<< PropertyValue2.KeyId(KeyId);
+					const FString NameString = InPropertyNameBuilder.ToString();
+
+					// Compute parent id from property id 
+					int32 ParentEndIndex;
+					NameString.FindLastChar(TEXT('.'), ParentEndIndex);
+					const FString ParentString = NameString.Left(ParentEndIndex);
+					
+					const uint32 TypeId = TraceStringId(InTypeString);
+					const uint32 NameId = TraceStringId(NameString);
+					const uint32 ParentId = TraceStringId(ParentString);
+	
+					UE_TRACE_LOG(Object, PropertyValue3, ObjectProperties)
+						<< PropertyValue3.Cycle(StartCycle)
+						<< PropertyValue3.ObjectId(ObjectId)
+						<< PropertyValue3.Value(*InValue)
+						<< PropertyValue3.ParentId(InParentId)
+						<< PropertyValue3.TypeId(TypeId)
+						<< PropertyValue3.NameId(NameId)
+						<< PropertyValue3.ParentNameId(ParentId);
 				});
 
 				uint64 EndCycle = FPlatformTime::Cycles64();
diff --git a/Engine/Plugins/Animation/GameplayInsights/Source/GameplayInsights/Private/PropertiesTrack.cpp b/Engine/Plugins/Animation/GameplayInsights/Source/GameplayInsights/Private/PropertiesTrack.cpp
new file mode 100644
--- /dev/null
+++ b/Engine/Plugins/Animation/GameplayInsights/Source/GameplayInsights/Private/PropertiesTrack.cpp
@@ -0,0 +1,330 @@
+﻿// Copyright Epic Games, Inc. All Rights Reserved.
+
+#include "PropertiesTrack.h"
+
+#include "AnimationProvider.h"
+#include "GameplayProvider.h"
+#include "IGameplayProvider.h"
+#include "IRewindDebugger.h"
+#include "TraceServices/Model/AnalysisSession.h"
+#include "TraceServices/Model/Frames.h"
+#include "SObjectPropertiesView.h"
+#include "SPropertyTrackView.h"
+#include "VariantTreeNode.h"
+#include "Styling/SlateIconFinder.h"
+#include "PropertyTrack.h"
+#include "PropertyHelpers.h"
+#include "PropertyWatchManager.h"
+
+#define LOCTEXT_NAMESPACE "FPropertiesTrack"
+
+namespace RewindDebugger
+{
+	///////////////////////////////////////////////////////////////////////////////
+	// Properties Track
+	///////////////////////////////////////////////////////////////////////////////
+
+	FPropertiesTrack::FPropertiesTrack(uint64 InObjectId) : ObjectId(InObjectId), bAreWatchedPropertiesDirty(false), TotalPropertyCountAtStart(0)
+	{
+		Icon = FSlateIconFinder::FindIconForClass(UObject::StaticClass());
+	}
+
+	void FPropertiesTrack::Initialize()
+	{
+		FPropertyWatchManager * PropertyWatchManager = FPropertyWatchManager::Instance();
+		check(PropertyWatchManager)
+
+		// Setup handlers so properties are added and deleted accordingly.
+		PropertyWatchManager->OnPropertyWatched().AddSP(this, &FPropertiesTrack::OnPropertyWatched);
+		PropertyWatchManager->OnPropertyUnwatched().AddSP(this, &FPropertiesTrack::OnPropertyUnwatched);
+
+		// Add any watched properties for this object.
+		for (const auto & WatchedProperty : PropertyWatchManager->GetWatchedProperties(ObjectId))
+		{
+			AddTrackedProperty(WatchedProperty);
+		}
+	}
+
+	int64 FPropertiesTrack::GetTotalPropertyCountAtStart() const
+	{
+		return TotalPropertyCountAtStart;
+	}
+
+	void FPropertiesTrack::AddTrackedProperty(uint32 InPropertyNameId)
+	{
+		// Cache size before modification
+		const int32 PrevSize = ChildrenTracks.Num();
+
+		// Add property and its respective parents if possible.
+		if (InPropertyNameId && !ChildrenTracks.Contains(InPropertyNameId))
+		{
+			const IRewindDebugger* RewindDebugger = IRewindDebugger::Instance();
+			check(RewindDebugger);
+			const TraceServices::IAnalysisSession* Session = RewindDebugger->GetAnalysisSession();
+
+			// Get current visible range 
+			const TRange<double> RecordingRange = RewindDebugger->GetCurrentTraceRange();
+			const double StartTime = RecordingRange.GetLowerBoundValue();
+			const double EndTime = RecordingRange.GetUpperBoundValue();
+
+			// Create track for watched property if found in recording.
+			FObjectPropertyInfo FoundProperty;
+			if (Session && FObjectPropertyHelpers::FindPropertyValueFromNameId(InPropertyNameId, ObjectId, *Session, StartTime, EndTime, FoundProperty))
+			{
+				FPropertyWatchManager * PropertyWatchManager = FPropertyWatchManager::Instance();
+				check(PropertyWatchManager)
+				
+				const TSharedPtr<FPropertyTrack> * PrevPropertyTrack = &ChildrenTracks.Add(InPropertyNameId, FPropertyTrack::Create(ObjectId, MakeShared<FObjectPropertyInfo>(FoundProperty.Name, FoundProperty.Property), nullptr));
+
+				{
+					uint32 NextNameId = FoundProperty.Property.ParentNameId;
+					FObjectPropertyInfo FoundParentProperty {};
+					while (NextNameId)
+					{
+						// Find parent property
+						if (FObjectPropertyHelpers::FindPropertyValueFromNameId(NextNameId, ObjectId, *Session, StartTime, EndTime, FoundParentProperty))
+						{
+							const TSharedPtr<FPropertyTrack> * ParentPropertyTrack = ChildrenTracks.Find(NextNameId);
+
+							// Ensure parent track exists.
+							if (!ParentPropertyTrack)
+							{
+								ParentPropertyTrack = &ChildrenTracks.Add(NextNameId, FPropertyTrack::Create(ObjectId, MakeShared<FObjectPropertyInfo>(FoundParentProperty.Name, FoundParentProperty.Property), nullptr));
+							}
+
+							// Setup track relationships.
+							(*PrevPropertyTrack)->SetParent(*ParentPropertyTrack);
+							(*ParentPropertyTrack)->AddUniqueChild(*PrevPropertyTrack);
+							PrevPropertyTrack = ParentPropertyTrack;
+
+							// Ensure we are also watching the parent property.
+							PropertyWatchManager->WatchProperty(ObjectId, NextNameId);
+							
+							// Get next id
+							NextNameId = FoundParentProperty.Property.ParentNameId; 
+						}
+						else
+						{
+							UE_LOG(LogRewindDebugger, Warning, TEXT("Failed to find parent property with named id {%d} while creating property watch track for {%s}"), FoundProperty.Property.ParentNameId, *FoundProperty.Name.ToString());
+							break;
+						}
+					}
+				}
+			
+				bAreWatchedPropertiesDirty = true;
+			}
+			else
+			{
+				UE_LOG(LogRewindDebugger, Warning, TEXT("Failed to find property with named id {%d} while creating it's watch track."), FoundProperty.Property.NameId, *FoundProperty.Name.ToString());
+			}
+		}
+
+		// Sort tracked properties if any change is detected
+		if (PrevSize != ChildrenTracks.Num())
+		{
+			ChildrenTracks.KeySort([](const uint32 A, const uint32 B)
+			{
+				return A < B;
+			});
+		}
+	}
+
+	void FPropertiesTrack::RemoveTrackedProperty(uint32 InPropertyNameId)
+	{
+		if (InPropertyNameId)
+		{
+			FPropertyWatchManager * PropertyWatchManager = FPropertyWatchManager::Instance();
+			check(PropertyWatchManager)
+			
+			// Cache size before modification
+			const int32 PrevSize = ChildrenTracks.Num();
+
+			if (const TSharedPtr<FPropertyTrack> * ToRemoveProperty = ChildrenTracks.Find(InPropertyNameId))
+			{
+				// Remove children of property
+				TArray<TConstArrayView<TSharedPtr<FPropertyTrack>>> Stack;
+				Stack.Push((*ToRemoveProperty)->GetChildren());
+
+				// Remove property from tracking map and parent if any.
+				{
+					if (const uint32 ParentNameId = (*ToRemoveProperty)->GetObjectProperty()->Property.ParentNameId)
+					{
+						const TSharedPtr<FPropertyTrack> * ParentProperty = ChildrenTracks.Find(ParentNameId);
+
+						if (ParentProperty)
+						{
+							(*ParentProperty)->RemoveChild(*ToRemoveProperty);
+						}
+						else
+						{
+							UE_LOG(LogRewindDebugger, Warning, TEXT("Unable to find watch track for parent property with id {%d}"), ParentNameId);
+						}
+					}
+
+					ChildrenTracks.Remove((*ToRemoveProperty)->GetObjectProperty()->Property.NameId);
+
+					// Ensure we are also not watching the property.
+					PropertyWatchManager->UnwatchProperty(ObjectId, InPropertyNameId);
+				}
+			
+				// Remove property tree from tracked properties
+				while (!Stack.IsEmpty())
+				{
+					// Get parent property
+					TConstArrayView<TSharedPtr<FPropertyTrack>> Parent = Stack.Top();
+					Stack.Pop(false);
+				
+					// Ensure child properties are deleted
+					for (const TSharedPtr<FPropertyTrack> & Child : Parent)
+					{
+						// Remove child property
+						ChildrenTracks.Remove(Child->GetObjectProperty()->Property.NameId);
+
+						// Ensure we are also not watching the property.
+						PropertyWatchManager->UnwatchProperty(ObjectId, Child->GetObjectProperty()->Property.NameId);
+
+						// Add children for deletion
+						Stack.Push(Child->GetChildren());
+					}
+				}
+			
+				bAreWatchedPropertiesDirty = true;
+			}
+
+			// Sort tracked properties if any change is detected
+			if (PrevSize != ChildrenTracks.Num())
+			{
+				ChildrenTracks.KeySort([](const uint32 A, const uint32 B)
+				{
+					return A < B;
+				});
+			}
+		}
+	}
+
+	bool FPropertiesTrack::UpdateInternal()
+	{
+		bool bChanged = bAreWatchedPropertiesDirty;
+		bAreWatchedPropertiesDirty = false;
+			
+		// Update root properties
+		for (const TTuple<uint32, TSharedPtr<FPropertyTrack>> & PropertyWatch : ChildrenTracks)
+		{
+			// Base node.
+			if (PropertyWatch.Value->GetObjectProperty()->Property.ParentId == INDEX_NONE)
+			{
+				bChanged = PropertyWatch.Value->Update() || bChanged;
+			}
+		}
+		
+		return bChanged;
+	}
+
+	TSharedPtr<SWidget> FPropertiesTrack::GetDetailsViewInternal()
+	{
+		// Properties View
+		const IRewindDebugger* RewindDebugger = IRewindDebugger::Instance();
+		TSharedRef<SObjectPropertiesView> Widget = SNew(SObjectPropertiesView, ObjectId, RewindDebugger->CurrentTraceTime(), *RewindDebugger->GetAnalysisSession())
+		.CurrentTime_Raw(RewindDebugger, &IRewindDebugger::CurrentTraceTime);
+		
+		return Widget;
+	}
+
+	FSlateIcon FPropertiesTrack::GetIconInternal()
+	{
+		return Icon;
+	}
+
+	FName FPropertiesTrack::GetNameInternal() const
+	{
+		static FName Name = "PropertyWatch";
+		return Name;
+	}
+
+	FText FPropertiesTrack::GetDisplayNameInternal() const
+	{
+		static FText DisplayName = LOCTEXT("PropertyTrackName", "Properties");
+		return DisplayName;
+	}
+
+	uint64 FPropertiesTrack::GetObjectIdInternal() const
+	{
+		return ObjectId;
+	}
+
+	void FPropertiesTrack::IterateSubTracksInternal(TFunction<void(TSharedPtr<FRewindDebuggerTrack> SubTrack)> IteratorFunction)
+	{
+		// Iterate through all the parent tracks
+		for (const TTuple<uint32, TSharedPtr<FPropertyTrack>> & PropertyWatch : ChildrenTracks)
+		{
+			if (!PropertyWatch.Value->GetObjectProperty()->Property.ParentNameId)
+			{
+				IteratorFunction(PropertyWatch.Value);
+			}
+		}
+	}
+
+	void FPropertiesTrack::OnPropertyWatched(uint64 InObjectId, uint32 InPropertyNameId)
+	{
+		if (InObjectId == ObjectId)
+		{
+			AddTrackedProperty(InPropertyNameId);
+		}
+	}
+
+	void FPropertiesTrack::OnPropertyUnwatched(uint64 InObjectId, uint32 InPropertyNameId)
+	{
+		if (InObjectId == ObjectId)
+		{
+			RemoveTrackedProperty(InPropertyNameId);
+		}
+	}
+
+	///////////////////////////////////////////////////////////////////////////////
+	// Properties Track Creator
+	///////////////////////////////////////////////////////////////////////////////
+	
+	FName FPropertiesTrackCreator::GetTargetTypeNameInternal() const
+	{
+		static const FName TargetTypeName("Object");
+		return TargetTypeName;
+	}
+
+	FName FPropertiesTrackCreator::GetNameInternal() const
+	{
+		static const FName PropertiesTrackCreatorName("PropertyWatch");
+		return PropertiesTrackCreatorName;
+	}
+
+	TSharedPtr<FRewindDebuggerTrack> FPropertiesTrackCreator::CreateTrackInternal(uint64 ObjectId) const
+	{
+		TSharedRef<FPropertiesTrack> NewTrack = MakeShared<FPropertiesTrack>(ObjectId);
+
+		NewTrack->Initialize();
+		
+		return NewTrack;
+	}
+
+	bool FPropertiesTrackCreator::HasDebugInfoInternal(uint64 ObjectId) const
+	{
+		bool bHasData = false;
+
+		if (const TraceServices::IAnalysisSession* AnalysisSession = IRewindDebugger::Instance()->GetAnalysisSession())
+		{
+			TraceServices::FAnalysisSessionReadScope SessionReadScope(*AnalysisSession);
+		
+			if (const FGameplayProvider* GameplayProvider = AnalysisSession->ReadProvider<FGameplayProvider>(FGameplayProvider::ProviderName))
+			{
+				GameplayProvider->ReadObjectPropertiesTimeline(ObjectId, [&bHasData](const IGameplayProvider::ObjectPropertiesTimeline & ObjPropsTimeline)
+				{
+					bHasData = true;
+				});
+			}
+		}
+
+		return bHasData;
+	}
+		
+}
+
+#undef LOCTEXT_NAMESPACE
diff --git a/Engine/Plugins/Animation/GameplayInsights/Source/GameplayInsights/Private/PropertiesTrack.h b/Engine/Plugins/Animation/GameplayInsights/Source/GameplayInsights/Private/PropertiesTrack.h
new file mode 100644
--- /dev/null
+++ b/Engine/Plugins/Animation/GameplayInsights/Source/GameplayInsights/Private/PropertiesTrack.h
@@ -0,0 +1,82 @@
+﻿// Copyright Epic Games, Inc. All Rights Reserved.
+
+#pragma once
+
+#include "IRewindDebuggerTrackCreator.h"
+#include "RewindDebuggerTrack.h"
+#include "Containers/Map.h"
+#include "IGameplayProvider.h"
+
+namespace RewindDebugger
+{
+	struct FObjectPropertyItem;
+	class FPropertyTrack;
+	
+	/**
+	 * Root track used to store an object's traced properties information and property tracks
+	 */
+	class FPropertiesTrack : public FRewindDebuggerTrack, public TSharedFromThis<FPropertiesTrack>
+	{
+	public:
+		
+		/** Constructor */
+		FPropertiesTrack(uint64 InObjectId);
+
+		/** Initialize properties track */
+		void Initialize();
+		
+		/** @return Total property count for traced object at the start of recording */
+		int64 GetTotalPropertyCountAtStart() const;
+		
+		/** Add a property from the Rewind Debugger timeline*/
+		void AddTrackedProperty(uint32 InPropertyNameId);
+		
+		/** Remove a property from the Rewind Debugger timeline*/
+		void RemoveTrackedProperty(uint32 InPropertyNameId);
+	
+	private:
+		/** Begin IRewindDebuggerTrack interface */
+		virtual bool UpdateInternal() override;
+		virtual TSharedPtr<SWidget> GetDetailsViewInternal() override;
+		virtual FSlateIcon GetIconInternal() override;
+		virtual FName GetNameInternal() const override;
+		virtual FText GetDisplayNameInternal() const override;
+		virtual uint64 GetObjectIdInternal() const override;
+		virtual void IterateSubTracksInternal(TFunction<void(TSharedPtr<FRewindDebuggerTrack> SubTrack)> IteratorFunction) override;
+		/** End IRewindDebuggerTrack interface */
+
+		/** Event handler used to create tracks for watched properties */
+		void OnPropertyWatched(uint64 InObjectId, uint32 InPropertyNameId);
+
+		/** Event handler used to delete tracks for an unwatched properties */
+		void OnPropertyUnwatched(uint64 InObjectId, uint32 InPropertyNameId);
+		
+		/** Used to determine track icon */
+		FSlateIcon Icon;
+		
+		/** Id for traced/target object */
+		uint64 ObjectId;
+
+		/** List of watch property tracks to display on Rewind Debugger timeline */
+		TMap<uint32, TSharedPtr<FPropertyTrack>> ChildrenTracks;
+
+		/** Used to determine if watched tracks have been externally modified */
+		bool bAreWatchedPropertiesDirty;
+
+		/** Count of recorded properties for traced object at the start of recording */
+		int64 TotalPropertyCountAtStart;
+	};
+	
+	/**
+	 * Used to automatically create an FPropertiesTrack in Rewind Debugger
+	 */
+	class FPropertiesTrackCreator : public IRewindDebuggerTrackCreator
+	{
+		/** Begin IRewindDebuggerTrackCreator interface */
+		virtual FName GetTargetTypeNameInternal() const override;
+		virtual FName GetNameInternal() const override;
+		virtual TSharedPtr<FRewindDebuggerTrack> CreateTrackInternal(uint64 ObjectId) const override;
+		virtual bool HasDebugInfoInternal(uint64 ObjectId) const override;
+		/** End IRewindDebuggerTrackCreator interface */
+	};
+}
\ No newline at end of file
diff --git a/Engine/Plugins/Animation/GameplayInsights/Source/GameplayInsights/Private/PropertyHelpers.cpp b/Engine/Plugins/Animation/GameplayInsights/Source/GameplayInsights/Private/PropertyHelpers.cpp
new file mode 100644
--- /dev/null
+++ b/Engine/Plugins/Animation/GameplayInsights/Source/GameplayInsights/Private/PropertyHelpers.cpp
@@ -0,0 +1,259 @@
+﻿// Copyright Epic Games, Inc. All Rights Reserved.
+
+#include "PropertyHelpers.h"
+#include "PropertyTrack.h"
+#include "GameplayProvider.h"
+#include "VariantTreeNode.h"
+
+FLinearColor FObjectPropertyHelpers::GetPropertyColor(const FObjectPropertyValue & InProperty)
+{
+	const IRewindDebugger* RewindDebugger = IRewindDebugger::Instance();
+	check(RewindDebugger)
+
+	if (const TraceServices::IAnalysisSession* AnalysisSession = RewindDebugger->GetAnalysisSession())
+	{
+		TraceServices::FAnalysisSessionReadScope SessionReadScope(*AnalysisSession);
+		
+		if (const FGameplayProvider* GameplayProvider = AnalysisSession->ReadProvider<FGameplayProvider>(FGameplayProvider::ProviderName))
+		{
+			// Cached type strings
+			static const FString FloatString = TEXT("float");
+			static const FString DoubleString = TEXT("double");
+			static const FString UInt8String = TEXT("uint8");
+			static const FString Int32String = TEXT("int32");
+			static const FString Int64String = TEXT("int64");
+
+			// Property info strings
+			const FString TypeString = GameplayProvider->GetPropertyName(InProperty.TypeStringId);
+			const FString ValueString = InProperty.Value;
+			
+			// Note: Currently only supporting color for types that are being rendered as tracks in Rewind Debugger.
+			// Additionally we are using hardcoded color values at the moment due standalone UnrealInsights being unable to depend on editor systems.
+			if (TypeString.Equals(FloatString))
+			{
+				return FLinearColor(0.357667f, 1.0f, 0.060000f, 1.0f);
+			}
+			else if (TypeString.Equals(DoubleString))
+			{
+				return FLinearColor(0.039216f, 0.666667f, 0.0f, 1.0f);
+			}
+			else if (IsBoolProperty(ValueString, TypeString))
+			{
+				return FLinearColor(0.300000f, 0.0f, 0.0f, 1.0f);
+			}
+			else if (TypeString.Equals(UInt8String))
+			{
+				return FLinearColor(0.0f, 0.160000f, 0.131270f, 1.0f);
+			}
+			else if (TypeString.Equals(Int32String))
+			{
+				return FLinearColor(0.013575f, 0.770000f, 0.429609f, 1.0f);
+			}
+			else if (TypeString.Equals(Int64String))
+			{
+				return FLinearColor(0.413575f, 0.770000f, 0.429609f, 1.0f);
+			}
+			else if (ValueString.IsNumeric())
+			{
+				return FLinearColor(0.607717f, 0.224984f, 1.0f, 1.0f);
+			}
+		}
+	}
+	
+	return FLinearColor(0.750000f, 0.6f, 0.4f, 1.0f);
+}
+
+TSharedRef<FVariantTreeNode> FObjectPropertyHelpers::GetVariantNodeFromProperty(uint32 InPropertyIndex, const IGameplayProvider& InGameplayProvider, const TConstArrayView<FObjectPropertyValue>& InStorage)
+{
+	static const FString Int32String = TEXT("int32");
+	static const FString FloatString = TEXT("float");
+	static const FString FVectorString = TEXT("FVector");
+	static const FString FVector2DString = TEXT("FVector2D");
+	static const FString FQuatString = TEXT("FQuat");
+	
+	const FObjectPropertyValue & InValue = InStorage[InPropertyIndex];
+
+	const FText DisplayName = FText::FromName(GetPropertyDisplayName(InValue, InGameplayProvider));
+	const FString PropertyTypeString = InGameplayProvider.GetPropertyName(InValue.TypeStringId);
+
+	TSharedPtr<FVariantTreeNode> OutputNode;
+
+	if (IsBoolProperty(InValue.Value, PropertyTypeString))
+	{
+		OutputNode = FVariantTreeNode::MakeBool(DisplayName, FString(InValue.Value).ToBool());
+	}
+	else if (PropertyTypeString.Equals(Int32String))
+	{
+		OutputNode = FVariantTreeNode::MakeInt32(DisplayName, InValue.ValueAsFloat);
+	}
+	else if (PropertyTypeString.Equals(FloatString))
+	{
+		OutputNode = FVariantTreeNode::MakeFloat(DisplayName, InValue.ValueAsFloat);
+	}
+	else if (PropertyTypeString.Equals(FVectorString))
+	{
+		const FVector& Vec3 = FVector(InStorage[InPropertyIndex + 1].ValueAsFloat,InStorage[InPropertyIndex + 2].ValueAsFloat,InStorage[InPropertyIndex + 3].ValueAsFloat);
+		OutputNode = FVariantTreeNode::MakeVector(DisplayName, Vec3);
+	}
+	else if (PropertyTypeString.Equals(FVector2DString))
+	{
+		const FVector2D Vec2D = { InStorage[InPropertyIndex + 1].ValueAsFloat, InStorage[InPropertyIndex + 2].ValueAsFloat };
+		OutputNode = FVariantTreeNode::MakeVector2D(DisplayName, Vec2D);
+	}
+	else if (PropertyTypeString.Equals(FQuatString))
+	{
+		const FQuat Quat = { InStorage[InPropertyIndex + 1].ValueAsFloat, InStorage[InPropertyIndex + 2].ValueAsFloat, InStorage[InPropertyIndex + 3].ValueAsFloat, InStorage[InPropertyIndex + 4].ValueAsFloat };
+		OutputNode = FVariantTreeNode::MakeVector(DisplayName, Quat.Euler());
+	}
+	else
+	{
+		OutputNode = FVariantTreeNode::MakeString(DisplayName, InValue.Value);
+	}
+
+	OutputNode->SetPropertyNameId(InValue.NameId);
+
+	return OutputNode.ToSharedRef();
+}
+
+bool FObjectPropertyHelpers::IsBoolProperty(const FString & InPropertyValueString, const FString & InPropertyTypeString)
+{
+	static const FString BoolString = TEXT("bool");
+	static const FString FalseString = TEXT("false");
+	static const FString TrueString = TEXT("true");
+
+	// We have to check the type string for a "bool" if the FBoolProperty was native bool. However, we also have to handle the case were is not a native bool but instead a bitfield.
+	return InPropertyTypeString.Equals(BoolString) || InPropertyValueString.Equals(FalseString, ESearchCase::IgnoreCase) || InPropertyValueString.Equals(TrueString, ESearchCase::IgnoreCase);
+}
+
+FName FObjectPropertyHelpers::GetPropertyDisplayName(const FObjectPropertyValue & InProperty, const IGameplayProvider& InGameplayProvider)
+{
+	// Query property full name.
+	const TCHAR * Name = InGameplayProvider.GetPropertyName(InProperty.NameId);
+	FString StringName { Name };
+
+	// Get property short name.
+	{
+		int32 SplitPos;
+		if (StringName.FindLastChar('.', SplitPos))
+		{
+			StringName = StringName.RightChop(SplitPos + 1);
+		}
+	}
+
+	// Ensure human readable names.
+	const FName OutputName = *FName::NameToDisplayString(StringName, IsBoolProperty(InProperty.Value, InGameplayProvider.GetPropertyName(InProperty.TypeStringId)));
+	
+	return OutputName;
+}
+
+FSlateIcon FObjectPropertyHelpers::GetPropertyIcon(const FObjectPropertyValue & InProperty)
+{
+	const IRewindDebugger* RewindDebugger = IRewindDebugger::Instance();
+	check(RewindDebugger)
+
+	if (const TraceServices::IAnalysisSession* AnalysisSession = RewindDebugger->GetAnalysisSession())
+	{
+		TraceServices::FAnalysisSessionReadScope SessionReadScope(*AnalysisSession);
+		
+		if (const FGameplayProvider* GameplayProvider = AnalysisSession->ReadProvider<FGameplayProvider>(FGameplayProvider::ProviderName))
+		{
+			const FString PropertyTypeString = GameplayProvider->GetPropertyName(InProperty.TypeStringId);
+		
+			if (PropertyTypeString.StartsWith(TEXT("TArray")))
+			{
+				return FSlateIcon("EditorStyle", "Kismet.VariableList.ArrayTypeIcon");
+			}
+			else if (PropertyTypeString.StartsWith(TEXT("TMap")))
+			{
+				return FSlateIcon("EditorStyle", "Kismet.VariableList.MapKeyTypeIcon", NAME_None, "Kismet.VariableList.MapValueTypeIcon");
+			}
+			else if (PropertyTypeString.StartsWith(TEXT("TSet")))
+			{
+				return FSlateIcon("EditorStyle", "Kismet.VariableList.SetTypeIcon");
+			}
+		}
+	}
+	
+	return FSlateIcon("EditorStyle", "Kismet.VariableList.TypeIcon");
+}
+
+TPair<const FObjectPropertyValue *, uint32> FObjectPropertyHelpers::ReadObjectPropertyValueCached(RewindDebugger::FObjectPropertyInfo& InProperty, uint64 InObjectId, const IGameplayProvider& InGameplayProvider, const FObjectPropertiesMessage& InMessage)
+{
+	TPair<const FObjectPropertyValue *, uint32> Output = { nullptr, 0 };
+	
+	InGameplayProvider.ReadObjectPropertiesStorage(InObjectId, InMessage, [&InProperty, &InMessage, &Output](const TConstArrayView<FObjectPropertyValue> & InStorage)
+	{
+		// Attempt to use cached index
+		if (InProperty.CachedId != INDEX_NONE && InStorage.IsValidIndex(InProperty.CachedId))
+		{
+			const FObjectPropertyValue & FoundPropertyValue = InStorage[InProperty.CachedId];
+
+			// Found property used previous known index.
+			if (InProperty.Property.NameId == FoundPropertyValue.NameId)
+			{
+				Output.Key = &FoundPropertyValue;
+				Output.Value = InProperty.CachedId;
+				return;
+			}
+		}
+
+		// No constant time access, now using linear search.
+		{
+			const int64 TotalPropertyCount = InMessage.PropertyValueEndIndex - InMessage.PropertyValueStartIndex;
+
+			for (int64 i = 0; i < TotalPropertyCount; ++i)
+			{
+				if (InStorage.IsValidIndex(i))
+				{
+					const FObjectPropertyValue & FoundPropertyValue = InStorage[i];
+					if (InProperty.Property.NameId == FoundPropertyValue.NameId)
+					{
+						InProperty.CachedId = i;
+											
+						Output.Key = &FoundPropertyValue;
+						Output.Value = InProperty.CachedId;
+						return;
+					}
+				}
+			}
+		}
+	});
+
+	return Output;
+}
+
+bool FObjectPropertyHelpers::FindPropertyValueFromNameId(uint32 InPropertyNameId, uint64 InObjectId, const TraceServices::IAnalysisSession & InSession, double InStartTime, double InEndTime, RewindDebugger::FObjectPropertyInfo & InOutProperty)
+{
+	bool bFoundProperty = false;
+
+	TraceServices::FAnalysisSessionReadScope SessionReadScope(InSession);
+
+	if (const FGameplayProvider* GameplayProvider = InSession.ReadProvider<FGameplayProvider>(FGameplayProvider::ProviderName))
+	{
+		GameplayProvider->ReadObjectPropertiesTimeline(InObjectId, [InStartTime, InEndTime, GameplayProvider, InObjectId, InPropertyNameId, &bFoundProperty, &InOutProperty](const IGameplayProvider::ObjectPropertiesTimeline & InObjPropsTimeline)
+		{
+			InObjPropsTimeline.EnumerateEvents(InStartTime, InEndTime, [GameplayProvider, InObjectId, InPropertyNameId, &bFoundProperty, &InOutProperty](double InStartTimeProcedure, double InEndTimeProcedure, uint32 InDepth, const FObjectPropertiesMessage& InMessage) -> TraceServices::EEventEnumerate
+			{
+				GameplayProvider->ReadObjectPropertiesStorage(InObjectId, InMessage, [GameplayProvider, InPropertyNameId, &bFoundProperty, &InOutProperty](const TConstArrayView<FObjectPropertyValue> & InStorage)
+				{
+					for (int32 i = 0; i < InStorage.Num(); ++i)
+					{
+						if (!bFoundProperty && InStorage[i].NameId == InPropertyNameId)
+						{
+							InOutProperty.Property = InStorage[i];
+							InOutProperty.Name = GetPropertyDisplayName(InStorage[i], *GameplayProvider);
+							InOutProperty.CachedId = INDEX_NONE;
+							
+							bFoundProperty = true;
+							return;
+						}
+					}
+				});
+
+				return bFoundProperty ? TraceServices::EEventEnumerate::Stop : TraceServices::EEventEnumerate::Continue;
+			});
+		});
+	}
+	
+	return bFoundProperty;
+};
\ No newline at end of file
diff --git a/Engine/Plugins/Animation/GameplayInsights/Source/GameplayInsights/Private/PropertyHelpers.h b/Engine/Plugins/Animation/GameplayInsights/Source/GameplayInsights/Private/PropertyHelpers.h
new file mode 100644
--- /dev/null
+++ b/Engine/Plugins/Animation/GameplayInsights/Source/GameplayInsights/Private/PropertyHelpers.h
@@ -0,0 +1,47 @@
+﻿// Copyright Epic Games, Inc. All Rights Reserved.
+
+
+#pragma once
+#include "IGameplayProvider.h"
+#include "SlateCore/Public/Textures/SlateIcon.h"
+
+struct FVariantTreeNode;
+
+namespace RewindDebugger
+{
+	struct FObjectPropertyInfo;
+}
+
+/**
+ * Helper class when working with traced object properties
+ */
+class FObjectPropertyHelpers
+{
+public:
+
+	/** @return True if traced property was a FBoolProperty. */
+	static bool IsBoolProperty(const FString & InPropertyValueString, const FString & InPropertyTypeString);
+	
+	/** @return Get formatted property display name */
+	static FName GetPropertyDisplayName(const FObjectPropertyValue & InProperty, const IGameplayProvider& InGameplayProvider);
+	
+	/** @return Icon associated with the given property's type */
+	static FSlateIcon GetPropertyIcon(const FObjectPropertyValue & InProperty);
+
+	/** @return Color associated with the given property's type */
+	static FLinearColor GetPropertyColor(const FObjectPropertyValue & InProperty);
+
+	/** @return Variant display node from traced property
+	 *  @note Assumes an FAnalysisSessionReadScope has been created before accessing data in InStorage variable.
+	 */
+	static TSharedRef<FVariantTreeNode> GetVariantNodeFromProperty(uint32 InPropertyIndex, const IGameplayProvider & InGameplayProvider, const TConstArrayView<FObjectPropertyValue> & InStorage);
+	
+	/** @return Read track's property traced value */
+	static TPair<const FObjectPropertyValue *, uint32> ReadObjectPropertyValueCached(RewindDebugger::FObjectPropertyInfo & InProperty, uint64 InObjectId, const IGameplayProvider & InGameplayProvider, const FObjectPropertiesMessage& InMessage);
+		
+	/**
+	 * Find object property value given a named id.
+	 * @return True if property was found, false otherwise.
+	 */
+	static bool FindPropertyValueFromNameId(uint32 InPropertyNameId, uint64 InObjectId, const TraceServices::IAnalysisSession & InSession, double InStartTime, double InEndTime, RewindDebugger::FObjectPropertyInfo & InOutProperty);
+};
diff --git a/Engine/Plugins/Animation/GameplayInsights/Source/GameplayInsights/Private/PropertyTrack.cpp b/Engine/Plugins/Animation/GameplayInsights/Source/GameplayInsights/Private/PropertyTrack.cpp
new file mode 100644
--- /dev/null
+++ b/Engine/Plugins/Animation/GameplayInsights/Source/GameplayInsights/Private/PropertyTrack.cpp
@@ -0,0 +1,254 @@
+﻿// Copyright Epic Games, Inc. All Rights Reserved.
+
+#include "PropertyTrack.h"
+
+#include "PropertyHelpers.h"
+#include "AnimationProvider.h"
+#include "GameplayProvider.h"
+#include "IRewindDebugger.h"
+// #include "ObjectTrace.h"
+#include "TraceServices/Model/AnalysisSession.h"
+#include "TraceServices/Model/Frames.h"
+#include "SPropertyTrackView.h"
+#include "VariantTreeNode.h"
+#include "BoolPropertyTrack.h"
+#include "NumericPropertyTrack.h"
+#include "PropertyWatchManager.h"
+
+#if WITH_EDITOR
+#include "ToolMenuSection.h"
+#endif
+
+#define LOCTEXT_NAMESPACE "FPropertyTrack"
+
+namespace RewindDebugger
+{
+	///////////////////////////////////////////////////////////////////////////////
+	// Property Track
+	///////////////////////////////////////////////////////////////////////////////
+
+	FPropertyTrack::FPropertyTrack(uint64 InObjectId, const TSharedPtr<FObjectPropertyInfo> & InObjectProperty, const TSharedPtr<FPropertyTrack> & InParentTrack) : ObjectId(InObjectId), ObjectProperty(InObjectProperty), Parent(InParentTrack)
+	{
+		Icon = FObjectPropertyHelpers::GetPropertyIcon(InObjectProperty->Property);
+	}
+
+	void FPropertyTrack::BuildContextMenu(FToolMenuSection& InMenuSection)
+	{
+#if WITH_EDITOR
+		const IRewindDebugger* RewindDebugger = IRewindDebugger::Instance();
+		check(RewindDebugger)
+
+		InMenuSection.AddMenuEntry(NAME_None,
+			LOCTEXT("RemovePropertyTrack", "Remove Property Track"),
+			LOCTEXT("RemovePropertyTrack", "Remove traced property track from timeline"),
+			{},
+			FUIAction(FExecuteAction::CreateLambda([this]()
+			{
+				FPropertyWatchManager * PropertyWatchManager = FPropertyWatchManager::Instance();
+				check(PropertyWatchManager)
+				
+				if (!PropertyWatchManager->UnwatchProperty(ObjectId, ObjectProperty->Property.NameId))
+				{
+					UE_LOG(LogRewindDebugger, Warning, TEXT("Failed to unwatch property with id {%d}"), ObjectProperty->Property.NameId);
+				}
+			})));
+#endif
+	}
+
+	const TSharedPtr<FObjectPropertyInfo>& FPropertyTrack::GetObjectProperty() const
+	{
+		return ObjectProperty;
+	}
+
+	const FName& FPropertyTrack::GetPropertyName() const
+	{
+		return ObjectProperty->Name;
+	}
+
+	TConstArrayView<TSharedPtr<FPropertyTrack>> FPropertyTrack::GetChildren() const
+	{
+		return Children;
+	}
+
+	const TSharedPtr<FPropertyTrack>& FPropertyTrack::GetParent() const
+	{
+		return Parent;
+	}
+
+	bool FPropertyTrack::AddUniqueChild(const TSharedPtr<FPropertyTrack>& InTrack)
+	{
+		bool bSuccess = false;
+			
+		if (InTrack.IsValid())
+		{
+			const int32 PrevSize = Children.Num();
+			Children.AddUnique(InTrack);
+			bSuccess = PrevSize != Children.Num();
+
+			if (bSuccess)
+			{
+				// Sort by distance from parent or start.
+				Children.StableSort([](const TSharedPtr<FPropertyTrack>& A, const TSharedPtr<FPropertyTrack>& B)
+				{
+					const FObjectPropertyValue& PropA = A.Get()->GetObjectProperty()->Property;
+					const FObjectPropertyValue& PropB = B.Get()->GetObjectProperty()->Property;
+					
+					return PropA.NameId < PropB.NameId;
+				});
+			}
+		}
+		
+		return bSuccess;
+	}
+
+	bool FPropertyTrack::RemoveChild(const TSharedPtr<FPropertyTrack>& InTrack)
+	{
+		bool bSuccess = false;
+			
+		if (InTrack.IsValid())
+		{
+			Children.Remove(InTrack);
+			bSuccess = true;
+		}
+
+		return bSuccess;
+	}
+
+	void FPropertyTrack::SetParent(const TSharedPtr<FPropertyTrack>& InParent)
+	{
+		Parent = InParent;
+	}
+
+	bool FPropertyTrack::ReadObjectPropertyValueAtFrame(const TraceServices::FFrame& InFrame, const IGameplayProvider& InGameplayProvider, PropertyAtFrameCallback InCallback) const
+	{
+		bool bFoundProperty = false;
+		
+		InGameplayProvider.ReadObjectPropertiesTimeline(ObjectId, [this, &InGameplayProvider, &InFrame, &InCallback, &bFoundProperty](const IGameplayProvider::ObjectPropertiesTimeline & InObjPropsTimeline)
+		{
+			InObjPropsTimeline.EnumerateEvents(InFrame.StartTime, InFrame.EndTime, [this, &InGameplayProvider, &InCallback, &bFoundProperty](double InStartTime, double InEndTime, uint32 InDepth, const FObjectPropertiesMessage& InObjPropsMessage) -> TraceServices::EEventEnumerate
+			{
+				const auto & [ReadProperty, ReadPropertyIndex] = FObjectPropertyHelpers::ReadObjectPropertyValueCached(*ObjectProperty, ObjectId, InGameplayProvider, InObjPropsMessage);
+				
+				if (ReadProperty)
+				{
+					InCallback(*ReadProperty, ReadPropertyIndex, InObjPropsMessage);
+					bFoundProperty = true;
+				}
+
+				return TraceServices::EEventEnumerate::Stop;
+			});
+		});
+
+		return bFoundProperty;
+	}
+	
+	void FPropertyTrack::ReadObjectPropertyValueOverTime(double InStartTime, double InEndTime, PropertyOverTimeCallback InCallback) const
+	{
+		const IRewindDebugger* RewindDebugger = IRewindDebugger::Instance();
+		check(RewindDebugger)
+		
+		if (const TraceServices::IAnalysisSession* AnalysisSession = RewindDebugger->GetAnalysisSession())
+		{
+			TraceServices::FAnalysisSessionReadScope SessionReadScope(*AnalysisSession);
+			
+			if (const FGameplayProvider* GameplayProvider = AnalysisSession->ReadProvider<FGameplayProvider>(FGameplayProvider::ProviderName))
+			{
+				// Query properties timeline
+				GameplayProvider->ReadObjectPropertiesTimeline(ObjectId, [this, &InCallback, InStartTime, InEndTime, GameplayProvider](const IGameplayProvider::ObjectPropertiesTimeline & InObjPropsTimeline)
+				{
+					// Query properties over time
+					InObjPropsTimeline.EnumerateEvents(InStartTime, InEndTime, [this, &InCallback, &InObjPropsTimeline, InStartTime, InEndTime, GameplayProvider](double InStartTimeInner, double InEndTimeInner, uint32 InDepth, const FObjectPropertiesMessage& InObjPropsMessage) -> TraceServices::EEventEnumerate
+					{
+						const auto & [ReadProperty, ReadPropertyIndex] =FObjectPropertyHelpers::ReadObjectPropertyValueCached(*ObjectProperty, ObjectId, *GameplayProvider, InObjPropsMessage);
+						if (ReadProperty)
+						{
+							InCallback(*ReadProperty, ReadPropertyIndex, InObjPropsMessage, InObjPropsTimeline, InStartTime, InEndTime);
+						}
+						
+						return TraceServices::EEventEnumerate::Continue;
+					});
+				});
+			}
+		}
+	}
+
+	TSharedPtr<FPropertyTrack> FPropertyTrack::Create(uint64 InObjectId, const TSharedPtr<FObjectPropertyInfo> & InObjectProperty, const TSharedPtr<FPropertyTrack> & InParentTrack)
+	{
+		// Handle properties with numeric values
+		if (FNumericPropertyTrack::CanBeCreated(*InObjectProperty))
+		{
+			return MakeShared<FNumericPropertyTrack>(InObjectId, InObjectProperty, InParentTrack);
+		}
+
+		// Handle boolean properties
+		if (FBoolPropertyTrack::CanBeCreated(*InObjectProperty))
+		{
+			return MakeShared<FBoolPropertyTrack>(InObjectId, InObjectProperty, InParentTrack);
+		}
+		
+		return MakeShared<FPropertyTrack>(InObjectId, InObjectProperty, InParentTrack);
+	}
+
+	bool FPropertyTrack::CanBeCreated(const FObjectPropertyInfo& InObjectProperty)
+	{
+		return true;
+	}
+
+	bool FPropertyTrack::UpdateInternal()
+	{
+		// Update children
+		bool bChanged = false;
+
+		for (const TSharedPtr<FPropertyTrack> & Child : Children)
+		{
+			bChanged = Child->Update() || bChanged;
+		}
+
+		return bChanged;
+	}
+
+	TSharedPtr<SWidget> FPropertyTrack::GetDetailsViewInternal()
+	{
+		const IRewindDebugger* RewindDebugger = IRewindDebugger::Instance();
+		check(RewindDebugger)
+
+		const TSharedRef<SPropertyTrackView> View = SNew(SPropertyTrackView, ObjectId, RewindDebugger->CurrentTraceTime(), *RewindDebugger->GetAnalysisSession())
+		.CurrentTime_Raw(RewindDebugger, &IRewindDebugger::CurrentTraceTime);
+		
+		View->SetTrack(AsWeak());
+		View->SetForegroundColor(FObjectPropertyHelpers::GetPropertyColor(ObjectProperty->Property));
+		
+		return View;
+	}
+
+	void FPropertyTrack::IterateSubTracksInternal(TFunction<void(TSharedPtr<FRewindDebuggerTrack> SubTrack)> IteratorFunction)
+	{
+		for (TSharedPtr<FPropertyTrack>& Track : Children)
+		{
+			IteratorFunction(Track);
+		}
+	}
+
+	FSlateIcon FPropertyTrack::GetIconInternal()
+	{
+		return Icon;
+	}
+
+	FName FPropertyTrack::GetNameInternal() const
+	{
+		static const FName Name = "VariantProperty";
+		return Name;
+	}
+
+	FText FPropertyTrack::GetDisplayNameInternal() const
+	{
+		return FText::FromName(ObjectProperty->Name);
+	}
+
+	uint64 FPropertyTrack::GetObjectIdInternal() const
+	{
+		return ObjectId;
+	}
+}
+
+#undef LOCTEXT_NAMESPACE
\ No newline at end of file
diff --git a/Engine/Plugins/Animation/GameplayInsights/Source/GameplayInsights/Private/PropertyTrack.h b/Engine/Plugins/Animation/GameplayInsights/Source/GameplayInsights/Private/PropertyTrack.h
new file mode 100644
--- /dev/null
+++ b/Engine/Plugins/Animation/GameplayInsights/Source/GameplayInsights/Private/PropertyTrack.h
@@ -0,0 +1,137 @@
+﻿// Copyright Epic Games, Inc. All Rights Reserved.
+
+#pragma once
+
+#include "RewindDebuggerTrack.h"
+#include "IGameplayProvider.h"
+
+namespace RewindDebugger
+{
+	class FPropertiesTrack;
+	
+	/**
+	 * Container item that stores information for a watched traced object's variable
+	 */
+	struct FObjectPropertyInfo
+	{
+		/** Default Constructor */
+		FObjectPropertyInfo()
+			: Name(NAME_None),
+			  Property(),
+			  CachedId(INDEX_NONE)
+		{
+		}
+		
+		/** Constructor */
+		FObjectPropertyInfo(const FName & Name, const FObjectPropertyValue& TracedProperty, int64 CachedId = INDEX_NONE)
+			: Name(Name),
+			  Property(TracedProperty),
+			  CachedId(CachedId)
+		{
+		}
+
+		/** Traced variable name */
+		FName Name = NAME_None;
+
+		/** Traced variable data. Used to store identifiers. */
+		FObjectPropertyValue Property;
+
+		/** Previously found storage index */
+		int64 CachedId = INDEX_NONE;
+	};
+
+	#define BASE_PROPERTY_TRACK() \
+	static bool CanBeCreated(const FObjectPropertyInfo & InObjectProperty);
+	
+	/**
+	 * Track that displays a traced object's property or traced value over time
+	 */
+	class FPropertyTrack : public FRewindDebuggerTrack, public TSharedFromThis<FPropertyTrack>
+	{
+	public:
+
+		BASE_PROPERTY_TRACK()
+		
+		/** Constructor */
+		FPropertyTrack(uint64 InObjectId, const TSharedPtr<FObjectPropertyInfo> & InObjectProperty, const TSharedPtr<FPropertyTrack> & InParentTrack);
+
+		/** Begin IRewindDebuggerTrack interface */
+		virtual void BuildContextMenu(FToolMenuSection& InMenuSection) override;
+		/** End IRewindDebuggerTrack interface */
+		
+		/** @return Internal data item for traced variable */
+		const TSharedPtr<FObjectPropertyInfo> & GetObjectProperty() const;
+
+		/** @return Property name identifier */
+		const FName & GetPropertyName() const;
+
+		/** @return Children property tracks */
+		TConstArrayView<TSharedPtr<FPropertyTrack>> GetChildren() const;
+
+		/** @return Parent track */
+		const TSharedPtr<FPropertyTrack> & GetParent() const;
+		
+		/**
+		 * Add a property track as a child
+		 * @return True if property was successfully added as a child.
+		 */
+		bool AddUniqueChild(const TSharedPtr<FPropertyTrack> & InTrack);
+		
+		/**
+		 * Remove a child property track
+		 * @return True if property was successfully removed from child container.
+		 */
+		bool RemoveChild(const TSharedPtr<FPropertyTrack> & InTrack);
+
+		/**
+		 * Update the track's parent.
+		 * @note This procedure does not validate if the parent - child relationship is valid.
+		 */
+		void SetParent(const TSharedPtr<FPropertyTrack> & InParent);
+		
+		/** Callback signature for reading a property at a given frame */
+		typedef TFunctionRef<void(const FObjectPropertyValue & /*InValue*/, uint32 /*InValueIndex*/, const FObjectPropertiesMessage & /*InObjPropsMessage*/)> PropertyAtFrameCallback;
+
+		/**
+		 * Read value for internal property at a given frame.
+		 * @note Assumes function is called within a TraceServices::IAnalysisSession read access scope.
+		 * @return True if property value at given frame was found.
+		 */
+		bool ReadObjectPropertyValueAtFrame(const TraceServices::FFrame & InFrame, const IGameplayProvider& InGameplayProvider, PropertyAtFrameCallback InCallback) const;
+
+		/** Callback signature for reading a property at over a time range */
+		typedef TFunctionRef<void(const FObjectPropertyValue & /*InValue*/, uint32 /*InValueIndex*/, const FObjectPropertiesMessage& /*InMessage*/, const IGameplayProvider::ObjectPropertiesTimeline & /* InTimeline */, double /*InStartTime*/, double /*InEndTime*/)> PropertyOverTimeCallback;
+
+		/** Read value for internal property over a range of time. */
+		void ReadObjectPropertyValueOverTime(double InStartTime, double InEndTime, PropertyOverTimeCallback InCallback) const;
+		
+		/** @return Specialized Property Track, if any, for given object property. */
+		static TSharedPtr<FPropertyTrack> Create(uint64 InObjectId, const TSharedPtr<FObjectPropertyInfo> & InObjectProperty, const TSharedPtr<FPropertyTrack> & InParentTrack);
+		
+	protected:
+		/** Begin IRewindDebuggerTrack interface */
+		virtual bool UpdateInternal() override;
+		virtual TSharedPtr<SWidget> GetDetailsViewInternal() override;
+		virtual void IterateSubTracksInternal(TFunction<void(TSharedPtr<FRewindDebuggerTrack> SubTrack)> IteratorFunction) override;
+		virtual FSlateIcon GetIconInternal() override;
+		virtual FName GetNameInternal() const override;
+		virtual FText GetDisplayNameInternal() const override;
+		virtual uint64 GetObjectIdInternal() const override;
+		/** End IRewindDebuggerTrack interface */
+		
+		/** Used to determine track icon */
+		FSlateIcon Icon;
+		
+		/** Id for traced/target object */
+		uint64 ObjectId;
+
+		/** Used to store traced variable information */
+		TSharedPtr<FObjectPropertyInfo> ObjectProperty;
+
+		/** Parent property track */
+		TSharedPtr<FPropertyTrack> Parent;
+		
+		/** Child properties tracks. Used if property is a struct, vector, etc. */
+		TArray<TSharedPtr<FPropertyTrack>> Children;
+	};
+}
\ No newline at end of file
diff --git a/Engine/Plugins/Animation/GameplayInsights/Source/GameplayInsights/Private/PropertyWatchManager.cpp b/Engine/Plugins/Animation/GameplayInsights/Source/GameplayInsights/Private/PropertyWatchManager.cpp
new file mode 100644
--- /dev/null
+++ b/Engine/Plugins/Animation/GameplayInsights/Source/GameplayInsights/Private/PropertyWatchManager.cpp
@@ -0,0 +1,116 @@
+﻿// Copyright Epic Games, Inc. All Rights Reserved.
+
+#include "PropertyWatchManager.h"
+#include "IRewindDebugger.h"
+
+FPropertyWatchManager* FPropertyWatchManager::InternalInstance = nullptr;
+
+FPropertyWatchManager::FPropertyWatchManager()
+{
+}
+
+FPropertyWatchManager::~FPropertyWatchManager()
+{
+}
+
+bool FPropertyWatchManager::WatchProperty(uint64 InObjectId, uint32 InPropertyNameId)
+{
+	if (InPropertyNameId)
+	{
+		TArray<uint32> & WatchStorage = WatchedProperties.FindOrAdd(InObjectId);
+
+		const int32 PrevSize = WatchStorage.Num();
+	
+		WatchStorage.AddUnique(InPropertyNameId);
+
+		// Trigger broadcast only if property was successfully added.
+		const bool bWasPropertyWatched = PrevSize != WatchStorage.Num();
+		if (bWasPropertyWatched)
+		{
+			OnPropertyWatchedDelegate.Broadcast(InObjectId, InPropertyNameId);
+		}
+
+		return bWasPropertyWatched;
+	}
+
+	return false;
+}
+
+bool FPropertyWatchManager::UnwatchProperty(uint64 InObjectId, uint32 InPropertyNameId)
+{
+	if (InPropertyNameId)
+	{
+		if (TArray<uint32> * WatchStorage = WatchedProperties.Find(InObjectId))
+		{
+			const int32 PrevSize = WatchStorage->Num();
+		
+			WatchStorage->RemoveSwap(InPropertyNameId, false);
+
+			// Trigger broadcast only if property was successfully removed.
+			const bool bWasPropertyUnwatched = PrevSize != WatchStorage->Num();
+			if (bWasPropertyUnwatched)
+			{
+				OnPropertyUnwatchedDelegate.Broadcast(InObjectId, InPropertyNameId);
+			}
+			
+			return bWasPropertyUnwatched;
+		}
+	}
+
+	return false;
+}
+
+int32 FPropertyWatchManager::GetWatchedPropertiesNum(uint64 InObjectId) const
+{
+	if (const TArray<uint32> * WatchStorage = WatchedProperties.Find(InObjectId))
+	{
+		return WatchStorage->Num();
+	}
+	return 0;
+}
+
+TConstArrayView<uint32> FPropertyWatchManager::GetWatchedProperties(uint64 InObjectId) const
+{
+	if (const TArray<uint32> * WatchStorage = WatchedProperties.Find(InObjectId))
+	{
+		return MakeArrayView(WatchStorage->GetData(), WatchStorage->Num());
+	}
+	
+	return {};
+}
+
+void FPropertyWatchManager::ClearWatchedProperties(uint64 InObjetId)
+{
+	WatchedProperties.Remove(InObjetId);
+}
+
+void FPropertyWatchManager::ClearAllWatchedProperties()
+{
+	WatchedProperties.Empty();
+}
+
+FPropertyWatchManager::FOnPropertyWatched& FPropertyWatchManager::OnPropertyWatched()
+{
+	return OnPropertyWatchedDelegate;
+}
+
+FPropertyWatchManager::FOnPropertyUnwatched& FPropertyWatchManager::OnPropertyUnwatched()
+{
+	return OnPropertyUnwatchedDelegate;
+}
+
+void FPropertyWatchManager::Initialize()
+{
+	InternalInstance = new FPropertyWatchManager;
+}
+
+void FPropertyWatchManager::Shutdown()
+{
+	delete InternalInstance;
+}
+
+FPropertyWatchManager* FPropertyWatchManager::Instance()
+{
+	check(InternalInstance != nullptr)
+	return InternalInstance;
+}
\ No newline at end of file
diff --git a/Engine/Plugins/Animation/GameplayInsights/Source/GameplayInsights/Private/PropertyWatchManager.h b/Engine/Plugins/Animation/GameplayInsights/Source/GameplayInsights/Private/PropertyWatchManager.h
new file mode 100644
--- /dev/null
+++ b/Engine/Plugins/Animation/GameplayInsights/Source/GameplayInsights/Private/PropertyWatchManager.h
@@ -0,0 +1,81 @@
+﻿// Copyright Epic Games, Inc. All Rights Reserved.
+
+#pragma once
+#include "IGameplayProvider.h"
+
+namespace RewindDebugger
+{
+	struct FObjectPropertyItem;
+}
+
+/**
+ * Used to keep track of traced variables that are being watched.
+ */
+class FPropertyWatchManager
+{
+public:
+
+	// TODO: We need to add code that clears the property watch manager right after a PIE session has started.
+	
+	DECLARE_MULTICAST_DELEGATE_TwoParams(FOnPropertyWatched, uint64 /*InObjectId*/, uint32 /*InPropertyNameId*/)
+	DECLARE_MULTICAST_DELEGATE_TwoParams(FOnPropertyUnwatched, uint64 /*InObjectId*/, uint32 /*InPropertyNameId*/)
+
+	/** Constructor */
+	FPropertyWatchManager();
+
+	/** Destructor */
+	~FPropertyWatchManager();
+
+	/**
+	 * Start watching a property
+	 * @return Whether property was started to being watched. If property is already being watched this will return false.
+	 */
+	bool WatchProperty(uint64 InObjectId, uint32 InPropertyNameId);
+
+	/**
+	 * Stop watching a property
+	 * @return Whether property was stopped being watched. If property is already not watched this will return false.
+	 */
+	bool UnwatchProperty(uint64 InObjectId, uint32 InPropertyNameId);
+
+	/** @return The total number of watched properties for a specific object */
+	int32 GetWatchedPropertiesNum(uint64 InObjectId) const;
+
+	/** @return A view of all the watched properties for a specific object */
+	TConstArrayView<uint32> GetWatchedProperties(uint64 InObjectId) const;
+
+	/** Stop watching all the registered properties of a given object */
+	void ClearWatchedProperties(uint64 InObjetId);
+
+	/** Stop watching all registered properties */
+	void ClearAllWatchedProperties();
+	
+	/** @return Delete that broadcast when any property has started to be watched */
+	FOnPropertyWatched & OnPropertyWatched();
+
+	/** @return Delete that broadcasts when any property has ceased to be watched */
+	FOnPropertyUnwatched & OnPropertyUnwatched();
+	
+	/** Create singleton instance */
+	static void Initialize();
+
+	/** Destroy singleton instance */
+	static void Shutdown();
+	
+	/** @return Current PropertyWatchManager singleton instance */
+	static FPropertyWatchManager* Instance();
+	
+protected:
+	
+	/** ObjectId -> Object's watched properties */
+	TMap<uint64, TArray<uint32>> WatchedProperties;
+
+	/** Delete variable to broadcast when any property has started to be watched */
+	FOnPropertyWatched OnPropertyWatchedDelegate;
+
+	/** Delete variable to broadcast when any property has ceased to be watched */
+	FOnPropertyUnwatched OnPropertyUnwatchedDelegate;
+
+	/** Singleton instance */
+	static FPropertyWatchManager* InternalInstance;
+};
diff --git a/Engine/Plugins/Animation/GameplayInsights/Source/GameplayInsights/Private/SObjectPropertiesView.cpp b/Engine/Plugins/Animation/GameplayInsights/Source/GameplayInsights/Private/SObjectPropertiesView.cpp
--- a/Engine/Plugins/Animation/GameplayInsights/Source/GameplayInsights/Private/SObjectPropertiesView.cpp
+++ b/Engine/Plugins/Animation/GameplayInsights/Source/GameplayInsights/Private/SObjectPropertiesView.cpp
@@ -5,41 +5,47 @@
 #include "TraceServices/Model/AnalysisSession.h"
 #include "GameplayProvider.h"
 #include "IRewindDebugger.h"
-#include "Styling/SlateIconFinder.h"
+#include "PropertiesTrack.h"
+#include "PropertyHelpers.h"
+#include "PropertyWatchManager.h"
 #include "TraceServices/Model/Frames.h"
 #include "VariantTreeNode.h"
+#include "Misc/DefaultValueHelper.h"
+#include "Framework/MultiBox/MultiBoxBuilder.h"
 
 #define LOCTEXT_NAMESPACE "SObjectPropertiesView"
 
 void SObjectPropertiesView::GetVariantsAtFrame(const TraceServices::FFrame& InFrame, TArray<TSharedRef<FVariantTreeNode>>& OutVariants) const
 {
+	TraceServices::FAnalysisSessionReadScope SessionReadScope(*AnalysisSession);
+
 	if (const FGameplayProvider* GameplayProvider = AnalysisSession->ReadProvider<FGameplayProvider>(FGameplayProvider::ProviderName))
 	{
-		TraceServices::FAnalysisSessionReadScope SessionReadScope(*AnalysisSession);
-
-		TSharedRef<FVariantTreeNode> Header = OutVariants.Add_GetRef(FVariantTreeNode::MakeHeader(LOCTEXT("Properties", "Properties"), INDEX_NONE));
-
+		TSharedRef<FVariantTreeNode> Header = OutVariants.Add_GetRef(FVariantTreeNode::MakeHeader(LOCTEXT("PropertiesHeader", "Properties"), INDEX_NONE));
 		TArray<TSharedPtr<FVariantTreeNode>> PropertyVariants;
 
 		GameplayProvider->ReadObjectPropertiesTimeline(ObjectId, [this, &InFrame, GameplayProvider, &Header, &PropertyVariants](const FGameplayProvider::ObjectPropertiesTimeline& InTimeline)
 		{
 			InTimeline.EnumerateEvents(InFrame.StartTime, InFrame.EndTime, [this, GameplayProvider, &Header, &PropertyVariants](double InStartTime, double InEndTime, uint32 InDepth, const FObjectPropertiesMessage& InMessage)
 			{
-				GameplayProvider->EnumerateObjectPropertyValues(ObjectId, InMessage, [GameplayProvider, &Header, &PropertyVariants](const FObjectPropertyValue& InValue)
+				GameplayProvider->ReadObjectPropertiesStorage(ObjectId, InMessage, [&PropertyVariants, &Header, GameplayProvider](const TConstArrayView<FObjectPropertyValue> & InStorage)
 				{
-					const TCHAR* Key = GameplayProvider->GetPropertyName(InValue.KeyStringId);
-					PropertyVariants.Add(FVariantTreeNode::MakeString(FText::FromString(Key), InValue.Value));
-
-					// note assumes that order is parent->child in the properties array
-					if(InValue.ParentId != INDEX_NONE)
+					for (int32 i = 0; i < InStorage.Num(); ++i)
 					{
-						PropertyVariants[InValue.ParentId]->AddChild(PropertyVariants.Last().ToSharedRef());
-					}
-					else
-					{
-						Header->AddChild(PropertyVariants.Last().ToSharedRef());
+						PropertyVariants.Add(FObjectPropertyHelpers::GetVariantNodeFromProperty(i, *GameplayProvider, InStorage));
+						
+						// note assumes that order is parent->child in the properties array
+						if(InStorage[i].ParentId != INDEX_NONE)
+						{
+							PropertyVariants[InStorage[i].ParentId]->AddChild(PropertyVariants.Last().ToSharedRef());
+						}
+						else
+						{
+							Header->AddChild(PropertyVariants.Last().ToSharedRef());
+						}
 					}
 				});
+
 				return TraceServices::EEventEnumerate::Stop;
 			});
 		});
@@ -53,48 +59,33 @@ FName SObjectPropertiesView::GetName() const
 	return ObjectPropertiesName;
 }
 
-FName FObjectPropertiesViewCreator::GetTargetTypeName() const
-{
-	static FName TargetTypeName = "Object";
-	return TargetTypeName;
-}
-
-FName FObjectPropertiesViewCreator::GetName() const
-{
-	return ObjectPropertiesName;
-}
-
-FText FObjectPropertiesViewCreator::GetTitle() const
-{
-	return LOCTEXT("Object Properties", "Properties");
-}
-
-FSlateIcon FObjectPropertiesViewCreator::GetIcon() const
+void SObjectPropertiesView::BuildContextMenu(FMenuBuilder& MenuBuilder)
 {
-	return FSlateIconFinder::FindIconForClass(UObject::StaticClass());
-}
-
-TSharedPtr<IRewindDebuggerView> FObjectPropertiesViewCreator::CreateDebugView(uint64 ObjectId, double CurrentTime, const TraceServices::IAnalysisSession& AnalysisSession) const
-{
-	return SNew(SObjectPropertiesView, ObjectId, CurrentTime, AnalysisSession);
-}
-
-bool FObjectPropertiesViewCreator::HasDebugInfo(uint64 ObjectId) const
-{
-	const TraceServices::IAnalysisSession* AnalysisSession = IRewindDebugger::Instance()->GetAnalysisSession();
-	
-	TraceServices::FAnalysisSessionReadScope SessionReadScope(*AnalysisSession);
-	bool bHasData = false;
-	if (const FGameplayProvider* GameplayProvider = AnalysisSession->ReadProvider<FGameplayProvider>(FGameplayProvider::ProviderName))
-	{
-		GameplayProvider->ReadObjectPropertiesTimeline(ObjectId, [this, &bHasData, GameplayProvider](const FGameplayProvider::ObjectPropertiesTimeline& InTimeline)
+	SPropertiesDebugViewBase::BuildContextMenu(MenuBuilder);
+
+	MenuBuilder.BeginSection(GetName(), LOCTEXT("DetailsLabel", "Property"));
+	MenuBuilder.AddMenuEntry(
+		LOCTEXT("AddPropertyWatch", "Add Property Watch"),
+		LOCTEXT("AddPropertyTooltip", "Adds a watch track for the given property"),
+		{},
+		FExecuteAction::CreateLambda([this]()
 		{
-			bHasData = true;
-		});
-	}
-	
-	return bHasData;
-}
+			FPropertyWatchManager * PropertyWatchManager = FPropertyWatchManager::Instance();
+			check(PropertyWatchManager)
 
+			if (SelectedPropertyId)
+			{
+				if (!PropertyWatchManager->WatchProperty(ObjectId, SelectedPropertyId))
+				{
+					UE_LOG(LogRewindDebugger, Warning, TEXT("Failed to watch property with id {%d}"), SelectedPropertyId);
+				}
+			}
+			else
+			{
+				UE_LOG(LogRewindDebugger, Warning, TEXT("No property was selected or an invalid property was used while attempting to add a property watch track."))
+			}
+		}));
+	MenuBuilder.EndSection();
+}
 
 #undef LOCTEXT_NAMESPACE
diff --git a/Engine/Plugins/Animation/GameplayInsights/Source/GameplayInsights/Private/SObjectPropertiesView.h b/Engine/Plugins/Animation/GameplayInsights/Source/GameplayInsights/Private/SObjectPropertiesView.h
--- a/Engine/Plugins/Animation/GameplayInsights/Source/GameplayInsights/Private/SObjectPropertiesView.h
+++ b/Engine/Plugins/Animation/GameplayInsights/Source/GameplayInsights/Private/SObjectPropertiesView.h
@@ -7,20 +7,15 @@
 
 namespace TraceServices { class IAnalysisSession; }
 
+/**
+ * Used to display the properties of a traced object
+ */
 class SObjectPropertiesView : public SPropertiesDebugViewBase
 {
 public:
+	/** Begin SPropertiesDebugViewBase interface */
 	virtual void GetVariantsAtFrame(const TraceServices::FFrame& InFrame, TArray<TSharedRef<FVariantTreeNode>>& OutVariants) const override;
 	virtual FName GetName() const override;
-};
-
-class FObjectPropertiesViewCreator : public IRewindDebuggerViewCreator
-{
-	public:
-		virtual FName GetTargetTypeName() const;
-		virtual FName GetName() const override;
-		virtual FText GetTitle() const override;
-		virtual FSlateIcon GetIcon() const override;
-		virtual TSharedPtr<IRewindDebuggerView> CreateDebugView(uint64 ObjectId, double CurrentTime, const TraceServices::IAnalysisSession& InAnalysisSession) const override;
-		virtual bool HasDebugInfo(uint64 ObjectId) const override;
+	virtual void BuildContextMenu(FMenuBuilder& MenuBuilder) override;
+	/** End SPropertiesDebugViewBase interface */
 };
\ No newline at end of file
diff --git a/Engine/Plugins/Animation/GameplayInsights/Source/GameplayInsights/Private/SPropertiesDebugViewBase.cpp b/Engine/Plugins/Animation/GameplayInsights/Source/GameplayInsights/Private/SPropertiesDebugViewBase.cpp
--- a/Engine/Plugins/Animation/GameplayInsights/Source/GameplayInsights/Private/SPropertiesDebugViewBase.cpp
+++ b/Engine/Plugins/Animation/GameplayInsights/Source/GameplayInsights/Private/SPropertiesDebugViewBase.cpp
@@ -4,9 +4,10 @@
 #include "AnimationProvider.h"
 #include "TraceServices/Model/AnalysisSession.h"
 #include "GameplayProvider.h"
-#include "Styling/SlateIconFinder.h"
 #include "TraceServices/Model/Frames.h"
 #include "VariantTreeNode.h"
+#include "PropertiesTrack.h"
+#include "Framework/MultiBox/MultiBoxBuilder.h"
 
 #define LOCTEXT_NAMESPACE "SPropertiesDebugViewBase"
 
@@ -15,11 +16,13 @@ void SPropertiesDebugViewBase::Construct(const FArguments& InArgs, uint64 InObje
 	ObjectId = InObjectId;
 	AnalysisSession = &InAnalysisSession;
 	CurrentTime = InArgs._CurrentTime;
-
-	View = SNew(SVariantValueView, InAnalysisSession).OnGetVariantValues(this, &SPropertiesDebugViewBase::GetVariantsAtFrame);
-
-	SetTimeMarker(InTimeMarker);
-
+	SelectedPropertyId = 0;
+	
+	View = SNew(SVariantValueView, InAnalysisSession)
+	.OnGetVariantValues(this, &SPropertiesDebugViewBase::GetVariantsAtFrame)
+	.OnContextMenuOpening(this, &SPropertiesDebugViewBase::CreateContextMenu)
+	.OnMouseButtonDownOnVariantValue(this, &SPropertiesDebugViewBase::HandleOnMouseButtonDown);
+	
 	ChildSlot
 	[
 		View.ToSharedRef()
@@ -51,4 +54,20 @@ void SPropertiesDebugViewBase::Tick( const FGeometry& AllottedGeometry, const do
 	SCompoundWidget::Tick(AllottedGeometry, InCurrentTime, InDeltaTime);
 }
 
+void SPropertiesDebugViewBase::BuildContextMenu(FMenuBuilder& MenuBuilder)
+{
+}
+
+TSharedPtr<SWidget> SPropertiesDebugViewBase::CreateContextMenu() 
+{
+	FMenuBuilder MenuBuilder(true, nullptr);
+	BuildContextMenu(MenuBuilder);
+	return MenuBuilder.MakeWidget();
+}
+
+void SPropertiesDebugViewBase::HandleOnMouseButtonDown(const TSharedPtr<FVariantTreeNode>& InVariantValueNode, const TraceServices::FFrame& InFrame, const FPointerEvent& InKeyEvent)
+{
+	SelectedPropertyId = InVariantValueNode->GetPropertyNamedId();
+}
+
 #undef LOCTEXT_NAMESPACE
diff --git a/Engine/Plugins/Animation/GameplayInsights/Source/GameplayInsights/Private/SPropertiesDebugViewBase.h b/Engine/Plugins/Animation/GameplayInsights/Source/GameplayInsights/Private/SPropertiesDebugViewBase.h
--- a/Engine/Plugins/Animation/GameplayInsights/Source/GameplayInsights/Private/SPropertiesDebugViewBase.h
+++ b/Engine/Plugins/Animation/GameplayInsights/Source/GameplayInsights/Private/SPropertiesDebugViewBase.h
@@ -9,6 +9,8 @@
 
 namespace TraceServices { class IAnalysisSession; }
 
+class FMenuBuilder;
+
 class SPropertiesDebugViewBase : public IRewindDebuggerView
 {
 	SLATE_BEGIN_ARGS(SPropertiesDebugViewBase) {}
@@ -17,18 +19,40 @@ class SPropertiesDebugViewBase : public IRewindDebuggerView
 public:
 	void Construct(const FArguments& InArgs, uint64 InObjectId, double InTimeMarker, const TraceServices::IAnalysisSession& InAnalysisSession);
 
+	/** Begin IRewindDebuggerView interface */
 	virtual void SetTimeMarker(double InTimeMarker) override;
 	virtual uint64 GetObjectId() const override { return ObjectId; }
+	/** End IRewindDebuggerView interface */
+
+	/** Begin SWidget interface */
+	virtual void Tick( const FGeometry& AllottedGeometry, const double InCurrentTime, const float InDeltaTime ) override;
+	/** End SWidget interface */
 
+	/** Get all the variants / properties at the given frame */
 	virtual void GetVariantsAtFrame(const TraceServices::FFrame& InFrame, TArray<TSharedRef<FVariantTreeNode>>& OutVariants) const = 0;
 
-	virtual void Tick( const FGeometry& AllottedGeometry, const double InCurrentTime, const float InDeltaTime ) override;
+	/** Build the context menu for the view widget */
+	virtual void BuildContextMenu(FMenuBuilder& MenuBuilder);
+
+	/** @return Inner view widget for the object variant */
+	const TSharedPtr<SVariantValueView> & GetView() const { return View; };
+	
 protected:
+
+	// Event handlers
+
+	/** Setup base context menu */
+	TSharedPtr<SWidget> CreateContextMenu();
+	
+	/** Handle mouse button event from widget */
+	void HandleOnMouseButtonDown(const TSharedPtr<FVariantTreeNode> & InVariantValueNode, const TraceServices::FFrame & InFrame, const FPointerEvent& InKeyEvent);
+	
 	TSharedPtr<SVariantValueView> View;
 
 	uint64 ObjectId;
 	double TimeMarker;
-
+	uint32 SelectedPropertyId;
+	
 	TAttribute<double> CurrentTime;
 	const TraceServices::IAnalysisSession* AnalysisSession;
 };
\ No newline at end of file
diff --git a/Engine/Plugins/Animation/GameplayInsights/Source/GameplayInsights/Private/SPropertyTrackView.cpp b/Engine/Plugins/Animation/GameplayInsights/Source/GameplayInsights/Private/SPropertyTrackView.cpp
new file mode 100644
--- /dev/null
+++ b/Engine/Plugins/Animation/GameplayInsights/Source/GameplayInsights/Private/SPropertyTrackView.cpp
@@ -0,0 +1,42 @@
+﻿// Copyright Epic Games, Inc. All Rights Reserved.
+
+#include "SPropertyTrackView.h"
+
+#include "GameplayProvider.h"
+#include "PropertyHelpers.h"
+#include "PropertyTrack.h"
+#include "VariantTreeNode.h"
+
+void SPropertyTrackView::GetVariantsAtFrame(const TraceServices::FFrame& InFrame,TArray<TSharedRef<FVariantTreeNode>>& OutVariants) const
+{
+	if (TSharedPtr<RewindDebugger::FPropertyTrack> PropertyTrack = Track.Pin())
+	{
+		TraceServices::FAnalysisSessionReadScope SessionReadScope(*AnalysisSession);
+		const TSharedRef<FVariantTreeNode> Header = OutVariants.Add_GetRef(FVariantTreeNode::MakeHeader(NSLOCTEXT("SPropertyTrackView", "Properties", "Properties"), INDEX_NONE));
+	
+		if (const FGameplayProvider* GameplayProvider = AnalysisSession->ReadProvider<FGameplayProvider>(FGameplayProvider::ProviderName))
+		{
+			PropertyTrack->ReadObjectPropertyValueAtFrame(InFrame, *GameplayProvider, [GameplayProvider, &PropertyTrack, &Header](const FObjectPropertyValue& InValue, uint32 InValueIndex, const FObjectPropertiesMessage& InMessage)
+			{
+				GameplayProvider->ReadObjectPropertiesStorage(PropertyTrack->GetObjectId(), InMessage, [GameplayProvider, InValueIndex, &PropertyTrack, &Header](const TConstArrayView<FObjectPropertyValue> & InStorage)
+				{
+					Header->AddChild(FObjectPropertyHelpers::GetVariantNodeFromProperty(InValueIndex, *GameplayProvider, InStorage));
+				});
+			});
+		}
+		else
+		{
+			Header->AddChild(FVariantTreeNode::MakeString(FText::FromName(PropertyTrack->GetPropertyName()), PropertyTrack->GetObjectProperty()->Property.Value));
+		}
+	}
+}
+
+FName SPropertyTrackView::GetName() const
+{
+	return Track.IsValid() ? Track.Pin()->GetName() : FName(TEXT("SPropertyTrackView"));
+}
+
+void SPropertyTrackView::SetTrack(const TWeakPtr<RewindDebugger::FPropertyTrack>& InTrack)
+{
+	Track = InTrack;
+}
diff --git a/Engine/Plugins/Animation/GameplayInsights/Source/GameplayInsights/Private/SPropertyTrackView.h b/Engine/Plugins/Animation/GameplayInsights/Source/GameplayInsights/Private/SPropertyTrackView.h
new file mode 100644
--- /dev/null
+++ b/Engine/Plugins/Animation/GameplayInsights/Source/GameplayInsights/Private/SPropertyTrackView.h
@@ -0,0 +1,34 @@
+﻿// Copyright Epic Games, Inc. All Rights Reserved.
+
+#pragma once
+
+#include "CoreMinimal.h"
+#include "PropertiesTrack.h"
+
+#include "SPropertiesDebugViewBase.h"
+
+namespace RewindDebugger
+{
+	class FPropertyTrack;
+}
+
+namespace TraceServices { class IAnalysisSession; }
+
+/**
+ * Used to display the properties of a single traced variable. 
+ */
+class SPropertyTrackView : public SPropertiesDebugViewBase
+{
+public:
+
+	/** Begin SPropertiesDebugViewBase interface */
+	virtual void GetVariantsAtFrame(const TraceServices::FFrame& InFrame, TArray<TSharedRef<FVariantTreeNode>>& OutVariants) const override;
+	virtual FName GetName() const override;
+	/** End SPropertiesDebugViewBase interface */
+
+	/** Update the track from which information will be queried and displayed */
+	void SetTrack(const TWeakPtr<RewindDebugger::FPropertyTrack> & InTrack);
+
+private:
+	TWeakPtr<RewindDebugger::FPropertyTrack> Track = nullptr;
+};
\ No newline at end of file
diff --git a/Engine/Plugins/Animation/GameplayInsights/Source/GameplayInsights/Private/SVariantValueView.cpp b/Engine/Plugins/Animation/GameplayInsights/Source/GameplayInsights/Private/SVariantValueView.cpp
--- a/Engine/Plugins/Animation/GameplayInsights/Source/GameplayInsights/Private/SVariantValueView.cpp
+++ b/Engine/Plugins/Animation/GameplayInsights/Source/GameplayInsights/Private/SVariantValueView.cpp
@@ -227,7 +227,7 @@ public:
 	SLATE_BEGIN_ARGS(SVariantValueNode) {}
 
 	SLATE_ATTRIBUTE(FText, HighlightText)
-
+	SLATE_EVENT(FOnMouseButtonDownOnVariantValueNode, OnMouseButtonDownOnVariantValueNode)
 	SLATE_END_ARGS()
 
 	void Construct(const FArguments& InArgs, const TSharedRef<STableViewBase>& InOwnerTable, TSharedRef<FVariantTreeNode> InNode, const TraceServices::IAnalysisSession& InAnalysisSession)
@@ -235,7 +235,8 @@ public:
 		Node = InNode;
 		AnalysisSession = &InAnalysisSession;
 		HighlightText = InArgs._HighlightText;
-
+		OnMouseButtonDownOnVariantValueNode = InArgs._OnMouseButtonDownOnVariantValueNode;
+		
 		SMultiColumnTableRow<TSharedRef<FVariantTreeNode>>::Construct(
 			FSuperRowType::FArguments()
 			.Padding(1.0f),
@@ -243,6 +244,14 @@ public:
 		);
 	}
 
+	virtual FReply OnMouseButtonDown(const FGeometry& MyGeometry, const FPointerEvent& MouseEvent) override
+	{
+		OnMouseButtonDownOnVariantValueNode.ExecuteIfBound(Node, MouseEvent);
+
+		/** We return a FReply::Unhandled to allow for event to bubble up in order for context menu to be built. */
+		return FReply::Unhandled();
+	}
+	
 	// SMultiColumnTableRow interface
 	virtual TSharedRef<SWidget> GenerateWidgetForColumn(const FName& InColumnName) override
 	{
@@ -296,12 +305,14 @@ public:
 	const TraceServices::IAnalysisSession* AnalysisSession;
 	TSharedPtr<FVariantTreeNode> Node;
 	TAttribute<FText> HighlightText;
+	FOnMouseButtonDownOnVariantValueNode OnMouseButtonDownOnVariantValueNode;
 };
 
 void SVariantValueView::Construct(const FArguments& InArgs, const TraceServices::IAnalysisSession& InAnalysisSession)
 {
 	OnGetVariantValues = InArgs._OnGetVariantValues;
-
+	OnMouseButtonDownOnVariantValue = InArgs._OnMouseButtonDownOnVariantValue;
+	
 	AnalysisSession = &InAnalysisSession;
 	bNeedsRefresh = false;
 	bRecordExpansion = true;
@@ -313,6 +324,7 @@ void SVariantValueView::Construct(const FArguments& InArgs, const TraceServices:
 		.OnGenerateRow(this, &SVariantValueView::HandleGeneratePropertyRow)
 		.OnGetChildren(this, &SVariantValueView::HandleGetPropertyChildren)
 		.OnExpansionChanged(this, &SVariantValueView::HandleExpansionChanged)
+		.OnContextMenuOpening(InArgs._OnContextMenuOpening)
 		.TreeItemsSource(&FilteredNodes)
 		.HeaderRow(
 			SNew(SHeaderRow)
@@ -322,7 +334,6 @@ void SVariantValueView::Construct(const FArguments& InArgs, const TraceServices:
 			+SHeaderRow::Column(VariantColumns::Value)
 			.DefaultLabel(LOCTEXT("ValueValueColumn", "Value"))
 		);
-
 	ChildSlot
 	[
 		SNew(SVerticalBox)
@@ -378,7 +389,11 @@ TSharedRef<ITableRow> SVariantValueView::HandleGeneratePropertyRow(TSharedRef<FV
 {
 	return 
 		SNew(SVariantValueNode, OwnerTable, Item, *AnalysisSession)
-		.HighlightText(MakeAttributeLambda([this](){ return FilterText; }));
+		.HighlightText(MakeAttributeLambda([this](){ return FilterText; }))
+		.OnMouseButtonDownOnVariantValueNode_Lambda([this](const TSharedPtr<FVariantTreeNode> & InVariantValueNode, const FPointerEvent & InPointerEvent)
+		{
+			OnMouseButtonDownOnVariantValue.ExecuteIfBound(InVariantValueNode, Frame, InPointerEvent);
+		});
 }
 
 void SVariantValueView::HandleGetPropertyChildren(TSharedRef<FVariantTreeNode> InItem, TArray<TSharedRef<FVariantTreeNode>>& OutChildren)
diff --git a/Engine/Plugins/Animation/GameplayInsights/Source/GameplayInsights/Private/SVariantValueView.h b/Engine/Plugins/Animation/GameplayInsights/Source/GameplayInsights/Private/SVariantValueView.h
--- a/Engine/Plugins/Animation/GameplayInsights/Source/GameplayInsights/Private/SVariantValueView.h
+++ b/Engine/Plugins/Animation/GameplayInsights/Source/GameplayInsights/Private/SVariantValueView.h
@@ -16,6 +16,8 @@ class FUICommandList;
 
 // Delegate called to get variant values to display
 DECLARE_DELEGATE_TwoParams(FOnGetVariantValues, const TraceServices::FFrame& /*InFrame*/, TArray<TSharedRef<FVariantTreeNode>>& /*OutValues*/);
+DECLARE_DELEGATE_ThreeParams(FOnMouseButtonDownOnVariantValue, const TSharedPtr<FVariantTreeNode> & /*InVariantValueNode*/, const TraceServices::FFrame & /*InFrame*/, const FPointerEvent& /*InKeyEvent*/);
+DECLARE_DELEGATE_TwoParams(FOnMouseButtonDownOnVariantValueNode, const TSharedPtr<FVariantTreeNode> & /*InVariantValueNode*/, const FPointerEvent& /*InKeyEvent*/);
 
 class SVariantValueView : public SCompoundWidget
 {
@@ -23,7 +25,8 @@ public:
 	SLATE_BEGIN_ARGS(SVariantValueView) {}
 
 	SLATE_EVENT(FOnGetVariantValues, OnGetVariantValues)
-
+	SLATE_EVENT(FOnContextMenuOpening, OnContextMenuOpening)
+	SLATE_EVENT(FOnMouseButtonDownOnVariantValue, OnMouseButtonDownOnVariantValue)
 	SLATE_END_ARGS()
 
 	void Construct(const FArguments& InArgs, const TraceServices::IAnalysisSession& InAnalysisSession);
@@ -31,6 +34,12 @@ public:
 	/** Refresh the displayed variants. */
 	void RequestRefresh(const TraceServices::FFrame& InFrame) { Frame = InFrame; bNeedsRefresh = true; }
 
+	/** @return Trace frame that the view represents */
+	const TraceServices::FFrame & GetFrame() const { return Frame; }
+
+	/** @return Trace session that the view references */
+	const TraceServices::IAnalysisSession&  GetAnalysisSession() const { return *AnalysisSession; }
+	
 private:
 	// SWidget interface
 	virtual FReply OnKeyDown(const FGeometry& MyGeometry, const FKeyEvent& InKeyEvent) override;
@@ -85,6 +94,8 @@ private:
 
 	FOnGetVariantValues OnGetVariantValues;
 
+	FOnMouseButtonDownOnVariantValue OnMouseButtonDownOnVariantValue;
+	
 	bool bNeedsRefresh;
 
 	bool bRecordExpansion;
diff --git a/Engine/Plugins/Animation/GameplayInsights/Source/GameplayInsights/Private/VariantTreeNode.h b/Engine/Plugins/Animation/GameplayInsights/Source/GameplayInsights/Private/VariantTreeNode.h
--- a/Engine/Plugins/Animation/GameplayInsights/Source/GameplayInsights/Private/VariantTreeNode.h
+++ b/Engine/Plugins/Animation/GameplayInsights/Source/GameplayInsights/Private/VariantTreeNode.h
@@ -4,6 +4,7 @@
 
 #include "CoreMinimal.h"
 #include "IAnimationProvider.h"
+#include "IGameplayProvider.h"
 
 namespace TraceServices { class IAnalysisSession; }
 
@@ -20,9 +21,17 @@ struct FVariantTreeNode : TSharedFromThis<FVariantTreeNode>
 	FVariantTreeNode(const FText& InName, const FVariantValue& InValue, uint64 InId = INDEX_NONE)
 		: Name(InName)
 		, Value(InValue)
+		, PropertyNameId(0)
 		, Id(InId)
 	{}
 
+	FVariantTreeNode(const FText& InName, const FVariantValue& InValue, uint32 InPropertyNamedId = 0, uint64 InId = INDEX_NONE)
+		: Name(InName)
+		, Value(InValue)
+		, PropertyNameId(InPropertyNamedId)
+		, Id(InId)
+	{}
+	
 	const TSharedRef<FVariantTreeNode>& AddChild(const TSharedRef<FVariantTreeNode>& InChild)
 	{
 		check(!InChild->Parent.IsValid());
@@ -42,6 +51,10 @@ struct FVariantTreeNode : TSharedFromThis<FVariantTreeNode>
 
 	uint64 GetId() const { return Id; }
 
+	uint32 GetPropertyNamedId() const { return PropertyNameId; }
+	
+	void SetPropertyNameId(uint32 InPropertyNameId) { PropertyNameId = InPropertyNameId; }
+	
 	void SetFilterState(EVariantTreeNodeFilterState InFilterState) { FilterState = InFilterState; }
 
 	EVariantTreeNodeFilterState GetFilterState() const { return FilterState; }
@@ -135,7 +148,8 @@ private:
 	FVariantValue Value;
 	TWeakPtr<FVariantTreeNode> Parent;
 	TArray<TSharedRef<FVariantTreeNode>> Children;
-
+	uint32 PropertyNameId;
+	
 	// Unique ID used to record expansion state
 	uint64 Id;
 
diff --git a/Engine/Plugins/Animation/GameplayInsights/Source/GameplayInsights/Public/IGameplayProvider.h b/Engine/Plugins/Animation/GameplayInsights/Source/GameplayInsights/Public/IGameplayProvider.h
--- a/Engine/Plugins/Animation/GameplayInsights/Source/GameplayInsights/Public/IGameplayProvider.h
+++ b/Engine/Plugins/Animation/GameplayInsights/Source/GameplayInsights/Public/IGameplayProvider.h
@@ -25,8 +25,10 @@ struct FObjectInfo
 
 struct FObjectPropertiesMessage
 {
-	int64 PropertyValueStartIndex = INDEX_NONE;
-	int64 PropertyValueEndIndex = INDEX_NONE;
+	int64 PropertyValueStartIndex = INDEX_NONE;	// Inclusive
+	int64 PropertyValueEndIndex = INDEX_NONE;	// Exclusive
+	double ProfileTime;
+	double ElapsedTime;
 };
 
 struct FObjectPropertyValue
@@ -34,7 +36,8 @@ struct FObjectPropertyValue
 	const TCHAR* Value = nullptr;
 	int32 ParentId = 0;
 	uint32 TypeStringId = 0;
-	uint32 KeyStringId = 0;
+	uint32 NameId = 0;
+	uint32 ParentNameId = 0;
 	float ValueAsFloat = 0.0f;
 };
 
@@ -109,10 +112,12 @@ public:
 	virtual bool ReadObjectEventsTimeline(uint64 InObjectId, TFunctionRef<void(const ObjectEventsTimeline&)> Callback) const = 0;
 	virtual bool ReadObjectEvent(uint64 InObjectId, uint64 InMessageId, TFunctionRef<void(const FObjectEventMessage&)> Callback) const = 0;
 	virtual bool ReadObjectPropertiesTimeline(uint64 InObjectId, TFunctionRef<void(const ObjectPropertiesTimeline&)> Callback) const = 0;
+	virtual bool ReadObjectPropertiesStorage(uint64 InObjectId, const FObjectPropertiesMessage& InMessage, TFunctionRef<void(const TConstArrayView<FObjectPropertyValue> &)> Callback) const = 0;
 	virtual void EnumerateObjectPropertyValues(uint64 InObjectId, const FObjectPropertiesMessage& InMessage, TFunctionRef<void(const FObjectPropertyValue&)> Callback) const = 0;
 	virtual void EnumerateObjects(TFunctionRef<void(const FObjectInfo&)> Callback) const = 0;
 	virtual void EnumerateObjects(double StartTime, double EndTime, TFunctionRef<void(const FObjectInfo&)> Callback) const = 0;
 	virtual void EnumerateSubobjects(uint64 ObjectId, TFunctionRef<void(uint64 SubobjectId)> Callback) const = 0;
+	virtual const FObjectPropertyValue * FindPropertyValueFromStorageIndex(uint64 InObjectId, int64 InStorageIndex) const = 0;
 	virtual const FClassInfo* FindClassInfo(uint64 InClassId) const = 0;
 	virtual const UClass* FindClass(uint64 InClassId) const = 0;
 	virtual const FClassInfo* FindClassInfo(const TCHAR* InClassPath) const = 0;
diff --git a/Engine/Plugins/Animation/GameplayInsights/Source/RewindDebugger/Private/RewindDebugger.cpp b/Engine/Plugins/Animation/GameplayInsights/Source/RewindDebugger/Private/RewindDebugger.cpp
--- a/Engine/Plugins/Animation/GameplayInsights/Source/RewindDebugger/Private/RewindDebugger.cpp
+++ b/Engine/Plugins/Animation/GameplayInsights/Source/RewindDebugger/Private/RewindDebugger.cpp
@@ -702,7 +702,7 @@ void FRewindDebugger::Tick(float DeltaTime)
 						
 						const TraceServices::IFrameProvider& FrameProvider = TraceServices::ReadFrameProvider(*Session);
 						TraceServices::FFrame Frame;
-						if(FrameProvider.GetFrameFromTime(ETraceFrameType::TraceFrameType_Game, CurrentTraceTime, Frame))
+						if (FrameProvider.GetFrameFromTime(ETraceFrameType::TraceFrameType_Game, CurrentTraceTime, Frame))
 						{
 							{
 								TRACE_CPUPROFILER_EVENT_SCOPE(FRewindDebugger::Tick_UpdateActorPosition);
@@ -1037,6 +1037,23 @@ void FRewindDebugger::UpdateDetailsPanel(TSharedRef<SDockTab> DetailsTab)
 	}
 }
 
+void FRewindDebugger::RegisterComponentContextMenu()
+{
+	UToolMenu* Menu = UToolMenus::Get()->FindMenu("RewindDebugger.ComponentContextMenu");
+	
+	FToolMenuSection& Section = Menu->FindOrAddSection("SelectedTrack");
+	
+	FToolMenuEntry& Entry = Section.AddDynamicEntry(NAME_None, FNewToolMenuSectionDelegate::CreateLambda([](FToolMenuSection& InSection)
+	{
+		const UComponentContextMenuContext* Context = InSection.FindContext<UComponentContextMenuContext>();
+
+		if (Context && Context->SelectedTrack.IsValid())
+		{
+			Context->SelectedTrack->BuildContextMenu(InSection);
+		}
+	}));
+}
+
 void FRewindDebugger::ComponentDoubleClicked(TSharedPtr<RewindDebugger::FRewindDebuggerTrack> SelectedObject)
 {
 	if (!SelectedObject.IsValid())
@@ -1048,11 +1065,12 @@ void FRewindDebugger::ComponentDoubleClicked(TSharedPtr<RewindDebugger::FRewindD
 	SelectedTrack->HandleDoubleClick();
 }
 
-TSharedPtr<SWidget> FRewindDebugger::BuildComponentContextMenu()
+TSharedPtr<SWidget> FRewindDebugger::BuildComponentContextMenu() const
 {
 	UComponentContextMenuContext* MenuContext = NewObject<UComponentContextMenuContext>();
 	MenuContext->SelectedObject = GetSelectedComponent();
-
+	MenuContext->SelectedTrack = SelectedTrack;
+	
 	if (SelectedTrack.IsValid())
 	{
 		// build a list of class hierarchy names to make it easier for extensions to enable menu entries by type
@@ -1095,6 +1113,11 @@ TSharedPtr<FDebugObjectInfo> FRewindDebugger::GetSelectedComponent() const
 	}
 }
 
+TSharedPtr<RewindDebugger::FRewindDebuggerTrack> FRewindDebugger::GetSelectedTrack() const
+{
+	return SelectedTrack;
+}
+
 // build a component tree that's compatible with the public api from 5.0 for GetDebugComponents.
 void FRewindDebugger::RefreshDebugComponents(TArray<TSharedPtr<RewindDebugger::FRewindDebuggerTrack>>& InTracks, TArray<TSharedPtr<FDebugObjectInfo>>& OutComponents)
 {
diff --git a/Engine/Plugins/Animation/GameplayInsights/Source/RewindDebugger/Private/RewindDebugger.h b/Engine/Plugins/Animation/GameplayInsights/Source/RewindDebugger/Private/RewindDebugger.h
--- a/Engine/Plugins/Animation/GameplayInsights/Source/RewindDebugger/Private/RewindDebugger.h
+++ b/Engine/Plugins/Animation/GameplayInsights/Source/RewindDebugger/Private/RewindDebugger.h
@@ -43,6 +43,7 @@ public:
 	virtual bool IsPIESimulating() const override { return bPIESimulating; }
 	virtual double GetRecordingDuration() const override { return RecordingDuration.Get(); }
 	virtual TSharedPtr<FDebugObjectInfo> GetSelectedComponent() const override;
+	virtual TSharedPtr<RewindDebugger::FRewindDebuggerTrack> GetSelectedTrack() const override;
 	virtual TArray<TSharedPtr<FDebugObjectInfo>>& GetDebugComponents() override;
 	
 	TArray<TSharedPtr<RewindDebugger::FRewindDebuggerTrack>>& GetDebugTracks() { return DebugTracks; }
@@ -103,10 +104,11 @@ public:
 	
 	void ComponentDoubleClicked(TSharedPtr<RewindDebugger::FRewindDebuggerTrack> SelectedObject);
 	void ComponentSelectionChanged(TSharedPtr<RewindDebugger::FRewindDebuggerTrack> SelectedObject);
-	TSharedPtr<SWidget> BuildComponentContextMenu();
-
+	TSharedPtr<SWidget> BuildComponentContextMenu() const;
+	
 	void UpdateDetailsPanel(TSharedRef<SDockTab> DetailsTab);
-
+	static void RegisterComponentContextMenu();
+	
 	DECLARE_DELEGATE_OneParam( FOnTrackCursor, bool)
 	void OnTrackCursor(const FOnTrackCursor& TrackCursorCallback);
 
diff --git a/Engine/Plugins/Animation/GameplayInsights/Source/RewindDebugger/Private/RewindDebuggerModule.cpp b/Engine/Plugins/Animation/GameplayInsights/Source/RewindDebugger/Private/RewindDebuggerModule.cpp
--- a/Engine/Plugins/Animation/GameplayInsights/Source/RewindDebugger/Private/RewindDebuggerModule.cpp
+++ b/Engine/Plugins/Animation/GameplayInsights/Source/RewindDebugger/Private/RewindDebuggerModule.cpp
@@ -227,6 +227,7 @@ void FRewindDebuggerModule::StartupModule()
 
 	FPropertyTraceMenu::Register();
 	FAnimInstanceMenu::Register();
+	FRewindDebugger::RegisterComponentContextMenu();
 }
 
 void FRewindDebuggerModule::ShutdownModule()
diff --git a/Engine/Source/Editor/RewindDebuggerInterface/Private/RewindDebuggerInterface.cpp b/Engine/Source/Editor/RewindDebuggerInterface/Private/RewindDebuggerInterface.cpp
--- a/Engine/Source/Editor/RewindDebuggerInterface/Private/RewindDebuggerInterface.cpp
+++ b/Engine/Source/Editor/RewindDebuggerInterface/Private/RewindDebuggerInterface.cpp
@@ -7,6 +7,8 @@
 #include "IRewindDebuggerViewCreator.h"
 #include "UObject/NameTypes.h"
 
+DEFINE_LOG_CATEGORY(LogRewindDebugger)
+
 const FName IRewindDebuggerExtension::ModularFeatureName = "RewindDebuggerExtension";
 const FName IRewindDebuggerViewCreator::ModularFeatureName = "RewindDebuggerViewCreator";
 const FName IRewindDebuggerDoubleClickHandler::ModularFeatureName = "RewindDebuggerDoubleClickHandler";
diff --git a/Engine/Source/Editor/RewindDebuggerInterface/Public/IRewindDebugger.h b/Engine/Source/Editor/RewindDebuggerInterface/Public/IRewindDebugger.h
--- a/Engine/Source/Editor/RewindDebuggerInterface/Public/IRewindDebugger.h
+++ b/Engine/Source/Editor/RewindDebuggerInterface/Public/IRewindDebugger.h
@@ -2,12 +2,15 @@
 #pragma once
 
 #include "CoreMinimal.h"
+#include "TraceServices/Model/Frames.h"
 #include "IRewindDebugger.generated.h"
 
 // IRewindDebugger
 //
 // Public interface to rewind debugger
 
+REWINDDEBUGGERINTERFACE_API DECLARE_LOG_CATEGORY_EXTERN(LogRewindDebugger, Log, All);
+
 namespace TraceServices
 {
 	class IAnalysisSession;
@@ -26,6 +29,12 @@ struct FDebugObjectInfo
 	TArray<TSharedPtr<FDebugObjectInfo>> Children;
 };
 
+namespace RewindDebugger
+{
+	class FRewindDebuggerTrack;
+	class FPropertiesTrack;
+}
+
 UCLASS()
 class REWINDDEBUGGERINTERFACE_API UComponentContextMenuContext : public UObject
 {
@@ -33,8 +42,11 @@ class REWINDDEBUGGERINTERFACE_API UComponentContextMenuContext : public UObject
 public:
 	TSharedPtr<FDebugObjectInfo> SelectedObject;
 	TArray<FName> TypeHierarchy;
+	TSharedPtr<RewindDebugger::FRewindDebuggerTrack> SelectedTrack;
 };
 
+class SPropertiesDebugViewBase;
+
 class REWINDDEBUGGERINTERFACE_API IRewindDebugger
 {
 public:
@@ -63,6 +75,9 @@ public:
 	// returns the currently selected debug component
 	virtual TSharedPtr<FDebugObjectInfo> GetSelectedComponent() const = 0;
 
+	// returns the currently selected track
+	virtual TSharedPtr<RewindDebugger::FRewindDebuggerTrack> GetSelectedTrack() const = 0;
+	
 	// get posiotion of the selected target actor (returns true if position is valid)
 	virtual bool GetTargetActorPosition(FVector& OutPosition) const = 0;
 
diff --git a/Engine/Source/Editor/RewindDebuggerInterface/Public/RewindDebuggerTrack.h b/Engine/Source/Editor/RewindDebuggerInterface/Public/RewindDebuggerTrack.h
--- a/Engine/Source/Editor/RewindDebuggerInterface/Public/RewindDebuggerTrack.h
+++ b/Engine/Source/Editor/RewindDebuggerInterface/Public/RewindDebuggerTrack.h
@@ -10,6 +10,8 @@ namespace TraceServices
 	class IAnalysisSession;
 }
 
+struct FToolMenuSection;
+
 namespace RewindDebugger
 {
 	
@@ -97,6 +99,9 @@ public:
 	bool IsVisible() const { return bVisible; }
 	void SetIsVisible(bool bInIsVisible) { bVisible = bInIsVisible; }
 
+	// Called to generate context menu for the current selected track
+	virtual void BuildContextMenu(FToolMenuSection& InMenuSection) {};
+
 private:
 
 	virtual bool UpdateInternal() { return false; }
diff --git a/Engine/Source/Editor/RewindDebuggerInterface/RewindDebuggerInterface.Build.cs b/Engine/Source/Editor/RewindDebuggerInterface/RewindDebuggerInterface.Build.cs
--- a/Engine/Source/Editor/RewindDebuggerInterface/RewindDebuggerInterface.Build.cs
+++ b/Engine/Source/Editor/RewindDebuggerInterface/RewindDebuggerInterface.Build.cs
@@ -11,7 +11,8 @@ public class RewindDebuggerInterface : ModuleRules
 				"Core",
 				"CoreUObject",
 				"SlateCore",
-				"ToolWidgets"
+				"ToolWidgets",
+				"TraceServices"
 			}
 		);
 	}
