Fixed issue regarding duplicate keys for two LOCTEXT calls in PropertyTrack.cpp.

#rnx

diff --git a/Engine/Plugins/Animation/GameplayInsights/Source/GameplayInsights/Private/PropertyTrack.cpp b/Engine/Plugins/Animation/GameplayInsights/Source/GameplayInsights/Private/PropertyTrack.cpp
--- a/Engine/Plugins/Animation/GameplayInsights/Source/GameplayInsights/Private/PropertyTrack.cpp
+++ b/Engine/Plugins/Animation/GameplayInsights/Source/GameplayInsights/Private/PropertyTrack.cpp
@@ -40,7 +40,7 @@ namespace RewindDebugger
 
 		InMenuSection.AddMenuEntry(NAME_None,
 			LOCTEXT("RemovePropertyTrack", "Remove Property Track"),
-			LOCTEXT("RemovePropertyTrack", "Remove traced property track from timeline"),
+			LOCTEXT("RemovePropertyTrackTooltip", "Remove traced property track from timeline"),
 			{},
 			FUIAction(FExecuteAction::CreateLambda([this]()
 			{
