Fix small memory leak due FPropertyWatchManager not clearing properties at the start of a PIE session.

#rnx

diff --git a/Engine/Plugins/Animation/GameplayInsights/Source/GameplayInsights/Private/PropertyWatchManager.cpp b/Engine/Plugins/Animation/GameplayInsights/Source/GameplayInsights/Private/PropertyWatchManager.cpp
--- a/Engine/Plugins/Animation/GameplayInsights/Source/GameplayInsights/Private/PropertyWatchManager.cpp
+++ b/Engine/Plugins/Animation/GameplayInsights/Source/GameplayInsights/Private/PropertyWatchManager.cpp
@@ -3,6 +3,10 @@
 #include "PropertyWatchManager.h"
 #include "IRewindDebugger.h"
 
+#if WITH_EDITOR
+#include "Editor.h"
+#endif
+
 FPropertyWatchManager* FPropertyWatchManager::InternalInstance = nullptr;
 
 FPropertyWatchManager::FPropertyWatchManager()
@@ -97,16 +101,27 @@ FPropertyWatchManager::FOnPropertyWatched& FPropertyWatchManager::OnPropertyWatc
 FPropertyWatchManager::FOnPropertyUnwatched& FPropertyWatchManager::OnPropertyUnwatched()
 {
 	return OnPropertyUnwatchedDelegate;
-}
+}		
 
 void FPropertyWatchManager::Initialize()
 {
 	InternalInstance = new FPropertyWatchManager;
+
+#if WITH_EDITOR
+	FEditorDelegates::BeginPIE.AddLambda([](const bool bIsSimulating)
+	{
+		if (InternalInstance)
+		{
+			InternalInstance->ClearAllWatchedProperties();
+		}
+	});
+#endif
 }
 
 void FPropertyWatchManager::Shutdown()
 {
 	delete InternalInstance;
+	InternalInstance = nullptr;
 }
 
 FPropertyWatchManager* FPropertyWatchManager::Instance()
diff --git a/Engine/Plugins/Animation/GameplayInsights/Source/GameplayInsights/Private/PropertyWatchManager.h b/Engine/Plugins/Animation/GameplayInsights/Source/GameplayInsights/Private/PropertyWatchManager.h
--- a/Engine/Plugins/Animation/GameplayInsights/Source/GameplayInsights/Private/PropertyWatchManager.h
+++ b/Engine/Plugins/Animation/GameplayInsights/Source/GameplayInsights/Private/PropertyWatchManager.h
@@ -14,8 +14,6 @@ namespace RewindDebugger
 class FPropertyWatchManager
 {
 public:
-
-	// TODO: We need to add code that clears the property watch manager right after a PIE session has started.
 	
 	DECLARE_MULTICAST_DELEGATE_TwoParams(FOnPropertyWatched, uint64 /*InObjectId*/, uint32 /*InPropertyNameId*/)
 	DECLARE_MULTICAST_DELEGATE_TwoParams(FOnPropertyUnwatched, uint64 /*InObjectId*/, uint32 /*InPropertyNameId*/)
