Add double click handling to montages track in Rewind Debugger

diff --git a/Engine/Plugins/Animation/GameplayInsights/Source/GameplayInsights/Private/MontagesTrack.cpp b/Engine/Plugins/Animation/GameplayInsights/Source/GameplayInsights/Private/MontagesTrack.cpp
--- a/Engine/Plugins/Animation/GameplayInsights/Source/GameplayInsights/Private/MontagesTrack.cpp
+++ b/Engine/Plugins/Animation/GameplayInsights/Source/GameplayInsights/Private/MontagesTrack.cpp
@@ -9,6 +9,8 @@
 #include "SMontageView.h"
 #if WITH_EDITOR
 #include "Animation/AnimMontage.h"
+#include "Editor.h"
+#include "Subsystems/AssetEditorSubsystem.h"
 #endif
 
 #define LOCTEXT_NAMESPACE "MontagesTrack"
@@ -284,7 +286,26 @@ TSharedPtr<SWidget> FMontageTrack::GetDetailsViewInternal()
 	MontageView->SetAssetFilter(AssetId);
 	return MontageView;
 }
+	
+bool FMontageTrack::HandleDoubleClickInternal()
+{
+#if WITH_EDITOR
+	IRewindDebugger* RewindDebugger = IRewindDebugger::Instance();
+	if (const TraceServices::IAnalysisSession* AnalysisSession = RewindDebugger->GetAnalysisSession())
+	{
+		TraceServices::FAnalysisSessionReadScope SessionReadScope(*AnalysisSession);
 
+		const FGameplayProvider* GameplayProvider = AnalysisSession->ReadProvider<FGameplayProvider>(FGameplayProvider::ProviderName);
+
+		const FObjectInfo& AssetInfo = GameplayProvider->GetObjectInfo(GetAssetId());
+
+		GEditor->GetEditorSubsystem<UAssetEditorSubsystem>()->OpenEditorForAsset(AssetInfo.PathName);
+
+		return true;
+	}
+#endif
+	return false;
+}
 
 FName FMontagesTrackCreator::GetTargetTypeNameInternal() const
 {
diff --git a/Engine/Plugins/Animation/GameplayInsights/Source/GameplayInsights/Private/MontagesTrack.h b/Engine/Plugins/Animation/GameplayInsights/Source/GameplayInsights/Private/MontagesTrack.h
--- a/Engine/Plugins/Animation/GameplayInsights/Source/GameplayInsights/Private/MontagesTrack.h
+++ b/Engine/Plugins/Animation/GameplayInsights/Source/GameplayInsights/Private/MontagesTrack.h
@@ -30,6 +30,7 @@ private:
 	virtual TSharedPtr<SWidget> GetTimelineViewInternal() override;
 	virtual TSharedPtr<SWidget> GetDetailsViewInternal() override;
 	virtual void IterateSubTracksInternal(TFunction<void(TSharedPtr<FRewindDebuggerTrack> SubTrack)> IteratorFunction) override;
+    virtual bool HandleDoubleClickInternal() override;
 	
 	virtual FSlateIcon GetIconInternal() override { return Icon; }
 	virtual FName GetNameInternal() const override { return "Montage"; }
