Property display names support for disabling or enabling friendly names in Editor Preferences.

#rnx

diff --git a/Engine/Plugins/Animation/GameplayInsights/Source/GameplayInsights/Private/PropertyHelpers.cpp b/Engine/Plugins/Animation/GameplayInsights/Source/GameplayInsights/Private/PropertyHelpers.cpp
--- a/Engine/Plugins/Animation/GameplayInsights/Source/GameplayInsights/Private/PropertyHelpers.cpp
+++ b/Engine/Plugins/Animation/GameplayInsights/Source/GameplayInsights/Private/PropertyHelpers.cpp
@@ -5,6 +5,10 @@
 #include "GameplayProvider.h"
 #include "VariantTreeNode.h"
 
+#if WITH_EDITOR
+#include "Settings/EditorStyleSettings.h"
+#endif
+
 FLinearColor FObjectPropertyHelpers::GetPropertyColor(const FObjectPropertyValue & InProperty)
 {
 	const IRewindDebugger* RewindDebugger = IRewindDebugger::Instance();
@@ -140,8 +144,16 @@ FName FObjectPropertyHelpers::GetPropertyDisplayName(const FObjectPropertyValue
 		}
 	}
 
+	
 	// Ensure human readable names.
-	const FName OutputName = *FName::NameToDisplayString(StringName, IsBoolProperty(InProperty.Value, InGameplayProvider.GetPropertyName(InProperty.TypeStringId)));
+	const bool bIsBool = IsBoolProperty(InProperty.Value, InGameplayProvider.GetPropertyName(InProperty.TypeStringId));
+
+#if WITH_EDITOR
+	const UEditorStyleSettings* Settings = GetDefault<UEditorStyleSettings>();
+	const FName OutputName = Settings->bShowFriendlyNames ? *FName::NameToDisplayString(StringName, bIsBool) : *StringName;
+#else
+	const FName OutputName = *FName::NameToDisplayString(StringName, bIsBool);
+#endif
 	
 	return OutputName;
 }
