Navigate to the same frame in a referenced animation sequence, blendspace or montage, when you click on the hyperlink from the Rewind Debugger
Also set the BlendSpace X/Y Parameters in the editor preview after opening a blend space.

diff --git a/Engine/Plugins/Animation/GameplayInsights/Source/GameplayInsights/GameplayInsights.Build.cs b/Engine/Plugins/Animation/GameplayInsights/Source/GameplayInsights/GameplayInsights.Build.cs
--- a/Engine/Plugins/Animation/GameplayInsights/Source/GameplayInsights/GameplayInsights.Build.cs
+++ b/Engine/Plugins/Animation/GameplayInsights/Source/GameplayInsights/GameplayInsights.Build.cs
@@ -48,6 +48,8 @@ namespace UnrealBuildTool.Rules
 					"ToolMenus",
 					"Kismet",
 					"SubobjectEditor",
+					"AnimGraph",
+					"AnimGraphRuntime"
 				});
 			}
 		}
diff --git a/Engine/Plugins/Animation/GameplayInsights/Source/GameplayInsights/Private/SBlendWeightsView.cpp b/Engine/Plugins/Animation/GameplayInsights/Source/GameplayInsights/Private/SBlendWeightsView.cpp
--- a/Engine/Plugins/Animation/GameplayInsights/Source/GameplayInsights/Private/SBlendWeightsView.cpp
+++ b/Engine/Plugins/Animation/GameplayInsights/Source/GameplayInsights/Private/SBlendWeightsView.cpp
@@ -33,7 +33,8 @@ void SBlendWeightsView::GetVariantsAtFrame(const TraceServices::FFrame& InFrame,
 					}
 					
 					const FClassInfo& ClassInfo = GameplayProvider->GetClassInfoFromObject(InMessage.AssetId);
-					TSharedRef<FVariantTreeNode> Header = OutVariants.Add_GetRef(FVariantTreeNode::MakeObject(FText::FromString(ClassInfo.Name), InMessage.AssetId, InMessage.AssetId));
+					TSharedRef<FVariantTreeNode> Header = OutVariants.Add_GetRef(FVariantTreeNode::MakeObject(FText::FromString(ClassInfo.Name), InMessage.AssetId, InMessage.AssetId,
+																								InMessage.PlaybackTime, InMessage.BlendSpaceFilteredPositionX, InMessage.BlendSpaceFilteredPositionY));
 					
 					Header->AddChild(FVariantTreeNode::MakeFloat(LOCTEXT("BlendWeight", "Blend Weight"), InMessage.BlendWeight));
 					Header->AddChild(FVariantTreeNode::MakeFloat(LOCTEXT("PlaybackTime", "Playback Time"), InMessage.PlaybackTime));
diff --git a/Engine/Plugins/Animation/GameplayInsights/Source/GameplayInsights/Private/SMontageView.cpp b/Engine/Plugins/Animation/GameplayInsights/Source/GameplayInsights/Private/SMontageView.cpp
--- a/Engine/Plugins/Animation/GameplayInsights/Source/GameplayInsights/Private/SMontageView.cpp
+++ b/Engine/Plugins/Animation/GameplayInsights/Source/GameplayInsights/Private/SMontageView.cpp
@@ -34,7 +34,7 @@ void SMontageView::GetVariantsAtFrame(const TraceServices::FFrame& InFrame, TArr
 					{
 						const FObjectInfo& MontageInfo = GameplayProvider->GetObjectInfo(InMessage.MontageId);
 					
-						TSharedRef<FVariantTreeNode> MontageHeader = OutVariants.Add_GetRef(FVariantTreeNode::MakeObject(FText::FromString(MontageInfo.Name), InMessage.MontageId, InMessage.MontageId));
+						TSharedRef<FVariantTreeNode> MontageHeader = OutVariants.Add_GetRef(FVariantTreeNode::MakeObject(FText::FromString(MontageInfo.Name), InMessage.MontageId, InMessage.MontageId, InMessage.Position));
 
 						MontageHeader->AddChild(FVariantTreeNode::MakeFloat(LOCTEXT("EventWeight", "Weight"), InMessage.Weight));
 						MontageHeader->AddChild(FVariantTreeNode::MakeFloat(LOCTEXT("EventDesiredWeight", "Desired Weight"), InMessage.DesiredWeight));
diff --git a/Engine/Plugins/Animation/GameplayInsights/Source/GameplayInsights/Private/SNotifiesView.cpp b/Engine/Plugins/Animation/GameplayInsights/Source/GameplayInsights/Private/SNotifiesView.cpp
--- a/Engine/Plugins/Animation/GameplayInsights/Source/GameplayInsights/Private/SNotifiesView.cpp
+++ b/Engine/Plugins/Animation/GameplayInsights/Source/GameplayInsights/Private/SNotifiesView.cpp
@@ -87,7 +87,7 @@ void SNotifiesView::GetVariantsAtFrame(const TraceServices::FFrame& InFrame, TAr
 
 			if(InMessage.NotifyEventType != EAnimNotifyMessageType::SyncMarker)
 			{
-				NotifyHeader->AddChild(FVariantTreeNode::MakeObject(LOCTEXT("Asset", "Asset"), InMessage.AssetId));
+				NotifyHeader->AddChild(FVariantTreeNode::MakeObject(LOCTEXT("Asset", "Asset"), InMessage.AssetId, InMessage.AssetId, InMessage.Time));
 
 				const FObjectInfo& NotifyInfo = GameplayProvider->GetObjectInfo(InMessage.NotifyId);
 				bool bIsState = InMessage.NotifyEventType == EAnimNotifyMessageType::Begin || InMessage.NotifyEventType == EAnimNotifyMessageType::End || InMessage.NotifyEventType == EAnimNotifyMessageType::Tick;
diff --git a/Engine/Plugins/Animation/GameplayInsights/Source/GameplayInsights/Private/SVariantValueView.cpp b/Engine/Plugins/Animation/GameplayInsights/Source/GameplayInsights/Private/SVariantValueView.cpp
--- a/Engine/Plugins/Animation/GameplayInsights/Source/GameplayInsights/Private/SVariantValueView.cpp
+++ b/Engine/Plugins/Animation/GameplayInsights/Source/GameplayInsights/Private/SVariantValueView.cpp
@@ -1,20 +1,17 @@
 // Copyright Epic Games, Inc. All Rights Reserved.
 
 #include "SVariantValueView.h"
-#include "AnimationProvider.h"
 #include "Widgets/Input/SSearchBox.h"
 #include "TraceServices/Model/AnalysisSession.h"
-#include "Insights/ITimingViewSession.h"
 #include "GameplayProvider.h"
 #include "TraceServices/Model/Frames.h"
 #include "Widgets/Layout/SBorder.h"
 #include "Widgets/SBoxPanel.h"
 #include "Styling/CoreStyle.h"
 #include "GameplayInsightsStyle.h"
-#include "Widgets/Layout/SScrollBorder.h"
-#include "Widgets/Input/SEditableTextBox.h"
 #include "Widgets/Input/SHyperlink.h"
 #include "Widgets/Input/SCheckBox.h"
+#include "Widgets/Layout/SScrollBorder.h"
 #include "VariantTreeNode.h"
 #include "Misc/TextFilterExpressionEvaluator.h"
 #include "Framework/Commands/GenericCommands.h"
@@ -26,6 +23,10 @@
 #include "Subsystems/AssetEditorSubsystem.h"
 #include "SourceCodeNavigation.h"
 #include "Editor.h"
+#include "IAnimationEditor.h"
+#include "IPersonaToolkit.h"
+#include "Animation/DebugSkelMeshComponent.h"
+#include "AnimGraph/Public/AnimPreviewInstance.h"
 #endif
 
 #define LOCTEXT_NAMESPACE "SVariantValueView"
@@ -160,9 +161,42 @@ static TSharedRef<SWidget> MakeVariantValueWidget(const TraceServices::IAnalysis
 				.Text(FText::FromString(ObjectInfo.Name))
 				.TextStyle(&FCoreStyle::Get().GetWidgetStyle<FTextBlockStyle>("SmallText"))
 				.ToolTipText(FText::Format(LOCTEXT("AssetHyperlinkTooltipFormat", "Open asset '{0}'"), FText::FromString(ObjectInfo.PathName)))
-				.OnNavigate_Lambda([ObjectInfo]()
+				.OnNavigate_Lambda([ObjectInfo, InValue]()
 				{
-					GEditor->GetEditorSubsystem<UAssetEditorSubsystem>()->OpenEditorForAsset(ObjectInfo.PathName);
+					UObject* Asset = nullptr;
+						
+					UPackage* Package = LoadPackage(NULL, ObjectInfo.PathName, LOAD_NoRedirects);
+					if (Package)
+					{
+						Package->FullyLoad();
+                
+						FString AssetName = FPaths::GetBaseFilename(ObjectInfo.PathName);
+						Asset = FindObject<UObject>(Package, *AssetName);
+					}
+					else
+					{
+						// fallback for unsaved assets
+						Asset = FindObject<UObject>(nullptr, ObjectInfo.PathName);
+					}
+                    	
+					if (Asset != nullptr)
+					{
+						if (UAssetEditorSubsystem* AssetEditorSS = GEditor->GetEditorSubsystem<UAssetEditorSubsystem>())
+						{
+							AssetEditorSS->OpenEditorForAsset(Asset);
+							if (IAssetEditorInstance* Editor = AssetEditorSS->FindEditorForAsset(Asset, true))
+							{
+								if (Editor->GetEditorName()=="AnimationEditor")
+								{
+									IAnimationEditor* AnimationEditor = static_cast<IAnimationEditor*>(Editor);
+									UDebugSkelMeshComponent* PreviewComponent = AnimationEditor->GetPersonaToolkit()->GetPreviewMeshComponent();
+									PreviewComponent->PreviewInstance->SetPosition(InValue.Object.PlaybackTime);
+									PreviewComponent->PreviewInstance->SetPlaying(false);
+									PreviewComponent->PreviewInstance->SetBlendSpacePosition(FVector(InValue.Object.BlendX, InValue.Object.BlendY, 0.0f));
+								}
+							}
+						}
+					}
 				});
 #else
 			return 
diff --git a/Engine/Plugins/Animation/GameplayInsights/Source/GameplayInsights/Private/VariantTreeNode.h b/Engine/Plugins/Animation/GameplayInsights/Source/GameplayInsights/Private/VariantTreeNode.h
--- a/Engine/Plugins/Animation/GameplayInsights/Source/GameplayInsights/Private/VariantTreeNode.h
+++ b/Engine/Plugins/Animation/GameplayInsights/Source/GameplayInsights/Private/VariantTreeNode.h
@@ -125,11 +125,14 @@ struct FVariantTreeNode : TSharedFromThis<FVariantTreeNode>
 		return MakeShared<FVariantTreeNode>(InName, Value, InId);
 	}
 
-	static TSharedRef<FVariantTreeNode> MakeObject(const FText& InName, uint64 InValue, uint64 InId = INDEX_NONE)
+	static TSharedRef<FVariantTreeNode> MakeObject(const FText& InName, uint64 InValue, uint64 InId = INDEX_NONE, float PlaybackTime = 0.0f, float BlendX = 0.0f, float BlendY = 0.0f)
 	{
 		FVariantValue Value;
 		Value.Type = EAnimNodeValueType::Object;
 		Value.Object.Value = InValue;
+		Value.Object.PlaybackTime = PlaybackTime;
+		Value.Object.BlendX = BlendX;
+		Value.Object.BlendY = BlendY;
 
 		return MakeShared<FVariantTreeNode>(InName, Value, InId);
 	}
diff --git a/Engine/Plugins/Animation/GameplayInsights/Source/GameplayInsights/Public/IAnimationProvider.h b/Engine/Plugins/Animation/GameplayInsights/Source/GameplayInsights/Public/IAnimationProvider.h
--- a/Engine/Plugins/Animation/GameplayInsights/Source/GameplayInsights/Public/IAnimationProvider.h
+++ b/Engine/Plugins/Animation/GameplayInsights/Source/GameplayInsights/Public/IAnimationProvider.h
@@ -164,6 +164,11 @@ struct FVariantValue
 		struct
 		{
 			uint64 Value;
+			// for animation assets, store the playback time so we can open the editor to a specific time
+			float PlaybackTime;
+			// for blend spaces, store the x/y parameters so we can open the editor at those parameters
+			float BlendX;
+			float BlendY;
 		} Object;
 		struct
 		{
