ifdefs for RewindDebugger to compile under shipping editor

diff --git a/Engine/Plugins/Animation/GameplayInsights/Source/RewindDebugger/Private/RewindDebugger.cpp b/Engine/Plugins/Animation/GameplayInsights/Source/RewindDebugger/Private/RewindDebugger.cpp
--- a/Engine/Plugins/Animation/GameplayInsights/Source/RewindDebugger/Private/RewindDebugger.cpp
+++ b/Engine/Plugins/Animation/GameplayInsights/Source/RewindDebugger/Private/RewindDebugger.cpp
@@ -135,7 +135,9 @@ void FRewindDebugger::OnPIEPaused(bool bSimulating)
 	if (bRecording)
 	{
 		UWorld* World = GetWorldToVisualize();
+#if OBJECT_TRACE_ENABLED
 		RecordingDuration.Set(FObjectTrace::GetWorldElapsedTime(World));
+#endif // OBJECT_TRACE_ENABLED
 		SetCurrentScrubTime(RecordingDuration.Get());
 	}
 }
@@ -172,7 +174,9 @@ void FRewindDebugger::OnPIESingleStepped(bool bSimulating)
 	if (bRecording)
 	{
 		UWorld* World = GetWorldToVisualize();
+#if OBJECT_TRACE_ENABLED
 		RecordingDuration.Set(FObjectTrace::GetWorldElapsedTime(World));
+#endif // OBJECT_TRACE_ENABLED
 		SetCurrentScrubTime(RecordingDuration.Get());
 	}
 }
@@ -294,11 +298,13 @@ void FRewindDebugger::StartRecording()
 	RecordingIndex++;
 	bRecording = true;
 
+#if OBJECT_TRACE_ENABLED
 	// setup FObjectTrace to start tracking tracing times from 0
 	// and increment the RecordingIndex so we can use it to distinguish between the latest recording and older ones
 	UWorld* World = GetWorldToVisualize();
 	FObjectTrace::ResetWorldElapsedTime(World);
 	FObjectTrace::SetWorldRecordingIndex(World, RecordingIndex);
+#endif // OBJECT_TRACE_ENABLED
 }
 
 bool FRewindDebugger::ShouldAutoRecordOnPIE() const
@@ -672,7 +678,9 @@ void FRewindDebugger::Tick(float DeltaTime)
 				if (bRecording)
 				{
 					TRACE_CPUPROFILER_EVENT_SCOPE(FRewindDebugger::Tick_UpdateSimulating);
+#if OBJECT_TRACE_ENABLED
 					RecordingDuration.Set(FObjectTrace::GetWorldElapsedTime(World));
+#endif // OBJECT_TRACE_ENABLED
 					SetCurrentScrubTime(RecordingDuration.Get());
 					TrackCursorDelegate.ExecuteIfBound(false);
 				}
@@ -710,6 +718,7 @@ void FRewindDebugger::Tick(float DeltaTime)
 								uint64 TargetActorId = GetTargetActorId();
 								if (TargetActorId != 0)
 								{
+#if OBJECT_TRACE_ENABLED
 									if(UObject* ObjectInstance = FObjectTrace::GetObjectFromId(TargetActorId))
 									{
 										if (AActor* TargetActor = Cast<AActor>(ObjectInstance))
@@ -743,6 +752,7 @@ void FRewindDebugger::Tick(float DeltaTime)
 											}
 										}
 									}
+#endif // OBJECT_TRACE_ENABLED
 								}
 							}
 							
@@ -754,6 +764,7 @@ void FRewindDebugger::Tick(float DeltaTime)
 								TRACE_CPUPROFILER_EVENT_SCOPE(FRewindDebugger::Tick_UpdatePoses);
 								AnimationProvider->EnumerateSkeletalMeshPoseTimelines([this, &Frame, AnimationProvider, GameplayProvider](uint64 ObjectId, const IAnimationProvider::SkeletalMeshPoseTimeline& TimelineData)
 								{
+#if OBJECT_TRACE_ENABLED
 									if(UObject* ObjectInstance = FObjectTrace::GetObjectFromId(ObjectId))
 									{
 										if(USkeletalMeshComponent* MeshComponent = Cast<USkeletalMeshComponent>(ObjectInstance))
@@ -793,6 +804,7 @@ void FRewindDebugger::Tick(float DeltaTime)
 											});
 										}
 									}
+#endif // OBJECT_TRACE_ENABLED
 								});
 							}
 
@@ -804,6 +816,7 @@ void FRewindDebugger::Tick(float DeltaTime)
 								// - if it is copy that debug data into the class debug data for the blueprint debugger
 								AnimationProvider->EnumerateAnimGraphTimelines([&Frame, AnimationProvider, GameplayProvider](uint64 ObjectId, const IAnimationProvider::AnimGraphTimeline& AnimGraphTimeline)
 								{
+#if OBJECT_TRACE_ENABLED
 									if(UObject* ObjectInstance = FObjectTrace::GetObjectFromId(ObjectId))
 									{
 										if(UAnimInstance* AnimInstance = Cast<UAnimInstance>(ObjectInstance))
@@ -966,6 +979,7 @@ void FRewindDebugger::Tick(float DeltaTime)
 											}
 										}
 									}
+#endif // OBJECT_TRACE_ENABLED
 									return TraceServices::EEventEnumerate::Continue;
 								});
 							}
diff --git a/Engine/Plugins/Animation/GameplayInsights/Source/RewindDebugger/Private/SRewindDebugger.cpp b/Engine/Plugins/Animation/GameplayInsights/Source/RewindDebugger/Private/SRewindDebugger.cpp
--- a/Engine/Plugins/Animation/GameplayInsights/Source/RewindDebugger/Private/SRewindDebugger.cpp
+++ b/Engine/Plugins/Animation/GameplayInsights/Source/RewindDebugger/Private/SRewindDebugger.cpp
@@ -363,10 +363,12 @@ void SRewindDebugger::Construct(const FArguments& InArgs, TSharedRef<FUICommandL
 										FSlateIcon ActorIcon = FSlateIconFinder::FindIconForClass(AActor::StaticClass());
 										if (DebugComponents != nullptr && DebugComponents->Num()>0)
 										{
+#if OBJECT_TRACE_ENABLED
 											if (UObject* Object = FObjectTrace::GetObjectFromId((*DebugComponents)[0]->GetObjectId()))
 											{
 												ActorIcon = FSlateIconFinder::FindIconForClass(Object->GetClass());
 											}
+#endif // OBJECT_TRACE_ENABLED
 										}
 
 										return ActorIcon.GetIcon();
@@ -383,7 +385,7 @@ void SRewindDebugger::Construct(const FArguments& InArgs, TSharedRef<FUICommandL
 									}
 
 									FText ReadableName = (*DebugComponents)[0]->GetDisplayName();
-
+#if OBJECT_TRACE_ENABLED
 									if (UObject* Object = FObjectTrace::GetObjectFromId((*DebugComponents)[0]->GetObjectId()))
 									{
 										if (AActor* Actor = Cast<AActor>(Object))
@@ -391,6 +393,7 @@ void SRewindDebugger::Construct(const FArguments& InArgs, TSharedRef<FUICommandL
 											ReadableName = FText::FromString(Actor->GetActorLabel());
 										}
 									}
+#endif // OBJECT_TRACE_ENABLED
 
 									return ReadableName;
 								} )
