Change Rewind Debugger curve rendering to display overlapping blend curves better

diff --git a/Engine/Plugins/Animation/GameplayInsights/Source/GameplayInsights/Private/SCurveTimelineView.cpp b/Engine/Plugins/Animation/GameplayInsights/Source/GameplayInsights/Private/SCurveTimelineView.cpp
--- a/Engine/Plugins/Animation/GameplayInsights/Source/GameplayInsights/Private/SCurveTimelineView.cpp
+++ b/Engine/Plugins/Animation/GameplayInsights/Source/GameplayInsights/Private/SCurveTimelineView.cpp
@@ -21,6 +21,8 @@ void SCurveTimelineView::Construct( const SCurveTimelineView::FArguments& InArgs
 
 	// set clipping on by default, since the OnPaint function is drawing outside the bounds
 	Clipping = EWidgetClipping::ClipToBounds;
+	
+	FillBrush = FAppStyle::GetBrush("Sequencer.SectionArea.Background");
 }
 
 int32 SCurveTimelineView::OnPaint( const FPaintArgs& Args, const FGeometry& AllottedGeometry, const FSlateRect& MyCullingRect, FSlateWindowElementList& OutDrawElements, int32 LayerId, const FWidgetStyle& InWidgetStyle, bool bParentEnabled ) const
@@ -160,103 +162,83 @@ int32 SCurveTimelineView::PaintCurve(const FGeometry& AllottedGeometry, const FS
 		{
 			if (RenderFill.Get())
 			{
-				// don't call this every frame
-				FSlateResourceHandle ResourceHandle = FSlateApplication::Get().GetRenderer()->GetResourceHandle(*FAppStyle::GetBrush("Sequencer.SectionArea.Background"));
 			
 				const FSlateRenderTransform& Transform = AllottedGeometry.GetAccumulatedRenderTransform();
 				FColor Color = FillColor.Get().ToFColor(true);
 
-				// filled polygons
-				TArray<FSlateVertex> Vertices;
-				TArray<SlateIndex> Indices;
-				Vertices.Reserve((NumPoints-1)*4);
-				Indices.Reserve((NumPoints-1)*6);
-				
+				float Y1 = Size.Y-1;
 				for(int i=1; i<NumPoints; i++)
 				{
-					const FTimelineCurveData::CurvePoint& Point1 = CurvePoints[i-1];
 					const FTimelineCurveData::CurvePoint& Point2 = CurvePoints[i];
-				
-					if (Point2.Time - Point1.Time < LargeFrameTime)
+
+					// find time of previous sample - step backwards until a sample that has a different time
+					double CurrentTime = CurvePoints[i].Time;
+					double PreviousTime = CurrentTime;
+					for (int j = i - 1; j >= 0; j--)
 					{
-						float X1 = RangeToScreen.InputToLocalX(Point1.Time);
-						const float NormalizedValue1 = (Point1.Value - MinValue) / Range;
-						float Y1 = Size.Y - (Size.Y - 2) * NormalizedValue1;
-						float X2 = RangeToScreen.InputToLocalX(Point2.Time);
-						const float NormalizedValue2 = (Point2.Value - MinValue) / Range;
-						float Y2 = Size.Y - (Size.Y - 2) * NormalizedValue2;
-					
-						Vertices.Add(FSlateVertex::Make<ESlateVertexRounding::Disabled>(Transform, FVector2f(X1,Size.Y-1) ,FVector2f(0,0), Color));
-						Vertices.Add(FSlateVertex::Make<ESlateVertexRounding::Disabled>(Transform, FVector2f(X1,Y1),FVector2f(0,0), Color));
-						Vertices.Add(FSlateVertex::Make<ESlateVertexRounding::Disabled>(Transform, FVector2f(X2,Y2),FVector2f(0,0), Color));
-						Vertices.Add(FSlateVertex::Make<ESlateVertexRounding::Disabled>(Transform, FVector2f(X2,Size.Y-1),FVector2f(0,0), Color));
-						
-						int32 Index0 = Vertices.Num() - 4;
-                        int32 Index1 = Vertices.Num() - 3;
-                        int32 Index2 = Vertices.Num() - 2;
-                        int32 Index3 = Vertices.Num() - 1;
-                       
-						Indices.Add(Index0);
-						Indices.Add(Index1);
-						Indices.Add(Index2);
-						Indices.Add(Index0);
-						Indices.Add(Index2);
-						Indices.Add(Index3);
+						if (CurvePoints[j].Time != CurrentTime)
+						{
+							PreviousTime = CurvePoints[j].Time;
+							break;
+						}
 					}
 					
-					FSlateDrawElement::MakeCustomVerts(
-								OutDrawElements,
-										LayerId++,
-										ResourceHandle,
-										Vertices,
-										Indices,
-										nullptr,
-										0,
-										0);
+					const float NormalizedValue = (Point2.Value - MinValue) / Range;
+					float Y2 = Size.Y - (Size.Y - 2) * NormalizedValue;
+				
+					if (CurrentTime - PreviousTime < LargeFrameTime)
+					{
+						float X1 = RangeToScreen.InputToLocalX(PreviousTime);
+						float X2 = RangeToScreen.InputToLocalX(CurrentTime);
+						FPaintGeometry PaintGeo = AllottedGeometry.ToPaintGeometry(FVector2D(X1, Y2) , FVector2D(X2-X1, Y1-Y2), 1.0);
+						FSlateDrawElement::MakeBox(OutDrawElements, LayerId, PaintGeo, FillBrush, ESlateDrawEffect::None, FillColor.Get());
+					}
 				}
 			}
+			else
+			{
+				// Lines rendering
 
+				TArray<FVector2D> Points;
+				Points.Reserve(NumPoints);
 
-			// Lines rendering
-			
-			TArray<FVector2D> Points;
-			Points.Reserve(NumPoints);
-
-			float PrevTime = CurvePoints[0].Time;
-			for(int i=0; i<NumPoints; i++)
-			{
-				const FTimelineCurveData::CurvePoint& Point = CurvePoints[i];
-				if (Point.Time - PrevTime > LargeFrameTime && Points.Num()>1)
+				float PrevTime = CurvePoints[0].Time;
+				for(int i=0; i<NumPoints; i++)
 				{
-					// break the line list - data has stopped and started again
-					FSlateDrawElement::MakeLines(
-								OutDrawElements,
-										LayerId++,
-										AllottedGeometry.ToPaintGeometry(),
-										Points,
-										LineDrawEffects,
-										LineColor,
-										false
-										);
-					Points.SetNum(0,false);
+					const FTimelineCurveData::CurvePoint& Point = CurvePoints[i];
+					if (Point.Time - PrevTime > LargeFrameTime && Points.Num()>1)
+					{
+						// break the line list - data has stopped and started again
+						FSlateDrawElement::MakeLines(
+									OutDrawElements,
+											LayerId++,
+											AllottedGeometry.ToPaintGeometry(),
+											Points,
+											LineDrawEffects,
+											LineColor,
+											false
+											);
+						Points.SetNum(0,false);
+					}
+			
+					float X = RangeToScreen.InputToLocalX(Point.Time);
+					PrevTime = Point.Time;
+					const float NormalizedValue = (Point.Value - MinValue) / Range;
+					float Y = Size.Y - (Size.Y - 2) * NormalizedValue;
+					Points.Add(FVector2D(X,Y));
 				}
 			
-				float X = RangeToScreen.InputToLocalX(Point.Time);
-				PrevTime = Point.Time;
-				const float NormalizedValue = (Point.Value - MinValue) / Range;
-				float Y = Size.Y - (Size.Y - 2) * NormalizedValue;
-				Points.Add(FVector2D(X,Y));
+				FSlateDrawElement::MakeLines(
+							OutDrawElements,
+									LayerId,
+									AllottedGeometry.ToPaintGeometry(),
+									Points,
+									LineDrawEffects,
+									LineColor,
+									false
+									);
 			}
-			
-			FSlateDrawElement::MakeLines(
-						OutDrawElements,
-								LayerId,
-								AllottedGeometry.ToPaintGeometry(),
-								Points,
-								LineDrawEffects,
-								LineColor,
-								false
-								);
+		
 		}
 	}
 
diff --git a/Engine/Plugins/Animation/GameplayInsights/Source/GameplayInsights/Private/SCurveTimelineView.h b/Engine/Plugins/Animation/GameplayInsights/Source/GameplayInsights/Private/SCurveTimelineView.h
--- a/Engine/Plugins/Animation/GameplayInsights/Source/GameplayInsights/Private/SCurveTimelineView.h
+++ b/Engine/Plugins/Animation/GameplayInsights/Source/GameplayInsights/Private/SCurveTimelineView.h
@@ -90,7 +90,7 @@ protected:
 	FText CurveToolTipInputText;		
 	FText CurveToolTipOutputText;
 	
-	const FSlateBrush* WhiteBrush;
+	const FSlateBrush* FillBrush;
 
 	bool bUseFixedRange;
 	float FixedRangeMin;
