Rewind debugger creates one trace file per recording

diff --git a/Engine/Plugins/Animation/GameplayInsights/Source/GameplayInsights/Private/AnimationProvider.cpp b/Engine/Plugins/Animation/GameplayInsights/Source/GameplayInsights/Private/AnimationProvider.cpp
--- a/Engine/Plugins/Animation/GameplayInsights/Source/GameplayInsights/Private/AnimationProvider.cpp
+++ b/Engine/Plugins/Animation/GameplayInsights/Source/GameplayInsights/Private/AnimationProvider.cpp
@@ -682,10 +682,13 @@ void FAnimationProvider::AppendSkeletalMeshComponent(uint64 InObjectId, uint64 I
 
 	// terminate existing scopes
 	uint64 NumEvents = TimelineStorage->Timeline->GetEventCount();
-	if(NumEvents > 0)
+	if (NumEvents > 0)
 	{
 		// Add end event at current time
-		TimelineStorage->Timeline->EndEvent(NumEvents - 1, InProfileTime);
+		if (TimelineStorage->Timeline->GetEventEndTime(NumEvents - 1) > InProfileTime)
+		{
+			TimelineStorage->Timeline->EndEvent(NumEvents - 1, InProfileTime);
+		}
 	}
 
 	const int32 NumCurves = InCurveIds.Num();
diff --git a/Engine/Plugins/Animation/GameplayInsights/Source/GameplayInsights/Private/GameplayProvider.cpp b/Engine/Plugins/Animation/GameplayInsights/Source/GameplayInsights/Private/GameplayProvider.cpp
--- a/Engine/Plugins/Animation/GameplayInsights/Source/GameplayInsights/Private/GameplayProvider.cpp
+++ b/Engine/Plugins/Animation/GameplayInsights/Source/GameplayInsights/Private/GameplayProvider.cpp
@@ -486,7 +486,10 @@ void FGameplayProvider::AppendObjectLifetimeEnd(uint64 InObjectId, double InProf
 
 	if (const uint64 *FoundIndex = ActiveObjectRecordingLifetimes.Find(InObjectId))
 	{
-		ObjectRecordingLifetimes.EndEvent(*FoundIndex, InRecordingTime);
+		if(ObjectRecordingLifetimes.GetEventEndTime(*FoundIndex) > InRecordingTime)
+		{
+			ObjectRecordingLifetimes.EndEvent(*FoundIndex, InRecordingTime);
+		}
 		// do not remove from ActiveObjectRecordingLifetimes - lifetimes can be queried by Object Id in GetObjectRecordingLifetime
 	}
 
diff --git a/Engine/Plugins/Animation/GameplayInsights/Source/RewindDebugger/Private/RewindDebugger.cpp b/Engine/Plugins/Animation/GameplayInsights/Source/RewindDebugger/Private/RewindDebugger.cpp
--- a/Engine/Plugins/Animation/GameplayInsights/Source/RewindDebugger/Private/RewindDebugger.cpp
+++ b/Engine/Plugins/Animation/GameplayInsights/Source/RewindDebugger/Private/RewindDebugger.cpp
@@ -7,8 +7,8 @@
 #include "IAnimationProvider.h"
 #include "IGameplayProvider.h"
 #include "Insights/IUnrealInsightsModule.h"
-#include "IGameplayInsightsModule.h"
 #include "Modules/ModuleManager.h"
+#include "Animation/AnimTrace.h"
 #include "ObjectTrace.h"
 #include "TraceServices/Model/Frames.h"
 #include "SLevelViewport.h"
@@ -19,12 +19,15 @@
 #include "RewindDebuggerObjectTrack.h"
 #include "Animation/AnimBlueprint.h"
 #include "Animation/AnimBlueprintGeneratedClass.h"
+#include "EngineUtils.h"
 #include "ToolMenus.h"
 #include "RewindDebuggerSettings.h"
 #include "LevelEditor.h"
 #include "RewindDebuggerModule.h"
 #include "Engine/PoseWatch.h"
 #include "ProfilingDebugging/TraceAuxiliary.h"
+#include "UObject/UObjectIterator.h"
+#include "Engine/World.h"
 
 static void IterateExtensions(TFunction<void(IRewindDebuggerExtension* Extension)> IteratorFunction)
 {
@@ -39,6 +42,16 @@ static void IterateExtensions(TFunction<void(IRewindDebuggerExtension* Extension
 	}
 }
 
+static void TraceSubobjects(UObject* OuterObject)
+{
+	TArray<UObject*> Subobjects;
+	GetObjectsWithOuter(OuterObject, Subobjects, true);
+	for (UObject* Subobject : Subobjects)
+	{
+		TRACE_OBJECT_LIFETIME_BEGIN(Subobject);
+	}
+}
+
 FRewindDebugger::FRewindDebugger()  :
 	ControlState(FRewindDebugger::EControlState::Pause),
 	bPIEStarted(false),
@@ -66,7 +79,21 @@ FRewindDebugger::FRewindDebugger()  :
 	FEditorDelegates::EndPIE.AddRaw(this, &FRewindDebugger::OnPIEStopped);
 	FEditorDelegates::SingleStepPIE.AddRaw(this, &FRewindDebugger::OnPIESingleStepped);
 
-	DebugTargetActor.OnPropertyChanged = DebugTargetActor.OnPropertyChanged.CreateLambda([this](FString Target) { RefreshDebugTracks(); });
+	DebugTargetActor.OnPropertyChanged = DebugTargetActor.OnPropertyChanged.CreateLambda([this](FString Target)
+		{
+			TargetObjectIds.SetNum(0);
+			GetTargetObjectIds(TargetObjectIds);
+			// make sure all the SubObjects of the target actor have been traced
+			for (uint64 TargetObjectId : TargetObjectIds)
+			{
+				if (UObject* TargetObject = FObjectTrace::GetObjectFromId(TargetObjectId))
+				{
+					TraceSubobjects(TargetObject);
+				}
+			}
+
+			RefreshDebugTracks();
+		});
 
 	UnrealInsightsModule = &FModuleManager::LoadModuleChecked<IUnrealInsightsModule>("TraceInsights");
 
@@ -78,9 +105,6 @@ FRewindDebugger::FRewindDebugger()  :
 
 		return true;
 	});
-
-	IGameplayInsightsModule* GameplayInsightsModule = &FModuleManager::LoadModuleChecked<IGameplayInsightsModule>("GameplayInsights");
-	GameplayInsightsModule->StartTrace();
 }
 
 FRewindDebugger::~FRewindDebugger() 
@@ -119,8 +143,6 @@ void FRewindDebugger::OnPIEStarted(bool bSimulating)
 	bPIEStarted = true;
 	bPIESimulating = true;
 
-	UE::Trace::ToggleChannel(TEXT("Object"), true);
-
 	if (ShouldAutoRecordOnPIE())
 	{
 		StartRecording();
@@ -188,8 +210,6 @@ void FRewindDebugger::OnPIEStopped(bool bSimulating)
 	bPIESimulating = false;
 	MeshComponentsToReset.Empty();
 
-	UE::Trace::ToggleChannel(TEXT("Object"), false);
-
 	StopRecording();
 	// clear the current recording (until we support playback in the Editor world on spawned actors)
 	RecordingDuration.Set(0);
@@ -215,8 +235,7 @@ uint64 FRewindDebugger::GetTargetActorId() const
 	{
 		if (const IGameplayProvider* GameplayProvider = Session->ReadProvider<IGameplayProvider>("GameplayProvider"))
 		{
-			double Time = CurrentTraceTime();
-			GameplayProvider->EnumerateObjects(Time, Time, [this,&TargetActorId](const FObjectInfo& InObjectInfo)
+			GameplayProvider->EnumerateObjects(CurrentTraceRange.GetLowerBoundValue(), CurrentTraceRange.GetUpperBoundValue(), [this, &TargetActorId](const FObjectInfo& InObjectInfo)
 			{
 				if (DebugTargetActor.Get() == InObjectInfo.Name)
 				{
@@ -229,9 +248,64 @@ uint64 FRewindDebugger::GetTargetActorId() const
 	return TargetActorId;
 }
 
+
+void FRewindDebugger::GetTargetObjectIds(TArray<uint64>& OutTargetObjectIds) const
+{
+	OutTargetObjectIds.Empty(2);
+
+	if (DebugTargetActor.Get() == "")
+	{
+		return;
+	}
+
+	if (const TraceServices::IAnalysisSession* Session = GetAnalysisSession())
+	{
+		TraceServices::FAnalysisSessionReadScope SessionReadScope(*Session);
+
+		if (const IGameplayProvider* GameplayProvider = Session->ReadProvider<IGameplayProvider>("GameplayProvider"))
+		{
+			GameplayProvider->EnumerateObjects(CurrentTraceRange.GetLowerBoundValue(), CurrentTraceRange.GetUpperBoundValue(), [this, &OutTargetObjectIds](const FObjectInfo& InObjectInfo)
+				{
+					if (DebugTargetActor.Get() == InObjectInfo.Name)
+					{
+						OutTargetObjectIds.Add(InObjectInfo.Id);
+					}
+				});
+		}
+	}
+
+	// make sure all the SubObjects of the target actor have been traced
+	for (uint64 OutTargetObjectId : TargetObjectIds)
+	{
+		if (UObject* TargetObject = FObjectTrace::GetObjectFromId(OutTargetObjectId))
+		{
+			TraceSubobjects(TargetObject);
+		}
+	}
+}
+
+
+
 void FRewindDebugger::RefreshDebugTracks()
 {
 	TRACE_CPUPROFILER_EVENT_SCOPE(FRewindDebugger::RefreshDebugTracks);
+
+	if (TargetObjectIds.Num() == 0)
+	{
+		GetTargetObjectIds(TargetObjectIds);
+
+		if (TargetObjectIds.Num() == 0)
+		{
+			if (DebugTracks.Num() != 0)
+			{
+				// clear tracks so we don't show data from previous recordings
+				DebugTracks.SetNum(0);
+				ComponentListChangedDelegate.ExecuteIfBound();
+			}
+			return;
+		}
+	}
+
 	if (const TraceServices::IAnalysisSession* Session = GetAnalysisSession())
 	{
 		TraceServices::FAnalysisSessionReadScope SessionReadScope(*Session);
@@ -240,31 +314,35 @@ void FRewindDebugger::RefreshDebugTracks()
 
 		if (const IGameplayProvider* GameplayProvider = Session->ReadProvider<IGameplayProvider>("GameplayProvider"))
 		{
-			uint64 TargetActorId = GetTargetActorId();
 
 			bool bChanged = false;
 
-			// add actor (even if it isn't found in the gameplay provider)
-			if (DebugTracks.Num() == 0)
+			// remove any existing tracks that don't match the current list of object ids
+			for (int TrackIndex = DebugTracks.Num()-1; TrackIndex>=0; TrackIndex--)
 			{
-				bChanged = true;
-				const bool bAddController = true;
-				DebugTracks.Add(MakeShared<RewindDebugger::FRewindDebuggerObjectTrack>(TargetActorId, DebugTargetActor.Get(), bAddController));
+				uint64* FoundId = TargetObjectIds.FindByPredicate([this, TrackIndex](const uint64& TrackId) { return DebugTracks[TrackIndex]->GetObjectId() == TrackId; });
+				if (!FoundId)
+				{
+					DebugTracks.RemoveAt(TrackIndex);
+				}
 			}
-			else
+
+			// add new tracks for current list of object ids if they don't already exist
+			for (uint64 TargetObjectId : TargetObjectIds)
 			{
-				FString DebugTargetActorName = DebugTargetActor.Get();
+				TSharedPtr<RewindDebugger::FRewindDebuggerTrack>* FoundTrack = DebugTracks.FindByPredicate([TargetObjectId](const TSharedPtr<RewindDebugger::FRewindDebuggerTrack>& Track) { return Track->GetObjectId() == TargetObjectId; });
 
-				if (DebugTracks[0]->GetDisplayName().ToString() != DebugTargetActorName || DebugTracks[0]->GetObjectId() != TargetActorId)
+				if (!FoundTrack)
 				{
+					DebugTracks.Add(MakeShared<RewindDebugger::FRewindDebuggerObjectTrack>(TargetObjectId, DebugTargetActor.Get(), true));
 					bChanged = true;
-					DebugTracks[0] = MakeShared<RewindDebugger::FRewindDebuggerObjectTrack>(TargetActorId, DebugTargetActorName);
 				}
 			}
 
-			if (TargetActorId != 0 && DebugTracks.Num() > 0)
+			// update all tracks
+			for (TSharedPtr<RewindDebugger::FRewindDebuggerTrack>& DebugTrack : DebugTracks )
 			{
-				bChanged = bChanged || DebugTracks[0]->Update();
+				bChanged = bChanged || DebugTrack->Update();
 			}
 
 			if (bChanged)
@@ -275,17 +353,53 @@ void FRewindDebugger::RefreshDebugTracks()
 	}
 }
 
+namespace
+{
+	static void DisableAllTraceChannels()
+	{
+		UE::Trace::EnumerateChannels([](const ANSICHAR* ChannelName, bool bEnabled, void*)
+						{
+							if (bEnabled)
+							{
+								FString ChannelNameFString(ChannelName);
+								UE::Trace::ToggleChannel(ChannelNameFString.GetCharArray().GetData(), false);
+							}
+						}
+						, nullptr);
+	}
+}
+
 void FRewindDebugger::StartRecording()
 {
 	if (!CanStartRecording())
 	{
 		return;
 	}
+	
+	// Clear caches
+	FObjectTrace::Reset();
+	FAnimTrace::Reset();
+	
+	RecordingDuration.Set(0);
+	// RecordingIndex++;
+	bRecording = true;
+
+	// Disable all trace channels, and then enable only the ones needed by RewindDebugger
+	// for systems with RewindDebugger integration, they should enable their channel(s) in an Extension in "RecordingStarted"
+	DisableAllTraceChannels();
 
-	// Enable Object and Animation Trace filters
+	// Clear all buffered data and prevent data from previous recordings from leaking into the new recording
+	FTraceAuxiliary::FOptions Options;
+	Options.bExcludeTail = true;
+	
+	FTraceAuxiliary::Start(FTraceAuxiliary::EConnectionType::Network, TEXT("127.0.0.1"), TEXT(""), &Options, LogRewindDebugger);
+
+	UE::Trace::ToggleChannel(TEXT("Object"), true);
 	UE::Trace::ToggleChannel(TEXT("ObjectProperties"), true);
 	UE::Trace::ToggleChannel(TEXT("Animation"), true);
 	UE::Trace::ToggleChannel(TEXT("Frame"), true);
+	
+	UnrealInsightsModule->StartAnalysisForLastLiveSession();
 
 	// update extensions
 	IterateExtensions([this](IRewindDebuggerExtension* Extension)
@@ -294,17 +408,27 @@ void FRewindDebugger::StartRecording()
 		}
 	);
 
-	RecordingDuration.Set(0);
-	RecordingIndex++;
-	bRecording = true;
-
 #if OBJECT_TRACE_ENABLED
-	// setup FObjectTrace to start tracking tracing times from 0
-	// and increment the RecordingIndex so we can use it to distinguish between the latest recording and older ones
-	UWorld* World = GetWorldToVisualize();
-	FObjectTrace::ResetWorldElapsedTime(World);
-	FObjectTrace::SetWorldRecordingIndex(World, RecordingIndex);
+	// trace each play-in-editor world, and all the actors in it.
+	for (TObjectIterator<UWorld> World; World; ++World)
+	{
+		if (World->IsPlayInEditor())
+		{
+			FObjectTrace::ResetWorldElapsedTime(*World);
+			FObjectTrace::SetWorldRecordingIndex(*World, RecordingIndex);
+			
+			TRACE_WORLD(*World);
+				
+			for (TActorIterator<AActor> Iterator(*World); Iterator; ++Iterator)
+			{
+				TRACE_OBJECT_LIFETIME_BEGIN(*Iterator);
+			}
+		}
+	}
 #endif // OBJECT_TRACE_ENABLED
+
+	TargetObjectIds.Empty(2);
+
 }
 
 bool FRewindDebugger::ShouldAutoRecordOnPIE() const
@@ -321,11 +445,6 @@ void FRewindDebugger::StopRecording()
 {
 	if (bRecording)
 	{
-		// Enable Object and Animation Trace filters
-		UE::Trace::ToggleChannel(TEXT("ObjectProperties"), false);
-		UE::Trace::ToggleChannel(TEXT("Animation"), false);
-		UE::Trace::ToggleChannel(TEXT("Frame"), false);
-
 		// update extensions
 		IterateExtensions([this](IRewindDebuggerExtension* Extension)
 			{
@@ -334,6 +453,9 @@ void FRewindDebugger::StopRecording()
 		);
 
 		bRecording = false;
+		
+		DisableAllTraceChannels();
+		FTraceAuxiliary::Stop();
 	}
 }
 
@@ -785,21 +907,23 @@ void FRewindDebugger::Tick(float DeltaTime)
 												if (PoseMessage)
 												{
 													FTransform ComponentWorldTransform;
-													const FSkeletalMeshInfo* SkeletalMeshInfo = AnimationProvider->FindSkeletalMeshInfo(PoseMessage->MeshId);
-													AnimationProvider->GetSkeletalMeshComponentSpacePose(*PoseMessage, *SkeletalMeshInfo, ComponentWorldTransform, MeshComponent->GetEditableComponentSpaceTransforms());
-													MeshComponent->ApplyEditedComponentSpaceTransforms();
-
-													if (MeshComponentsToReset.Find(ObjectId) == nullptr)
+													if (const FSkeletalMeshInfo* SkeletalMeshInfo = AnimationProvider->FindSkeletalMeshInfo(PoseMessage->MeshId))
 													{
-														FMeshComponentResetData ResetData;
-														ResetData.Component = MeshComponent;
-														ResetData.RelativeTransform = MeshComponent->GetRelativeTransform();
-														MeshComponentsToReset.Add(ObjectId, ResetData);
-													}
+														AnimationProvider->GetSkeletalMeshComponentSpacePose(*PoseMessage, *SkeletalMeshInfo, ComponentWorldTransform, MeshComponent->GetEditableComponentSpaceTransforms());
+														MeshComponent->ApplyEditedComponentSpaceTransforms();
 
-													MeshComponent->SetWorldTransform(ComponentWorldTransform, false, nullptr, ETeleportType::TeleportPhysics);
-													MeshComponent->SetForcedLOD(PoseMessage->LodIndex + 1);
-													MeshComponent->UpdateChildTransforms(EUpdateTransformFlags::None, ETeleportType::TeleportPhysics);
+														if (MeshComponentsToReset.Find(ObjectId) == nullptr)
+														{
+															FMeshComponentResetData ResetData;
+															ResetData.Component = MeshComponent;
+															ResetData.RelativeTransform = MeshComponent->GetRelativeTransform();
+															MeshComponentsToReset.Add(ObjectId, ResetData);
+														}
+
+														MeshComponent->SetWorldTransform(ComponentWorldTransform, false, nullptr, ETeleportType::TeleportPhysics);
+														MeshComponent->SetForcedLOD(PoseMessage->LodIndex + 1);
+														MeshComponent->UpdateChildTransforms(EUpdateTransformFlags::None, ETeleportType::TeleportPhysics);
+													}
 												}
 											});
 										}
@@ -1136,7 +1260,7 @@ TSharedPtr<RewindDebugger::FRewindDebuggerTrack> FRewindDebugger::GetSelectedTra
 void FRewindDebugger::RefreshDebugComponents(TArray<TSharedPtr<RewindDebugger::FRewindDebuggerTrack>>& InTracks, TArray<TSharedPtr<FDebugObjectInfo>>& OutComponents)
 {
 	OutComponents.SetNum(0);
-	for(auto& Track : InTracks)
+	for(TSharedPtr<RewindDebugger::FRewindDebuggerTrack>& Track : InTracks)
 	{
 		int Index = OutComponents.Num();
 		OutComponents.Add(MakeShared<FDebugObjectInfo>(Track->GetObjectId(), Track->GetDisplayName().ToString()));
diff --git a/Engine/Plugins/Animation/GameplayInsights/Source/RewindDebugger/Private/RewindDebugger.h b/Engine/Plugins/Animation/GameplayInsights/Source/RewindDebugger/Private/RewindDebugger.h
--- a/Engine/Plugins/Animation/GameplayInsights/Source/RewindDebugger/Private/RewindDebugger.h
+++ b/Engine/Plugins/Animation/GameplayInsights/Source/RewindDebugger/Private/RewindDebugger.h
@@ -46,6 +46,9 @@ public:
 	virtual TSharedPtr<RewindDebugger::FRewindDebuggerTrack> GetSelectedTrack() const override;
 	virtual TArray<TSharedPtr<FDebugObjectInfo>>& GetDebugComponents() override;
 	
+
+	void GetTargetObjectIds(TArray<uint64>& OutActorIds) const;
+
 	TArray<TSharedPtr<RewindDebugger::FRewindDebuggerTrack>>& GetDebugTracks() { return DebugTracks; }
 
 	// create singleton instance
@@ -183,6 +186,7 @@ private:
 		FTransform RelativeTransform;
 	};
 
+	TArray<uint64> TargetObjectIds;
 	TMap<uint64, FMeshComponentResetData> MeshComponentsToReset;
 
 	mutable class IUnrealInsightsModule *UnrealInsightsModule;
diff --git a/Engine/Source/Runtime/Engine/Private/Animation/AnimTrace.cpp b/Engine/Source/Runtime/Engine/Private/Animation/AnimTrace.cpp
--- a/Engine/Source/Runtime/Engine/Private/Animation/AnimTrace.cpp
+++ b/Engine/Source/Runtime/Engine/Private/Animation/AnimTrace.cpp
@@ -269,6 +269,12 @@ uint32 GAnimTraceCurrentNameId = 1;
 // Critical section used to lock global name map & index
 FCriticalSection GAnimTraceNameCriticalSection;
 
+void FAnimTrace::Reset()
+{
+	GSkeletalMeshTraceAnnotations.ClearAll();
+	GAnimTraceCurrentNameId = 1;
+}
+
 // Scratch buffers for various traces to avoid allocation churn.
 // These can be removed when lambda support is added for array fields to remove a memcpy.
 struct FAnimTraceScratchBuffers : public TThreadSingleton<FAnimTraceScratchBuffers>
diff --git a/Engine/Source/Runtime/Engine/Private/ObjectTrace.cpp b/Engine/Source/Runtime/Engine/Private/ObjectTrace.cpp
--- a/Engine/Source/Runtime/Engine/Private/ObjectTrace.cpp
+++ b/Engine/Source/Runtime/Engine/Private/ObjectTrace.cpp
@@ -180,6 +180,12 @@ void FObjectTrace::Destroy()
 	FWorldDelegates::OnWorldTickStart.Remove(WorldTickStartHandle);
 }
 
+void FObjectTrace::Reset()
+{
+	GObjectIdAnnotations.RemoveAllAnnotations();
+	GObjectTracedAnnotations.RemoveAllAnnotations();
+}
+
 uint64 FObjectTrace::GetObjectId(const UObject* InObject)
 {
 	// An object ID uses a combination of its own and its outer's index
@@ -234,6 +240,8 @@ void FObjectTrace::ResetWorldElapsedTime(const UWorld* World)
 	{
 		WorldSubsystem->ElapsedTime = 0;
 	}
+
+	GObjectTracedAnnotations.RemoveAllAnnotations();
 }
 
 double FObjectTrace::GetWorldElapsedTime(const UWorld* World)
@@ -416,6 +424,11 @@ void FObjectTrace::OutputObject(const UObject* InObject)
 		<< Object.OuterId(GetObjectId(InObject->GetOuter()))
 		<< Object.Name(ObjectName, ObjectNameLength)
 		<< Object.Path(*ObjectPathName, ObjectPathName.Len());
+
+	UE_TRACE_LOG(Object, ObjectLifetimeBegin2, ObjectChannel)
+		<< ObjectLifetimeBegin2.Cycle(FPlatformTime::Cycles64())
+		<< ObjectLifetimeBegin2.RecordingTime(FObjectTrace::GetWorldElapsedTime(InObject->GetWorld()))
+		<< ObjectLifetimeBegin2.Id(GetObjectId(InObject));
 }
 
 void FObjectTrace::OutputObjectEvent(const UObject* InObject, const TCHAR* InEvent)
@@ -446,28 +459,7 @@ void FObjectTrace::OutputObjectEvent(const UObject* InObject, const TCHAR* InEve
 
 void FObjectTrace::OutputObjectLifetimeBegin(const UObject* InObject)
 {
-	const bool bChannelEnabled = UE_TRACE_CHANNELEXPR_IS_ENABLED(ObjectChannel);
-	if (!bChannelEnabled || InObject == nullptr)
-	{
-		return;
-	}
-
-	if(InObject->HasAnyFlags(RF_ClassDefaultObject))
-	{
-		return;
-	}
-
-	if (CANNOT_TRACE_OBJECT(InObject->GetWorld()))
-	{
-		return;
-	}
-
 	TRACE_OBJECT(InObject);
-
-	UE_TRACE_LOG(Object, ObjectLifetimeBegin2, ObjectChannel)
-		<< ObjectLifetimeBegin2.Cycle(FPlatformTime::Cycles64())
-		<< ObjectLifetimeBegin2.RecordingTime(FObjectTrace::GetWorldElapsedTime(InObject->GetWorld()))
-		<< ObjectLifetimeBegin2.Id(GetObjectId(InObject));
 }
 
 void FObjectTrace::OutputObjectLifetimeEnd(const UObject* InObject)
diff --git a/Engine/Source/Runtime/Engine/Private/TraceFilter.cpp b/Engine/Source/Runtime/Engine/Private/TraceFilter.cpp
--- a/Engine/Source/Runtime/Engine/Private/TraceFilter.cpp
+++ b/Engine/Source/Runtime/Engine/Private/TraceFilter.cpp
@@ -277,7 +277,8 @@ FAutoConsoleCommand FlushFilterStateCommand(TEXT("TraceFilter.FlushState"), TEXT
 template<>
 bool ENGINE_API FTraceFilter::IsObjectTraceable</*bForceThreadSafe = */ true>(const UObject* InObject)
 {
-	return InObject ? GObjectFilterAnnotations.GetAnnotation(InObject).bIsTraceable : true;
+	// Object not found in the AnnotationMap means that it is at the default value, which is bIsTraceable == true
+	return GObjectFilterAnnotations.GetAnnotationMap().Find(InObject) == nullptr;
 }
 
 template<>
diff --git a/Engine/Source/Runtime/Engine/Public/Animation/AnimTrace.h b/Engine/Source/Runtime/Engine/Public/Animation/AnimTrace.h
--- a/Engine/Source/Runtime/Engine/Public/Animation/AnimTrace.h
+++ b/Engine/Source/Runtime/Engine/Public/Animation/AnimTrace.h
@@ -113,6 +113,9 @@ struct FAnimTrace
 		float Thickness;
 		bool bPersistentLines;
 	};
+	
+	/** Reset Caches so a new trace can be started*/
+	ENGINE_API static void Reset();
 
 	/** Helper function to output a tick record */
 	ENGINE_API FORCENOINLINE static void OutputAnimTickRecord(const FAnimationBaseContext& InContext, const FAnimTickRecord& InTickRecord);
diff --git a/Engine/Source/Runtime/Engine/Public/ObjectTrace.h b/Engine/Source/Runtime/Engine/Public/ObjectTrace.h
--- a/Engine/Source/Runtime/Engine/Public/ObjectTrace.h
+++ b/Engine/Source/Runtime/Engine/Public/ObjectTrace.h
@@ -48,6 +48,9 @@ struct FObjectTrace
 
 	/** Shut down object tracing */
 	ENGINE_API static void Destroy();
+	
+	/** Reset Caches so a new trace can be started*/
+	ENGINE_API static void Reset();
 
 	/** Helper function to output an object */
 	ENGINE_API static void OutputClass(const UClass* InClass);
