Fixed compile errors if OBJECT_TRACE_ENABLED is 0.

#rnx

diff --git a/Engine/Plugins/Animation/GameplayInsights/Source/RewindDebugger/Private/RewindDebugger.cpp b/Engine/Plugins/Animation/GameplayInsights/Source/RewindDebugger/Private/RewindDebugger.cpp
--- a/Engine/Plugins/Animation/GameplayInsights/Source/RewindDebugger/Private/RewindDebugger.cpp
+++ b/Engine/Plugins/Animation/GameplayInsights/Source/RewindDebugger/Private/RewindDebugger.cpp
@@ -84,6 +84,7 @@ FRewindDebugger::FRewindDebugger()  :
 			TargetObjectIds.SetNum(0);
 			GetTargetObjectIds(TargetObjectIds);
 			// make sure all the SubObjects of the target actor have been traced
+#if OBJECT_TRACE_ENABLED
 			for (uint64 TargetObjectId : TargetObjectIds)
 			{
 				if (UObject* TargetObject = FObjectTrace::GetObjectFromId(TargetObjectId))
@@ -91,6 +92,7 @@ FRewindDebugger::FRewindDebugger()  :
 					TraceSubobjects(TargetObject);
 				}
 			}
+#endif
 
 			RefreshDebugTracks();
 		});
@@ -275,6 +277,7 @@ void FRewindDebugger::GetTargetObjectIds(TArray<uint64>& OutTargetObjectIds) con
 	}
 
 	// make sure all the SubObjects of the target actor have been traced
+#if OBJECT_TRACE_ENABLED
 	for (uint64 OutTargetObjectId : TargetObjectIds)
 	{
 		if (UObject* TargetObject = FObjectTrace::GetObjectFromId(OutTargetObjectId))
@@ -282,6 +285,7 @@ void FRewindDebugger::GetTargetObjectIds(TArray<uint64>& OutTargetObjectIds) con
 			TraceSubobjects(TargetObject);
 		}
 	}
+#endif
 }
 
 
@@ -377,8 +381,10 @@ void FRewindDebugger::StartRecording()
 	}
 	
 	// Clear caches
+#if OBJECT_TRACE_ENABLED
 	FObjectTrace::Reset();
 	FAnimTrace::Reset();
+#endif
 	
 	RecordingDuration.Set(0);
 	// RecordingIndex++;
