Added support for Rewinder Debugger to open the current selected object's blueprint editor via double clicking in the object track.

diff --git a/Engine/Plugins/Animation/GameplayInsights/Source/RewindDebugger/Private/AnimInstanceHelpers.cpp b/Engine/Plugins/Animation/GameplayInsights/Source/RewindDebugger/Private/AnimInstanceHelpers.cpp
--- a/Engine/Plugins/Animation/GameplayInsights/Source/RewindDebugger/Private/AnimInstanceHelpers.cpp
+++ b/Engine/Plugins/Animation/GameplayInsights/Source/RewindDebugger/Private/AnimInstanceHelpers.cpp
@@ -16,7 +16,7 @@
 #define LOCTEXT_NAMESPACE "SRewindDebuggerAnimBPTools"
 
 #if OBJECT_TRACE_ENABLED
-static bool OpenBlueprintAndAttachDebugger(const TraceServices::IAnalysisSession* Session, uint64 ObjectId)
+static bool OpenAnimBlueprintAndAttachDebugger(const TraceServices::IAnalysisSession* Session, uint64 ObjectId)
 {
 	TraceServices::FAnalysisSessionReadScope SessionReadScope(*Session);
 
@@ -62,7 +62,7 @@ bool FAnimInstanceDoubleClickHandler::HandleDoubleClick(IRewindDebugger* RewindD
 	TSharedPtr<FDebugObjectInfo> SelectedObject = RewindDebugger->GetSelectedComponent();
 	if (SelectedObject.IsValid())
 	{
-		return OpenBlueprintAndAttachDebugger(RewindDebugger->GetAnalysisSession(), SelectedObject->ObjectId);
+		return OpenAnimBlueprintAndAttachDebugger(RewindDebugger->GetAnalysisSession(), SelectedObject->ObjectId);
 	}
 #endif
 	return false;
@@ -92,7 +92,7 @@ void FAnimInstanceMenu::Register()
 							{
 								IUnrealInsightsModule& UnrealInsightsModule = FModuleManager::LoadModuleChecked<IUnrealInsightsModule>("TraceInsights");
 								TSharedPtr<const TraceServices::IAnalysisSession> Session = UnrealInsightsModule.GetAnalysisSession();
-								OpenBlueprintAndAttachDebugger(Session.Get(), ObjectId);
+								OpenAnimBlueprintAndAttachDebugger(Session.Get(), ObjectId);
 							}));
 		}
 	}));
diff --git a/Engine/Plugins/Animation/GameplayInsights/Source/RewindDebugger/Private/BlueprintDoubleClickHandler.cpp b/Engine/Plugins/Animation/GameplayInsights/Source/RewindDebugger/Private/BlueprintDoubleClickHandler.cpp
new file mode 100644
--- /dev/null
+++ b/Engine/Plugins/Animation/GameplayInsights/Source/RewindDebugger/Private/BlueprintDoubleClickHandler.cpp
@@ -0,0 +1,64 @@
+﻿// Copyright Epic Games, Inc. All Rights Reserved.
+
+#include "BlueprintDoubleClickHandler.h"
+
+#include "IGameplayProvider.h"
+#include "ObjectTrace.h"
+#include "Engine/Blueprint.h"
+#include "Engine/BlueprintGeneratedClass.h"
+#include "Editor.h"
+#include "Subsystems/AssetEditorSubsystem.h"
+
+#if OBJECT_TRACE_ENABLED
+
+/** Open blueprint editor given an object id. If no editor found, it will use the properties editor. */
+static bool OpenBlueprintAndAttachDebugger(const TraceServices::IAnalysisSession* Session, uint64 ObjectId)
+{
+	TraceServices::FAnalysisSessionReadScope SessionReadScope(*Session);
+
+	if (const IGameplayProvider* GameplayProvider = Session->ReadProvider<IGameplayProvider>("GameplayProvider"))
+	{
+		if (const FObjectInfo* ObjectInfo = GameplayProvider->FindObjectInfo(ObjectId))
+		{
+			if (const FClassInfo* ClassInfo = GameplayProvider->FindClassInfo(ObjectInfo->ClassId))
+			{
+				if (const UBlueprintGeneratedClass* InstanceClass = TSoftObjectPtr<UBlueprintGeneratedClass>(FSoftObjectPath(ClassInfo->PathName)).LoadSynchronous())
+				{
+					if (UBlueprint* Blueprint = Cast<UBlueprint>(InstanceClass->ClassGeneratedBy))
+					{
+						GEditor->GetEditorSubsystem<UAssetEditorSubsystem>()->OpenEditorForAsset(Blueprint);
+						
+						if (UObject* SelectedInstance = FObjectTrace::GetObjectFromId(ObjectId))
+						{
+							Blueprint->SetObjectBeingDebugged(SelectedInstance);
+						}
+
+						return true;
+					}
+				}
+			}
+		}
+	}
+	return false;
+}
+
+#endif
+
+bool FBlueprintDoubleClickHandler::HandleDoubleClick(IRewindDebugger* RewindDebugger)
+{
+#if OBJECT_TRACE_ENABLED
+	const TSharedPtr<FDebugObjectInfo> SelectedObject = RewindDebugger->GetSelectedComponent();
+	if (SelectedObject.IsValid())
+	{
+		return OpenBlueprintAndAttachDebugger(RewindDebugger->GetAnalysisSession(), SelectedObject->ObjectId);
+	}
+#endif
+	
+	return false;
+}
+
+FName FBlueprintDoubleClickHandler::GetTargetTypeName() const
+{
+	static const FName ObjectTypeName = "Object";
+	return ObjectTypeName;
+}
\ No newline at end of file
diff --git a/Engine/Plugins/Animation/GameplayInsights/Source/RewindDebugger/Private/BlueprintDoubleClickHandler.h b/Engine/Plugins/Animation/GameplayInsights/Source/RewindDebugger/Private/BlueprintDoubleClickHandler.h
new file mode 100644
--- /dev/null
+++ b/Engine/Plugins/Animation/GameplayInsights/Source/RewindDebugger/Private/BlueprintDoubleClickHandler.h
@@ -0,0 +1,10 @@
+﻿// Copyright Epic Games, Inc. All Rights Reserved.
+#pragma once
+
+#include "IRewindDebuggerDoubleClickHandler.h"
+
+class FBlueprintDoubleClickHandler : public IRewindDebuggerDoubleClickHandler
+{
+	virtual FName GetTargetTypeName() const override;
+	virtual bool HandleDoubleClick(IRewindDebugger* RewindDebugger) override;
+};
\ No newline at end of file
diff --git a/Engine/Plugins/Animation/GameplayInsights/Source/RewindDebugger/Private/RewindDebuggerModule.cpp b/Engine/Plugins/Animation/GameplayInsights/Source/RewindDebugger/Private/RewindDebuggerModule.cpp
--- a/Engine/Plugins/Animation/GameplayInsights/Source/RewindDebugger/Private/RewindDebuggerModule.cpp
+++ b/Engine/Plugins/Animation/GameplayInsights/Source/RewindDebugger/Private/RewindDebuggerModule.cpp
@@ -17,6 +17,7 @@
 #include "WorkspaceMenuStructure.h"
 #include "WorkspaceMenuStructureModule.h"
 #include "AnimInstanceHelpers.h"
+#include "BlueprintDoubleClickHandler.h"
 #include "PropertyTraceMenu.h"
 #include "ToolMenus.h"
 #include "LevelEditor.h"
@@ -181,6 +182,7 @@ TSharedRef<SDockTab> FRewindDebuggerModule::SpawnRewindDebuggerTab(const FSpawnT
 }
 
 static FAnimInstanceDoubleClickHandler AnimInstanceDoubleClickHandler;
+static FBlueprintDoubleClickHandler BlueprintDoubleClickHandler;
 
 void FRewindDebuggerModule::StartupModule()
 {
@@ -224,6 +226,7 @@ void FRewindDebuggerModule::StartupModule()
 	RewindDebuggerCameraExtension.Initialize();
 	IModularFeatures::Get().RegisterModularFeature(IRewindDebuggerExtension::ModularFeatureName, &RewindDebuggerCameraExtension);
 	IModularFeatures::Get().RegisterModularFeature(IRewindDebuggerDoubleClickHandler::ModularFeatureName, &AnimInstanceDoubleClickHandler);
+	IModularFeatures::Get().RegisterModularFeature(IRewindDebuggerDoubleClickHandler::ModularFeatureName, &BlueprintDoubleClickHandler);
 
 	FPropertyTraceMenu::Register();
 	FAnimInstanceMenu::Register();
@@ -234,6 +237,7 @@ void FRewindDebuggerModule::ShutdownModule()
 {
 	IModularFeatures::Get().UnregisterModularFeature(IRewindDebuggerExtension::ModularFeatureName, &RewindDebuggerCameraExtension);
 	IModularFeatures::Get().UnregisterModularFeature(IRewindDebuggerDoubleClickHandler::ModularFeatureName, &AnimInstanceDoubleClickHandler);
+	IModularFeatures::Get().UnregisterModularFeature(IRewindDebuggerDoubleClickHandler::ModularFeatureName, &BlueprintDoubleClickHandler);
 
 	FRewindDebuggerCommands::Unregister();
 	FRewindDebuggerStyle::Shutdown();
