Animation Curve Runtime & Editor Improvements

diff --git a/Engine/Plugins/Animation/GameplayInsights/Source/GameplayInsights/Private/AnimationProvider.cpp b/Engine/Plugins/Animation/GameplayInsights/Source/GameplayInsights/Private/AnimationProvider.cpp
--- a/Engine/Plugins/Animation/GameplayInsights/Source/GameplayInsights/Private/AnimationProvider.cpp
+++ b/Engine/Plugins/Animation/GameplayInsights/Source/GameplayInsights/Private/AnimationProvider.cpp
@@ -142,6 +142,19 @@ void FAnimationProvider::EnumerateSkeletalMeshCurves(const FSkeletalMeshPoseMess
 	}
 }
 
+void FAnimationProvider::EnumeratePoseWatchCurves(const FPoseWatchMessage& InMessage, TFunctionRef<void(const FSkeletalMeshNamedCurve&)> Callback) const
+{
+	Session.ReadAccessCheck();
+
+	uint64 StartCurveIndex = InMessage.CurveStartIndex;
+	uint64 EndCurveIndex = InMessage.CurveStartIndex + InMessage.NumCurves;
+
+	for(uint64 CurveIndex = StartCurveIndex; CurveIndex < EndCurveIndex; ++CurveIndex)
+	{
+		Callback(SkeletalMeshCurves[CurveIndex]);
+	}
+}
+
 void FAnimationProvider::EnumerateTickRecordIds(uint64 InObjectId, TFunctionRef<void(uint64, int32)> Callback) const
 {
 	Session.ReadAccessCheck();
@@ -1315,11 +1328,79 @@ void FAnimationProvider::AppendPoseWatch(uint64 InAnimInstanceId, double InTime,
 	bHasAnyData = true;
 
 	TSharedPtr<TraceServices::TPointTimeline<FPoseWatchMessage>> Timeline;
-	uint32* IndexPtr = ObjectIdToPoseWatchTimelines.Find(InAnimInstanceId);
-	if (IndexPtr != nullptr)
+	uint32* PoseWatchIndexPtr = ObjectIdToPoseWatchTimelines.Find(InAnimInstanceId);
+	if (PoseWatchIndexPtr != nullptr)
+	{
+		check(PoseWatchTimelines.IsValidIndex(*PoseWatchIndexPtr));
+		Timeline = PoseWatchTimelines[*PoseWatchIndexPtr];
+	}
+	else
+	{
+		Timeline = MakeShared<TraceServices::TPointTimeline<FPoseWatchMessage>>(Session.GetLinearAllocator());
+		ObjectIdToPoseWatchTimelines.Add(InAnimInstanceId, PoseWatchTimelines.Num());
+		PoseWatchTimelines.Add(Timeline.ToSharedRef());
+	}
+
+	FTransform WorldTransform;
+	FMemory::Memcpy(&WorldTransform, &WorldTransformRaw[0], sizeof(FTransform));
+
+	FPoseWatchMessage Message;
+	Message.RecordingTime = InRecordingTime;
+	Message.ComponentId = 0;
+	Message.AnimInstanceId = InAnimInstanceId;
+	Message.PoseWatchId = PoseWatchId;
+	Message.BoneTransformsStartIndex = SkeletalMeshPoseTransforms.Num();
+	Message.CurveStartIndex = SkeletalMeshCurves.Num();
+	Message.NumBoneTransforms = BoneTransformsRaw.Num() / (sizeof(FTransform) / sizeof(float));
+	Message.WorldTransform = WorldTransform;
+	Message.RequiredBonesStartIndex = PoseWatchRequiredBones.Num();
+	Message.NumRequiredBones = RequiredBones.Num();
+	Message.NumCurves = 0;
+	Message.bIsEnabled = bIsEnabled;
+
+	Timeline->AppendEvent(InTime, Message);
+
+	const int CaptureTransformSize = sizeof(FTransform) / sizeof(float); // TODO
+	const int LocalTransformSize = sizeof(FTransform) / sizeof(float);
+	const int PoseTransformCount = BoneTransformsRaw.Num() / CaptureTransformSize;
+
+	// Dump PoseBoneTransformsRaw into SkeletalMeshPoseTransforms
+	if (CaptureTransformSize == LocalTransformSize)
 	{
-		check(PoseWatchTimelines.IsValidIndex(*IndexPtr));
-		Timeline = PoseWatchTimelines[*IndexPtr];
+		for (int i = 0; i < PoseTransformCount; i++)
+		{
+			FMemory::Memcpy(&SkeletalMeshPoseTransforms.PushBack(), &BoneTransformsRaw[CaptureTransformSize * i], sizeof(FTransform));
+		}
+	}
+	else
+	{
+		for (int i = 0; i < PoseTransformCount; i++)
+		{
+			SkeletalMeshPoseTransforms.PushBack() = ConvertTransform(CaptureTransformSize, &BoneTransformsRaw[CaptureTransformSize * i]);
+		}
+	}
+
+	// Dump RequiredBones into PoseWatchRequiredBones
+	for (const uint16 RequiredBone : RequiredBones)
+	{
+		PoseWatchRequiredBones.PushBack() = RequiredBone;
+	}
+
+	Session.UpdateDurationSeconds(InTime);
+}
+
+void FAnimationProvider::AppendPoseWatch(uint64 InComponentId, uint64 InAnimInstanceId, double InTime, double InRecordingTime, uint64 PoseWatchId, uint32 NameId, FColor Color, const TArrayView<const float>& BoneTransformsRaw, const TArrayView<const uint32>& CurveIds, const TArrayView<const float>& CurveValues, const TArrayView<const uint16>& RequiredBones, const TArrayView<const float>& WorldTransformRaw, const bool bIsEnabled)
+{
+	Session.WriteAccessCheck();
+
+	bHasAnyData = true;
+
+	TSharedPtr<TraceServices::TPointTimeline<FPoseWatchMessage>> Timeline;
+	uint32* PoseWatchIndexPtr = ObjectIdToPoseWatchTimelines.Find(InAnimInstanceId);
+	if (PoseWatchIndexPtr != nullptr)
+	{
+		check(PoseWatchTimelines.IsValidIndex(*PoseWatchIndexPtr));
+		Timeline = PoseWatchTimelines[*PoseWatchIndexPtr];
 	}
 	else
 	{
@@ -1328,18 +1409,39 @@ void FAnimationProvider::AppendPoseWatch(uint64 InAnimInstanceId, double InTime,
 		PoseWatchTimelines.Add(Timeline.ToSharedRef());
 	}
 
+	TSharedPtr<FSkeletalMeshTimelineStorage> TimelineStorage;
+	uint32* SkeletalMeshIndexPtr = ObjectIdToSkeletalMeshPoseTimelines.Find(InComponentId);
+	if(SkeletalMeshIndexPtr != nullptr)
+	{
+		TimelineStorage = SkeletalMeshPoseTimelineStorage[*SkeletalMeshIndexPtr];
+	}
+	else
+	{
+		TimelineStorage = MakeShared<FSkeletalMeshTimelineStorage>();
+		TimelineStorage->Timeline = MakeShared<TraceServices::TIntervalTimeline<FSkeletalMeshPoseMessage>>(Session.GetLinearAllocator());
+		ObjectIdToSkeletalMeshPoseTimelines.Add(InComponentId, SkeletalMeshPoseTimelineStorage.Num());
+		SkeletalMeshPoseTimelineStorage.Add(TimelineStorage.ToSharedRef());
+	}
+	
+	const int32 NumCurves = CurveIds.Num();
+
 	FTransform WorldTransform;
 	FMemory::Memcpy(&WorldTransform, &WorldTransformRaw[0], sizeof(FTransform));
 
 	FPoseWatchMessage Message;
 	Message.RecordingTime = InRecordingTime;
+	Message.ComponentId = InComponentId;
 	Message.AnimInstanceId = InAnimInstanceId;
 	Message.PoseWatchId = PoseWatchId;
+	Message.NameId = NameId;
+	Message.Color = Color;
 	Message.BoneTransformsStartIndex = SkeletalMeshPoseTransforms.Num();
+	Message.CurveStartIndex = SkeletalMeshCurves.Num();
 	Message.NumBoneTransforms = BoneTransformsRaw.Num() / (sizeof(FTransform) / sizeof(float));
 	Message.WorldTransform = WorldTransform;
 	Message.RequiredBonesStartIndex = PoseWatchRequiredBones.Num();
 	Message.NumRequiredBones = RequiredBones.Num();
+	Message.NumCurves = NumCurves;
 	Message.bIsEnabled = bIsEnabled;
 
 	Timeline->AppendEvent(InTime, Message);
@@ -1364,6 +1466,12 @@ void FAnimationProvider::AppendPoseWatch(uint64 InAnimInstanceId, double InTime,
 		}
 	}
 
+	for(int32 CurveIndex = 0; CurveIndex < NumCurves; ++CurveIndex)
+	{
+		SkeletalMeshCurves.PushBack() = { CurveIds[CurveIndex], CurveValues[CurveIndex] };
+		TimelineStorage->AllCurveIds.Add(CurveIds[CurveIndex]);
+	}
+	
 	// Dump RequiredBones into PoseWatchRequiredBones
 	for (const uint16 RequiredBone : RequiredBones)
 	{
diff --git a/Engine/Plugins/Animation/GameplayInsights/Source/GameplayInsights/Private/AnimationProvider.h b/Engine/Plugins/Animation/GameplayInsights/Source/GameplayInsights/Private/AnimationProvider.h
--- a/Engine/Plugins/Animation/GameplayInsights/Source/GameplayInsights/Private/AnimationProvider.h
+++ b/Engine/Plugins/Animation/GameplayInsights/Source/GameplayInsights/Private/AnimationProvider.h
@@ -29,6 +29,7 @@ public:
 	virtual void GetPoseWatchData(const FPoseWatchMessage& InMessage, TArray<FTransform>& BoneTransforms, TArray<uint16>& RequiredBones) const;
 	virtual void EnumerateSkeletalMeshCurveIds(uint64 InObjectId, TFunctionRef<void(uint32)> Callback) const override;
 	virtual void EnumerateSkeletalMeshCurves(const FSkeletalMeshPoseMessage& InMessage, TFunctionRef<void(const FSkeletalMeshNamedCurve&)> Callback) const override;
+	virtual void EnumeratePoseWatchCurves(const FPoseWatchMessage& InMessage, TFunctionRef<void(const FSkeletalMeshNamedCurve&)> Callback) const override;
 	virtual bool ReadTickRecordTimeline(uint64 InObjectId, TFunctionRef<void(const TickRecordTimeline&)> Callback) const override;
 	virtual void EnumerateTickRecordIds(uint64 InObjectId, TFunctionRef<void(uint64, int32)> Callback) const override;
 	virtual void EnumerateAnimGraphTimelines(TFunctionRef<void(uint64 ObjectId, const AnimGraphTimeline&)> Callback) const override;
@@ -111,6 +112,7 @@ public:
 
 	/** Append pose watch data */
 	void AppendPoseWatch(uint64 InAnimInstanceId, double InTime, double InRecordingTime, uint64 PoseWatchId, const TArrayView<const float>& BoneTransformsRaw, const TArrayView<const uint16>& RequiredBones, const TArrayView<const float>& WorldTransformRaw, const bool bIsEnabled);
+	void AppendPoseWatch(uint64 InComponentId, uint64 InAnimInstanceId, double InTime, double InRecordingTime, uint64 PoseWatchId, uint32 NameId, FColor Color, const TArrayView<const float>& BoneTransformsRaw, const TArrayView<const uint32>& CurveIds, const TArrayView<const float>& CurveValues, const TArrayView<const uint16>& RequiredBones, const TArrayView<const float>& WorldTransformRaw, const bool bIsEnabled);
 
 private:
 	/** Add anim node values helper */
diff --git a/Engine/Plugins/Animation/GameplayInsights/Source/GameplayInsights/Private/PoseWatchTrack.cpp b/Engine/Plugins/Animation/GameplayInsights/Source/GameplayInsights/Private/PoseWatchTrack.cpp
--- a/Engine/Plugins/Animation/GameplayInsights/Source/GameplayInsights/Private/PoseWatchTrack.cpp
+++ b/Engine/Plugins/Animation/GameplayInsights/Source/GameplayInsights/Private/PoseWatchTrack.cpp
@@ -1,9 +1,7 @@
 // Copyright Epic Games, Inc. All Rights Reserved.
 
-
 #if WITH_EDITOR
 
-
 #include "PoseWatchTrack.h"
 #include "SCurveTimelineView.h"
 #include "IRewindDebugger.h"
@@ -16,40 +14,78 @@
 #include "Engine/PoseWatchRenderData.h"
 #include "Engine/PoseWatch.h"
 #include "ObjectTrace.h"
+#include "SPoseWatchCurvesView.h"
 
 #define LOCTEXT_NAMESPACE "PoseWatchesTrack"
 
 namespace RewindDebugger
 {
 
-FPoseWatchTrack::FPoseWatchTrack(uint64 InObjectId, const FPoseWatchTrack::FPoseWatchTrackId& InPoseWatchTrackId)
-	: ObjectId(InObjectId)
+FPoseWatchCurveTrack::FPoseWatchCurveTrack(uint64 InObjectId, uint32 InCurveId, uint64 InPoseWatchTrackId)
+	: FAnimCurveTrack(InObjectId, InCurveId)
 	, PostWatchTrackId(InPoseWatchTrackId)
-	, PoseWatchOwner(nullptr)
 {
-	EnabledSegments = MakeShared<SSegmentedTimelineView::FSegmentData>();
-	Icon = UPoseWatchPoseElement::StaticGetIcon();
+}
+
+void FPoseWatchCurveTrack::UpdateCurvePointsInternal()
+{
+	IRewindDebugger* RewindDebugger = IRewindDebugger::Instance();
+	const TraceServices::IAnalysisSession* AnalysisSession = RewindDebugger->GetAnalysisSession();
+	const FAnimationProvider* AnimationProvider = AnalysisSession->ReadProvider<FAnimationProvider>(FAnimationProvider::ProviderName);
 
-#if OBJECT_TRACE_ENABLED
-	if (UObject* ObjectInstance = FObjectTrace::GetObjectFromId(ObjectId))
+	// convert time range to from rewind debugger times to profiler times
+	TRange<double> TraceTimeRange = RewindDebugger->GetCurrentTraceRange();
+	double StartTime = TraceTimeRange.GetLowerBoundValue();
+	double EndTime = TraceTimeRange.GetUpperBoundValue();
+	
+	auto& CurvePoints = CurveData->Points;
+	CurvePoints.SetNum(0,false);
+
+	TraceServices::FAnalysisSessionReadScope SessionReadScope(*AnalysisSession);
+
+	AnimationProvider->ReadPoseWatchTimeline(ObjectId, [this, StartTime, EndTime, AnimationProvider, &CurvePoints](const FAnimationProvider::PoseWatchTimeline& InPoseWatchTimeline)
 	{
-		if (UAnimInstance* AnimInstance = Cast<UAnimInstance>(ObjectInstance))
+		InPoseWatchTimeline.EnumerateEvents(StartTime, EndTime, [this, StartTime, EndTime, AnimationProvider, &CurvePoints](double InStartTime, double InEndTime, uint32 InDepth, const FPoseWatchMessage& InMessage)
 		{
-			if (UAnimBlueprintGeneratedClass* InstanceClass = Cast<UAnimBlueprintGeneratedClass>(AnimInstance->GetClass()))
+			if (InMessage.PoseWatchId == PostWatchTrackId)
 			{
-				FAnimBlueprintDebugData& DebugData = InstanceClass->GetAnimBlueprintDebugData();
-				for (const FAnimNodePoseWatch& PoseWatch : DebugData.AnimNodePoseWatch)
+				if (InEndTime > StartTime && InStartTime < EndTime)
 				{
-					if (PoseWatch.NodeID == PostWatchTrackId.NameId)
+					if(InMessage.bIsEnabled)
 					{
-						PoseWatchOwner = PoseWatch.PoseWatchPoseElement;
-						break;
+						double Time = InMessage.RecordingTime;
+						AnimationProvider->EnumeratePoseWatchCurves(InMessage, [this, Time, &CurvePoints](const FSkeletalMeshNamedCurve& InCurve)
+						{
+							if (InCurve.Id == CurveId)
+							{
+								CurvePoints.Add({Time,InCurve.Value});
+							}
+						});
 					}
 				}
 			}
-		}
-	}
-#endif
+			return TraceServices::EEventEnumerate::Continue;
+		});
+	});
+}
+
+TSharedPtr<SWidget> FPoseWatchCurveTrack::GetDetailsViewInternal()
+{
+	IRewindDebugger* RewindDebugger = IRewindDebugger::Instance();
+	TSharedPtr<SPoseWatchCurvesView> PoseWatchCurvesView = SNew(SPoseWatchCurvesView, ObjectId, RewindDebugger->CurrentTraceTime(), *RewindDebugger->GetAnalysisSession())
+				.CurrentTime_Lambda([RewindDebugger]{ return RewindDebugger->CurrentTraceTime(); });
+	PoseWatchCurvesView->SetPoseWatchCurveFilter(PostWatchTrackId, CurveId);
+	return PoseWatchCurvesView;
+}
+
+FPoseWatchTrack::FPoseWatchTrack(uint64 InObjectId, uint64 InPoseWatchTrackId, FColor InColor, uint32 InNameId)
+	: ObjectId(InObjectId)
+	, PoseWatchTrackId(InPoseWatchTrackId)
+	, Color(InColor)
+	, NameId(InNameId)
+{
+	EnabledSegments = MakeShared<SSegmentedTimelineView::FSegmentData>();
+	Icon = UPoseWatchPoseElement::StaticGetIcon();
 }
 
 FText FPoseWatchTrack::GetDisplayNameInternal() const
@@ -74,6 +110,8 @@ bool FPoseWatchTrack::UpdateInternal()
 	const FGameplayProvider* GameplayProvider = AnalysisSession->ReadProvider<FGameplayProvider>(FGameplayProvider::ProviderName);
 	const FAnimationProvider* AnimationProvider = AnalysisSession->ReadProvider<FAnimationProvider>(FAnimationProvider::ProviderName);
 
+	bool bChanged = false;
+
 	if(GameplayProvider && AnimationProvider && EnabledSegments.IsValid())
 	{
 		EnabledSegments->Segments.SetNum(0, false);
@@ -93,17 +131,27 @@ bool FPoseWatchTrack::UpdateInternal()
 
 		TraceServices::FAnalysisSessionReadScope SessionReadScope(*AnalysisSession);
 
-		AnimationProvider->ReadPoseWatchTimeline(ObjectId, [this, StartTime, EndTime, &RecordingTimes](const FAnimationProvider::PoseWatchTimeline& InPoseWatchTimeline)
+		TArray<uint64> UniqueCurveIds;
+
+		AnimationProvider->ReadPoseWatchTimeline(ObjectId, [this, StartTime, EndTime, &RecordingTimes, &UniqueCurveIds, AnimationProvider](const FAnimationProvider::PoseWatchTimeline& InPoseWatchTimeline)
 		{
-			InPoseWatchTimeline.EnumerateEvents(StartTime, EndTime, [this, StartTime, EndTime, &RecordingTimes](double InStartTime, double InEndTime, uint32 InDepth, const FPoseWatchMessage& InMessage)
+			InPoseWatchTimeline.EnumerateEvents(StartTime, EndTime, [this, StartTime, EndTime, &RecordingTimes, &UniqueCurveIds, AnimationProvider](double InStartTime, double InEndTime, uint32 InDepth, const FPoseWatchMessage& InMessage)
 			{
-				if (InMessage.PoseWatchId == PostWatchTrackId.NameId)
+				if (InMessage.PoseWatchId == PoseWatchTrackId)
 				{
+					Color = InMessage.Color;
+					NameId = InMessage.NameId;
+
 					FPoseWatchEnabledTime EnabledTime;
 					EnabledTime.RecordingTime = InMessage.RecordingTime;
 					EnabledTime.bIsEnabled = InMessage.bIsEnabled;
 
 					RecordingTimes.Add(EnabledTime);
+
+					AnimationProvider->EnumeratePoseWatchCurves(InMessage, [this, &UniqueCurveIds](const FSkeletalMeshNamedCurve& InCurve)
+					{
+						UniqueCurveIds.AddUnique(InCurve.Id);
+					});
 				}
 				return TraceServices::EEventEnumerate::Continue;
 			});
@@ -139,13 +187,32 @@ bool FPoseWatchTrack::UpdateInternal()
 				EnabledSegments->Segments.Add(TRange<double>(RangeStart, RecordingTimes.Last().RecordingTime));
 			}
 		}
+
+		UniqueCurveIds.StableSort();
+		const int32 TrackCount = UniqueCurveIds.Num();
+
+		if (Children.Num() != TrackCount)
+		{
+			bChanged = true;
+		}
+
+		Children.SetNum(UniqueCurveIds.Num());
+		for(int i = 0; i < TrackCount; i++)
+		{
+			if (!Children[i].IsValid() || !(Children[i].Get()->GetCurveId() == UniqueCurveIds[i]))
+			{
+				Children[i] = MakeShared<FPoseWatchCurveTrack>(ObjectId, UniqueCurveIds[i], PoseWatchTrackId);
+				bChanged = true;
+			}
+
+			bChanged = bChanged || Children[i]->Update();
+		}
 	}
 
-	bool bChanged = false;
-	if (PoseWatchOwner && PoseWatchOwner->GetParent())
+	if(const TCHAR* FoundName = AnimationProvider->GetName(NameId))
 	{
-		bChanged = !TrackName.EqualTo(PoseWatchOwner->GetParent()->GetLabel());
-		TrackName = PoseWatchOwner->GetParent()->GetLabel();
+		bChanged = bChanged || TrackName.ToString().Compare(FoundName) != 0;
+		TrackName = FText::FromString(FoundName);
 	}
 
 	return bChanged;
@@ -153,7 +220,7 @@ bool FPoseWatchTrack::UpdateInternal()
 
 TSharedPtr<SWidget> FPoseWatchTrack::GetTimelineViewInternal()
 {
-	const auto GetPoseWatchColorLambda = [this]() -> FLinearColor { return PoseWatchOwner ? FLinearColor(PoseWatchOwner->GetColor()) : FLinearColor::White; };
+	const auto GetPoseWatchColorLambda = [this]() -> FLinearColor { return FLinearColor(Color); };
 
 	const auto TimelineView = SNew(SSegmentedTimelineView)
 		.FillColor_Lambda(GetPoseWatchColorLambda)
@@ -163,6 +230,14 @@ TSharedPtr<SWidget> FPoseWatchTrack::GetTimelineViewInternal()
 	return TimelineView;
 }
 
+void FPoseWatchTrack::IterateSubTracksInternal(TFunction<void(TSharedPtr<FRewindDebuggerTrack> SubTrack)> IteratorFunction)
+{
+	for(TSharedPtr<FPoseWatchCurveTrack>& Track : Children)
+	{
+		IteratorFunction(Track);
+	}
+}
+
 FName FPoseWatchesTrackCreator::GetTargetTypeNameInternal() const
 {
 	static const FName AnimInstanceName("AnimInstance");
@@ -185,22 +260,8 @@ FPoseWatchesTrack::FPoseWatchesTrack(uint64 InObjectId)
 	: ObjectId(InObjectId)
 {
 	Icon = UPoseWatchPoseElement::StaticGetIcon();
-
-#if OBJECT_TRACE_ENABLED
-	if (UObject* ObjectInstance = FObjectTrace::GetObjectFromId(ObjectId))
-	{
-		if (UAnimInstance* AnimInstance = Cast<UAnimInstance>(ObjectInstance))
-		{
-			if (UAnimBlueprintGeneratedClass* InstanceClass = Cast<UAnimBlueprintGeneratedClass>(AnimInstance->GetClass()))
-			{
-				AnimBPGenClass = InstanceClass;
-			}
-		}
-	}
-#endif
 }
 
-
 bool FPoseWatchesTrack::UpdateInternal()
 {
 	IRewindDebugger* RewindDebugger = IRewindDebugger::Instance();
@@ -217,42 +278,48 @@ bool FPoseWatchesTrack::UpdateInternal()
 	
 	if(GameplayProvider && AnimationProvider)
 	{
-		TArray<FPoseWatchTrack::FPoseWatchTrackId> UniqueTrackIds;
+		struct FUniquePoseWatch
+		{
+			uint64 PoseWatchId;
+			FColor Color;
+			uint32 NameId;
+
+			bool operator==(const FUniquePoseWatch& InOther) const
+			{
+				return PoseWatchId == InOther.PoseWatchId;
+			}
+
+			bool operator<(const FUniquePoseWatch& InOther) const
+			{
+				return PoseWatchId < InOther.PoseWatchId;
+			}
+		};
 		
+		TArray<FUniquePoseWatch> UniquePoseWatches;
+
 		TraceServices::FAnalysisSessionReadScope SessionReadScope(*AnalysisSession);
 
-		AnimationProvider->ReadPoseWatchTimeline(ObjectId, [this, StartTime, EndTime, &UniqueTrackIds](const FAnimationProvider::PoseWatchTimeline& InPoseWatchTimeline)
+		AnimationProvider->ReadPoseWatchTimeline(ObjectId, [this, StartTime, EndTime, &UniquePoseWatches](const FAnimationProvider::PoseWatchTimeline& InPoseWatchTimeline)
 		{
-			InPoseWatchTimeline.EnumerateEvents(StartTime, EndTime, [this, &UniqueTrackIds](double InStartTime, double InEndTime, uint32 InDepth, const FPoseWatchMessage& InMessage)
+			InPoseWatchTimeline.EnumerateEvents(StartTime, EndTime, [this, &UniquePoseWatches](double InStartTime, double InEndTime, uint32 InDepth, const FPoseWatchMessage& InMessage)
 			{
-				// Only add the pose watch track if the pose watch still exists to attach to
-				if (AnimBPGenClass)
-				{
-					for (const FAnimNodePoseWatch& PoseWatch : AnimBPGenClass->GetAnimBlueprintDebugData().AnimNodePoseWatch)
-					{
-						if (PoseWatch.NodeID == InMessage.PoseWatchId)
-						{
-							UniqueTrackIds.AddUnique({ InMessage.PoseWatchId });
-							break;
-						}
-					}
-				}
+				UniquePoseWatches.AddUnique({ InMessage.PoseWatchId, FColor(InMessage.Color), InMessage.NameId });
 				return TraceServices::EEventEnumerate::Continue;
 			});
 		});
 		
-		UniqueTrackIds.StableSort();
-		const int32 TrackCount = UniqueTrackIds.Num();
+		UniquePoseWatches.StableSort();
+		const int32 TrackCount = UniquePoseWatches.Num();
 
 		if (Children.Num() != TrackCount)
 			bChanged = true;
 
-		Children.SetNum(UniqueTrackIds.Num());
+		Children.SetNum(UniquePoseWatches.Num());
 		for(int i = 0; i < TrackCount; i++)
 		{
-			if (!Children[i].IsValid() || !(Children[i].Get()->GetPoseWatchTrackId() == UniqueTrackIds[i]))
+			if (!Children[i].IsValid() || !(Children[i].Get()->GetPoseWatchTrackId() == UniquePoseWatches[i].PoseWatchId))
 			{
-				Children[i] = MakeShared<FPoseWatchTrack>(ObjectId, UniqueTrackIds[i]);
+				Children[i] = MakeShared<FPoseWatchTrack>(ObjectId, UniquePoseWatches[i].PoseWatchId, UniquePoseWatches[i].Color, UniquePoseWatches[i].NameId);
 				bChanged = true;
 			}
 
diff --git a/Engine/Plugins/Animation/GameplayInsights/Source/GameplayInsights/Private/PoseWatchTrack.h b/Engine/Plugins/Animation/GameplayInsights/Source/GameplayInsights/Private/PoseWatchTrack.h
--- a/Engine/Plugins/Animation/GameplayInsights/Source/GameplayInsights/Private/PoseWatchTrack.h
+++ b/Engine/Plugins/Animation/GameplayInsights/Source/GameplayInsights/Private/PoseWatchTrack.h
@@ -8,6 +8,7 @@
 #include "RewindDebuggerTrack.h"
 #include "Textures/SlateIcon.h"
 #include "SSegmentedTimelineView.h"
+#include "AnimCurvesTrack.h"
 
 class UPoseWatchPoseElement;
 class UAnimBlueprintGeneratedClass;
@@ -15,26 +16,27 @@ class UAnimBlueprintGeneratedClass;
 namespace RewindDebugger
 {
 
-class FPoseWatchTrack : public FRewindDebuggerTrack
+// Sub Track for a single curve in a pose watch
+class FPoseWatchCurveTrack : public FAnimCurveTrack
 {
 public:
-	struct FPoseWatchTrackId
-	{
-		uint64 NameId;
+	FPoseWatchCurveTrack(uint64 InObjectId, uint32 InCurveId, uint64 InPoseWatchTrackId);
 
-		bool operator == (const FPoseWatchTrackId& Other) const
-		{
-			return NameId == Other.NameId;
-		}
+	uint64 GetPoseWatchTrackId() const { return PostWatchTrackId; }
 
-		bool operator <(const FPoseWatchTrack::FPoseWatchTrackId& Other) const
-		{
-			return NameId < Other.NameId; 
-		}		
-	};
+private:
+	virtual void UpdateCurvePointsInternal() override;
+	virtual TSharedPtr<SWidget> GetDetailsViewInternal() override;
+	virtual FName GetNameInternal() const override { return "PoseWatchAnimationCurve"; }
+	
+	uint64 PostWatchTrackId;
+};
 
-	FPoseWatchTrack(uint64 InObjectId, const FPoseWatchTrackId& TrackId);
-	const FPoseWatchTrackId& GetPoseWatchTrackId() const { return PostWatchTrackId; }
+class FPoseWatchTrack : public FRewindDebuggerTrack
+{
+public:
+	FPoseWatchTrack(uint64 InObjectId, const uint64 TrackId, FColor InColor, uint32 InNameId);
+	uint64 GetPoseWatchTrackId() const { return PoseWatchTrackId; }
 
 private:
 	virtual bool UpdateInternal() override;
@@ -43,16 +45,20 @@ private:
 	virtual FSlateIcon GetIconInternal() override { return Icon; }
 	virtual FText GetDisplayNameInternal() const override;
 	virtual uint64 GetObjectIdInternal() const override { return ObjectId; }
-	
+	virtual void IterateSubTracksInternal(TFunction<void(TSharedPtr<FRewindDebuggerTrack> SubTrack)> IteratorFunction) override;
+		
 	TSharedPtr<SSegmentedTimelineView::FSegmentData> GetSegmentData() const;
 	
 	uint64 ObjectId;
-	FPoseWatchTrackId PostWatchTrackId;
+	uint64 PoseWatchTrackId;
 	FSlateIcon Icon;
 	FText TrackName;
+	FColor Color;
+	uint32 NameId;
 
-	const UPoseWatchPoseElement* PoseWatchOwner;
 	TSharedPtr<SSegmentedTimelineView::FSegmentData> EnabledSegments;
+
+	TArray<TSharedPtr<FPoseWatchCurveTrack>> Children;
 };
 
 
@@ -74,8 +80,6 @@ private:
 	FSlateIcon Icon;
 	uint64 ObjectId;
 
-	class UAnimBlueprintGeneratedClass* AnimBPGenClass;
-
 	TArray<TSharedPtr<FPoseWatchTrack>> Children;
 };
 
diff --git a/Engine/Plugins/Animation/GameplayInsights/Source/GameplayInsights/Public/IAnimationProvider.h b/Engine/Plugins/Animation/GameplayInsights/Source/GameplayInsights/Public/IAnimationProvider.h
--- a/Engine/Plugins/Animation/GameplayInsights/Source/GameplayInsights/Public/IAnimationProvider.h
+++ b/Engine/Plugins/Animation/GameplayInsights/Source/GameplayInsights/Public/IAnimationProvider.h
@@ -38,14 +38,19 @@ struct FSkeletalMeshPoseMessage
 
 struct FPoseWatchMessage
 {
+	FTransform WorldTransform;
+	FColor Color = FColor::Black;
 	double RecordingTime = 0.0;
+	uint64 ComponentId = 0;
 	uint64 AnimInstanceId = 0;
 	uint64 PoseWatchId = 0;
 	uint64 BoneTransformsStartIndex = 0;
-	uint16 NumBoneTransforms = 0;
-	FTransform WorldTransform;
+	uint64 CurveStartIndex = 0;
 	uint64 RequiredBonesStartIndex = 0;
+	uint32 NameId = 0;
+	uint16 NumBoneTransforms = 0;
 	uint16 NumRequiredBones = 0;
+	uint16 NumCurves = 0;
 	bool bIsEnabled = false;
 };
 
@@ -291,6 +296,7 @@ public:
 	virtual void EnumerateSkeletalMeshCurveIds(uint64 InObjectId, TFunctionRef<void(uint32)> Callback) const = 0;
 	virtual void GetPoseWatchData(const FPoseWatchMessage& InMessage, TArray<FTransform>& BoneTransforms, TArray<uint16>& RequiredBones) const = 0;
 	virtual void EnumerateSkeletalMeshCurves(const FSkeletalMeshPoseMessage& InMessage, TFunctionRef<void(const FSkeletalMeshNamedCurve&)> Callback) const = 0;
+	virtual void EnumeratePoseWatchCurves(const FPoseWatchMessage& InMessage, TFunctionRef<void(const FSkeletalMeshNamedCurve&)> Callback) const = 0;
 	virtual bool ReadTickRecordTimeline(uint64 InObjectId, TFunctionRef<void(const TickRecordTimeline&)> Callback) const = 0;
 	virtual void EnumerateTickRecordIds(uint64 InObjectId, TFunctionRef<void(uint64, int32)> Callback) const = 0;
 	virtual void EnumerateAnimGraphTimelines(TFunctionRef<void(uint64 ObjectId, const AnimGraphTimeline&)> Callback) const =0;
diff --git a/Engine/Plugins/Animation/IKRig/Source/IKRig/Private/AnimNodes/AnimNode_RetargetPoseFromMesh.cpp b/Engine/Plugins/Animation/IKRig/Source/IKRig/Private/AnimNodes/AnimNode_RetargetPoseFromMesh.cpp
--- a/Engine/Plugins/Animation/IKRig/Source/IKRig/Private/AnimNodes/AnimNode_RetargetPoseFromMesh.cpp
+++ b/Engine/Plugins/Animation/IKRig/Source/IKRig/Private/AnimNodes/AnimNode_RetargetPoseFromMesh.cpp
@@ -1,6 +1,8 @@
 // Copyright Epic Games, Inc. All Rights Reserved.
 
 #include "AnimNodes/AnimNode_RetargetPoseFromMesh.h"
+
+#include "Animation/AnimCurveUtils.h"
 #include "Animation/AnimInstanceProxy.h"
 #include "Components/SkeletalMeshComponent.h"
 #include "Engine/SkeletalMesh.h"
@@ -131,14 +133,7 @@ void FAnimNode_RetargetPoseFromMesh::Evaluate_AnyThread(FPoseContext& Output)
 	// copy curves over
 	if (bCopyCurves)
 	{
-		for (const TPair<FName, float>& SourceCurve : SourceCurveValues)
-		{
-			if (const SmartName::UID_Type* UID = CurveNameToUIDMap.Find(SourceCurve.Key))
-			{
-				// set source value to output curve
-				Output.Curve.Set(*UID, SourceCurve.Value);
-			}
-		}
+		Output.Curve.CopyFrom(SourceCurves);
 	}
 }
 
@@ -163,13 +158,18 @@ void FAnimNode_RetargetPoseFromMesh::PreUpdate(const UAnimInstance* InAnimInstan
 
 		if(bCopyCurves)
 		{
-			SourceCurveValues.Reset();
 			if (const UAnimInstance* SourceAnimInstance = SourceMeshComponent->GetAnimInstance())
 			{
-				// attribute curve contains all list	
-				SourceCurveValues.Append(SourceAnimInstance->GetAnimationCurveList(EAnimCurveType::AttributeCurve));
+				// Potential optimization/tradeoff: If we stored the curve results on the mesh component in non-editor scenarios, this would be
+				// much faster (but take more memory). As it is, we need to translate the map stored on the anim instance.
+				const TMap<FName, float>& AnimCurveList = SourceAnimInstance->GetAnimationCurveList(EAnimCurveType::AttributeCurve);
+				UE::Anim::FCurveUtils::BuildUnsorted(SourceCurves, AnimCurveList);
 			}
-
+			else
+			{
+				SourceCurves.Empty();
+			}
+			
 			UpdateSpeedValuesFromCurves();
 		}
 	}
@@ -194,12 +194,12 @@ void FAnimNode_RetargetPoseFromMesh::UpdateSpeedValuesFromCurves()
 		}
 
 		// value will be negative if the curve was not found (retargeter ignores negative speeds)
-		float CurveValue = -1.0f;
-		if (SourceCurveValues.Contains(CurveName))
+		bool bCurveValid = false;
+		float CurveValue = SourceCurves.Get(CurveName, bCurveValid, -1.0f);
+		if (bCurveValid)
 		{
-			CurveValue = SourceCurveValues[CurveName];
+			SpeedValuesFromCurves.Add(CurveName, CurveValue);
 		}
-		SpeedValuesFromCurves.Add(CurveName, CurveValue);
 	}
 }
 
@@ -263,27 +263,6 @@ bool FAnimNode_RetargetPoseFromMesh::EnsureProcessorIsInitialized(const TObjectP
 		// initialize retarget processor with source and target skeletal meshes
 		// (asset is passed in as outer UObject for new UIKRigProcessor) 
 		Processor->Initialize(SourceMesh,	TargetMesh,IKRetargeterAsset);
-
-		// create a map of curve names to IDs that are present on the target skeleton
-		if (bCopyCurves)
-		{
-			const FSmartNameMapping* SourceContainer = SourceSkeleton->GetSmartNameContainer(USkeleton::AnimCurveMappingName);
-			const FSmartNameMapping* TargetContainer = TargetSkeleton->GetSmartNameContainer(USkeleton::AnimCurveMappingName);
-
-			TArray<FName> SourceCurveNames;
-			SourceContainer->FillNameArray(SourceCurveNames);
-			CurveNameToUIDMap.Reset();
-			for (int32 Index = 0; Index < SourceCurveNames.Num(); ++Index)
-			{
-				SmartName::UID_Type UID = TargetContainer->FindUID(SourceCurveNames[Index]);
-				if (UID != SmartName::MaxUID)
-				{
-					// has a valid UID, add to the list
-					SmartName::UID_Type& Value = CurveNameToUIDMap.Add(SourceCurveNames[Index]);
-					Value = UID;
-				}
-			}
-		}
 	}
 
 	return Processor->IsInitialized();
diff --git a/Engine/Plugins/Animation/IKRig/Source/IKRig/Public/AnimNodes/AnimNode_RetargetPoseFromMesh.h b/Engine/Plugins/Animation/IKRig/Source/IKRig/Public/AnimNodes/AnimNode_RetargetPoseFromMesh.h
--- a/Engine/Plugins/Animation/IKRig/Source/IKRig/Public/AnimNodes/AnimNode_RetargetPoseFromMesh.h
+++ b/Engine/Plugins/Animation/IKRig/Source/IKRig/Public/AnimNodes/AnimNode_RetargetPoseFromMesh.h
@@ -72,10 +72,8 @@ private:
 
 	// mapping from required bones to actual bones within the target skeleton
 	TArray< TPair<int32, int32> > RequiredToTargetBoneMapping;
-	// mapping of Skeleton curve names to the UID (for copying curves from the source mesh component)
-	TMap<FName, SmartName::UID_Type> CurveNameToUIDMap;
 	// cached curves, copied on the game thread
-	TMap<FName, float> SourceCurveValues;
+	FBlendedHeapCurve SourceCurves;
 
 	/** update map of curve values containing speeds used for IK planting */
 	void UpdateSpeedValuesFromCurves();
diff --git a/Engine/Plugins/Animation/IKRig/Source/IKRigEditor/Private/RetargetEditor/IKRetargetBatchOperation.cpp b/Engine/Plugins/Animation/IKRig/Source/IKRigEditor/Private/RetargetEditor/IKRetargetBatchOperation.cpp
--- a/Engine/Plugins/Animation/IKRig/Source/IKRigEditor/Private/RetargetEditor/IKRetargetBatchOperation.cpp
+++ b/Engine/Plugins/Animation/IKRig/Source/IKRigEditor/Private/RetargetEditor/IKRetargetBatchOperation.cpp
@@ -282,19 +282,6 @@ void UIKRetargetBatchOperation::ConvertAnimation(
 	// get names of the curves the retargeter is looking for
 	TArray<FName> SpeedCurveNames;
 	Context.IKRetargetAsset->GetSpeedCurveNames(SpeedCurveNames);
-	// get all curve names in the source skeleton
-	const FSmartNameMapping* SourceContainer = Context.SourceMesh->GetSkeleton()->GetSmartNameContainer(USkeleton::AnimCurveMappingName);
-	TArray<FName> SourceCurveNames;
-	SourceContainer->FillNameArray(SourceCurveNames);
-	// get list of all the source curve IDs that hold speed values the retargeter is looking for
-	TMap<FName, SmartName::UID_Type> SourceSpeedCurveIDs;
-	for (int32 CurveIndex = 0; CurveIndex < SourceCurveNames.Num(); ++CurveIndex)
-	{
-		if (SpeedCurveNames.Contains(SourceCurveNames[CurveIndex]))
-		{
-			SourceSpeedCurveIDs.Add(SourceCurveNames[CurveIndex],SourceContainer->FindUID(SourceCurveNames[CurveIndex]));
-		}
-	}
 	
 	// for each pair of source / target animation sequences
 	for (TPair<UAnimationAsset*, UAnimationAsset*>& Pair : DuplicatedAnimAssets)
@@ -363,9 +350,9 @@ void UIKRetargetBatchOperation::ConvertAnimation(
 			
 			// get the curve values from the source sequence (for speed-based IK planting)
 			TMap<FName, float> SpeedCurveValues;
-			for (TPair<FName, SmartName::UID_Type> CurveNameAndID : SourceSpeedCurveIDs)
+			for (const FName& SpeedCurveName : SpeedCurveNames)
 			{
-				SpeedCurveValues.Add(CurveNameAndID.Key, SourceSequence->EvaluateCurveData(CurveNameAndID.Value, TimeAtCurrentFrame));
+				SpeedCurveValues.Add(SpeedCurveName, SourceSequence->EvaluateCurveData(SpeedCurveName, TimeAtCurrentFrame));
 			}
 
 			// run the retargeter
diff --git a/Engine/Plugins/Animation/IKRig/Source/IKRigEditor/Private/RetargetEditor/IKRetargetPoseExporter.cpp b/Engine/Plugins/Animation/IKRig/Source/IKRigEditor/Private/RetargetEditor/IKRetargetPoseExporter.cpp
--- a/Engine/Plugins/Animation/IKRig/Source/IKRigEditor/Private/RetargetEditor/IKRetargetPoseExporter.cpp
+++ b/Engine/Plugins/Animation/IKRig/Source/IKRigEditor/Private/RetargetEditor/IKRetargetPoseExporter.cpp
@@ -143,10 +143,9 @@ void FIKRetargetPoseExporter::OnRetargetPoseSelected(const FAssetData& SelectedA
 	}
 
 	// store pose names in list used by combobox
-	const TArray<FSmartName> AllPoseNames = PoseAsset->GetPoseNames();
-	for (const FSmartName& PoseName : AllPoseNames)
+	for (const FName& PoseName : PoseAsset->GetPoseFNames())
 	{
-		PosesInSelectedAsset.Add(MakeShared<FName>(PoseName.DisplayName));
+		PosesInSelectedAsset.Add(MakeShared<FName>(PoseName));
 	}
 
 	// set the selected pose to the first pose in the asset (or null if empty)
@@ -626,33 +625,20 @@ void FIKRetargetPoseExporter::HandleExportPoseAsset()
 
 	// fill new pose asset with current pose
 	NewPoseAsset->SetSkeleton(const_cast<USkeleton*>(Mesh->GetSkeleton()));
-	FSmartName NewPoseName;
-	bool bSuccess = NewPoseAsset->AddOrUpdatePoseWithUniqueName(MeshComponent, &NewPoseName);
+	const FName NewPoseName = NewPoseAsset->AddPoseWithUniqueName(MeshComponent);
 
-	if (bSuccess)
-	{
-		// mark asset dirty 
-		FAssetRegistryModule::AssetCreated(NewPoseAsset);
-		NewPoseAsset->MarkPackageDirty();
+	// mark asset dirty 
+	FAssetRegistryModule::AssetCreated(NewPoseAsset);
+	NewPoseAsset->MarkPackageDirty();
+
+	// show in content browser
+	TArray<UObject*> ObjectsToSync;
+	ObjectsToSync.Add(NewPoseAsset);
+	GEditor->SyncBrowserToObjects(ObjectsToSync);
 	
-		// show in content browser
-		TArray<UObject*> ObjectsToSync;
-		ObjectsToSync.Add(NewPoseAsset);
-		GEditor->SyncBrowserToObjects(ObjectsToSync);
-		
-		// notify user of newly created asset
-		const FText Message = FText::Format(LOCTEXT("IKRigExportPoseSuccess", "Exported Pose Asset: {0} with pose, '{1}'."), FText::FromString(AssetName), FText::FromName(NewPoseName.DisplayName));
-		NotifyUser(Message, SNotificationItem::CS_Success);
-	}
-	else
-	{
-		// remove the empty asset
-		NewPoseAsset->ConditionalBeginDestroy();
-		
-		// notify user that asset failed to export
-		const FText Message = FText::Format(LOCTEXT("IKRigExportPoseFailedAfter", "Failed to exported pose asset, {0}."), FText::FromString(AssetName));
-		NotifyUser(Message, SNotificationItem::CS_Fail);
-	}
+	// notify user of newly created asset
+	const FText Message = FText::Format(LOCTEXT("IKRigExportPoseSuccess", "Exported Pose Asset: {0} with pose, '{1}'."), FText::FromString(AssetName), FText::FromName(NewPoseName));
+	NotifyUser(Message, SNotificationItem::CS_Success);
 }
 
 void FIKRetargetPoseExporter::NotifyUser(const FText& Message, SNotificationItem::ECompletionState NotificationType)
diff --git a/Engine/Plugins/Animation/IKRig/Source/IKRigEditor/Public/RetargetEditor/IKRetargetBatchOperation.h b/Engine/Plugins/Animation/IKRig/Source/IKRigEditor/Public/RetargetEditor/IKRetargetBatchOperation.h
--- a/Engine/Plugins/Animation/IKRig/Source/IKRigEditor/Public/RetargetEditor/IKRetargetBatchOperation.h
+++ b/Engine/Plugins/Animation/IKRig/Source/IKRigEditor/Public/RetargetEditor/IKRetargetBatchOperation.h
@@ -9,6 +9,9 @@
 
 class UIKRetargeter;
 struct FScopedSlowTask;
+class USkeletalMesh;
+class UAnimationAsset;
+class UAnimBlueprint;
 
 //** Data needed to run a batch "duplicate and retarget" operation on a set of animation assets */
 struct FIKRetargetBatchOperationContext
diff --git a/Engine/Plugins/Animation/MLDeformer/MLDeformerFramework/Source/MLDeformerFramework/Private/MLDeformerInputInfo.cpp b/Engine/Plugins/Animation/MLDeformer/MLDeformerFramework/Source/MLDeformerFramework/Private/MLDeformerInputInfo.cpp
--- a/Engine/Plugins/Animation/MLDeformer/MLDeformerFramework/Source/MLDeformerFramework/Private/MLDeformerInputInfo.cpp
+++ b/Engine/Plugins/Animation/MLDeformer/MLDeformerFramework/Source/MLDeformerFramework/Private/MLDeformerInputInfo.cpp
@@ -143,15 +143,11 @@ bool UMLDeformerInputInfo::IsCompatible(USkeletalMesh* InSkeletalMesh) const
 
 	// Verify that all required curves are there.
 	USkeleton* Skeleton = InSkeletalMesh->GetSkeleton();
-	const FSmartNameMapping* SmartNameMapping = Skeleton ? Skeleton->GetSmartNameContainer(USkeleton::AnimCurveMappingName) : nullptr;
-	if (SmartNameMapping)
+	for (const FName CurveName : CurveNames)
 	{
-		for (const FName CurveName : CurveNames)
+		if (Skeleton->GetCurveMetaData(CurveName) == nullptr)
 		{
-			if (!SmartNameMapping->Exists(CurveName))
-			{
-				return false;
-			}
+			return false;
 		}
 	}
 
@@ -179,15 +175,11 @@ FString UMLDeformerInputInfo::GenerateCompatibilityErrorString(USkeletalMesh* In
 
 	// Verify that all required curves are there.
 	USkeleton* Skeleton = InSkeletalMesh->GetSkeleton();
-	const FSmartNameMapping* SmartNameMapping = Skeleton ? Skeleton->GetSmartNameContainer(USkeleton::AnimCurveMappingName) : nullptr;
-	if (SmartNameMapping)
+	for (const FName CurveName : CurveNames)
 	{
-		for (const FName CurveName : CurveNames)
+		if (Skeleton->GetCurveMetaData(CurveName) == nullptr)
 		{
-			if (!SmartNameMapping->Exists(CurveName))
-			{
-				ErrorString += FString::Format(TEXT("Required curve '{0}' is missing.\n"), {*CurveName.ToString()});
-			}
+			ErrorString += FString::Format(TEXT("Required curve '{0}' is missing.\n"), {*CurveName.ToString()});
 		}
 	}
 
diff --git a/Engine/Plugins/Animation/MLDeformer/MLDeformerFramework/Source/MLDeformerFrameworkEditor/Private/MLDeformerEditorModel.cpp b/Engine/Plugins/Animation/MLDeformer/MLDeformerFramework/Source/MLDeformerFrameworkEditor/Private/MLDeformerEditorModel.cpp
--- a/Engine/Plugins/Animation/MLDeformer/MLDeformerFramework/Source/MLDeformerFrameworkEditor/Private/MLDeformerEditorModel.cpp
+++ b/Engine/Plugins/Animation/MLDeformer/MLDeformerFramework/Source/MLDeformerFrameworkEditor/Private/MLDeformerEditorModel.cpp
@@ -1050,29 +1050,18 @@ namespace UE::MLDeformer
 		if (bIncludeCurves && SkeletalMesh)
 		{
 			// Anim curves.
-			const FSmartNameMapping* SmartNameMapping = Skeleton ? Skeleton->GetSmartNameContainer(USkeleton::AnimCurveMappingName) : nullptr;
-			if (SmartNameMapping) // When there are curves.
+			// Include all curves when no list was provided.
+			if (Model->GetCurveIncludeList().IsEmpty())
 			{
-				// Include all curves when no list was provided.
-				if (Model->GetCurveIncludeList().IsEmpty())
-				{
-					SmartNameMapping->FillNameArray(CurveNames);
-				}
-				else // A list of curve names was provided.
+				Skeleton->GetCurveMetaDataNames(CurveNames);
+			}
+			else // A list of curve names was provided.
+			{
+				for (const FMLDeformerCurveReference& CurveReference : Model->GetCurveIncludeList())
 				{
-					for (const FMLDeformerCurveReference& CurveReference : Model->GetCurveIncludeList())
+					if (CurveReference.CurveName.IsValid())
 					{
-						if (CurveReference.CurveName.IsValid())
-						{
-							const FName CurveName = CurveReference.CurveName;
-							if (!SmartNameMapping->Exists(CurveName))
-							{
-								UE_LOG(LogMLDeformer, Warning, TEXT("Curve '%s' doesn't exist, ignoring it."), *CurveName.ToString());
-								continue;
-							}
-
-							CurveNames.Add(CurveName);
-						}
+						CurveNames.Add(CurveReference.CurveName);
 					}
 				}
 			}
@@ -1206,35 +1195,31 @@ namespace UE::MLDeformer
 
 		// Iterate over all curves that are both in the skeleton and the animation.
 		TArray<FName> AnimatedCurveList;
-		const FSmartNameMapping* Mapping = Skeleton->GetSmartNameContainer(USkeleton::AnimCurveMappingName);
-		if (Mapping)
+		TArray<FName> SkeletonCurveNames;
+		Skeleton->GetCurveMetaDataNames(SkeletonCurveNames);
+		for (const FName& SkeletonCurveName : SkeletonCurveNames)
 		{
-			TArray<FName> SkeletonCurveNames;
-			Mapping->FillNameArray(SkeletonCurveNames);
-			for (const FName& SkeletonCurveName : SkeletonCurveNames)
+			const TArray<FFloatCurve>& AnimCurves = DataModel->GetFloatCurves();
+			for (const FFloatCurve& AnimCurve : AnimCurves)
 			{
-				const TArray<FFloatCurve>& AnimCurves = DataModel->GetFloatCurves();
-				for (const FFloatCurve& AnimCurve : AnimCurves)
+				if (AnimCurve.GetName() == SkeletonCurveName)
 				{
-					if (AnimCurve.Name.IsValid() && AnimCurve.Name.DisplayName == SkeletonCurveName)
+					TArray<float> TimeValues;
+					TArray<float> KeyValues;
+					AnimCurve.GetKeys(TimeValues, KeyValues);
+					if (KeyValues.Num() > 0)
 					{
-						TArray<float> TimeValues;
-						TArray<float> KeyValues;
-						AnimCurve.GetKeys(TimeValues, KeyValues);
-						if (KeyValues.Num() > 0)
+						const float FirstKeyValue = KeyValues[0];
+						for (float CurKeyValue : KeyValues)
 						{
-							const float FirstKeyValue = KeyValues[0];					
-							for (float CurKeyValue : KeyValues)
+							if (CurKeyValue != FirstKeyValue)
 							{
-								if (CurKeyValue != FirstKeyValue)
-								{
-									AnimatedCurveList.Add(SkeletonCurveName);
-									break;
-								}
+								AnimatedCurveList.Add(SkeletonCurveName);
+								break;
 							}
 						}
-						break;
 					}
+					break;
 				}
 			}
 		}
@@ -1877,13 +1862,7 @@ namespace UE::MLDeformer
 			const USkeleton* Skeleton = SkelMesh->GetSkeleton();
 			if (Skeleton)
 			{
-				const FSmartNameMapping* Mapping = Skeleton->GetSmartNameContainer(USkeleton::AnimCurveMappingName);
-				if (Mapping)
-				{
-					TArray<FName> SkeletonCurveNames;
-					Mapping->FillNameArray(SkeletonCurveNames);
-					return SkeletonCurveNames.Num();
-				}
+				return Skeleton->GetNumCurveMetaData();
 			}
 		}
 
diff --git a/Engine/Plugins/Animation/MLDeformer/MLDeformerFramework/Source/MLDeformerFrameworkEditor/Private/SMLDeformerCurvePickerWidget.cpp b/Engine/Plugins/Animation/MLDeformer/MLDeformerFramework/Source/MLDeformerFrameworkEditor/Private/SMLDeformerCurvePickerWidget.cpp
--- a/Engine/Plugins/Animation/MLDeformer/MLDeformerFramework/Source/MLDeformerFrameworkEditor/Private/SMLDeformerCurvePickerWidget.cpp
+++ b/Engine/Plugins/Animation/MLDeformer/MLDeformerFramework/Source/MLDeformerFrameworkEditor/Private/SMLDeformerCurvePickerWidget.cpp
@@ -78,17 +78,13 @@ namespace UE::MLDeformer
 		// Get the curve names.
 		if (Skeleton)
 		{
-			const FSmartNameMapping* Mapping = Skeleton->GetSmartNameContainer(USkeleton::AnimCurveMappingName);
-			if (Mapping)
+			TArray<FName> SkeletonCurveNames;
+			Skeleton->GetCurveMetaDataNames(SkeletonCurveNames);
+
+			UniqueCurveNames.Reserve(SkeletonCurveNames.Num());
+			for (const FName& CurveName : SkeletonCurveNames)
 			{
-				TArray<FName> SkeletonCurveNames;
-				Mapping->FillNameArray(SkeletonCurveNames);
-
-				UniqueCurveNames.Reserve(SkeletonCurveNames.Num());
-				for (const FName& CurveName : SkeletonCurveNames)
-				{
-					UniqueCurveNames.Add(CurveName.ToString());
-				}
+				UniqueCurveNames.Add(CurveName.ToString());
 			}
 		}
 
diff --git a/Engine/Plugins/Animation/MLDeformer/NeuralMorphModel/Source/NeuralMorphModelEditor/Private/NeuralMorphEditorModel.cpp b/Engine/Plugins/Animation/MLDeformer/NeuralMorphModel/Source/NeuralMorphModelEditor/Private/NeuralMorphEditorModel.cpp
--- a/Engine/Plugins/Animation/MLDeformer/NeuralMorphModel/Source/NeuralMorphModelEditor/Private/NeuralMorphEditorModel.cpp
+++ b/Engine/Plugins/Animation/MLDeformer/NeuralMorphModel/Source/NeuralMorphModelEditor/Private/NeuralMorphEditorModel.cpp
@@ -167,49 +167,39 @@ namespace UE::NeuralMorphModel
 		// Handle curves.
 		if (bIncludeCurves && SkeletalMesh)
 		{
-			const FSmartNameMapping* SmartNameMapping = Skeleton ? Skeleton->GetSmartNameContainer(USkeleton::AnimCurveMappingName) : nullptr;
-			if (SmartNameMapping) // When there are curves.
-			{
-				const TArray<FNeuralMorphCurveGroup>& ModelCurveGroups = NeuralMorphModel->GetCurveGroups();
-				const FReferenceSkeleton& RefSkeleton = SkeletalMesh->GetRefSkeleton();
-				for (int32 CurveGroupIndex = 0; CurveGroupIndex < ModelCurveGroups.Num(); ++CurveGroupIndex)
-				{				
-					const FNeuralMorphCurveGroup& CurveGroup = ModelCurveGroups[CurveGroupIndex];
-					if (CurveGroup.CurveNames.IsEmpty())
+			const TArray<FNeuralMorphCurveGroup>& ModelCurveGroups = NeuralMorphModel->GetCurveGroups();
+			const FReferenceSkeleton& RefSkeleton = SkeletalMesh->GetRefSkeleton();
+			for (int32 CurveGroupIndex = 0; CurveGroupIndex < ModelCurveGroups.Num(); ++CurveGroupIndex)
+			{				
+				const FNeuralMorphCurveGroup& CurveGroup = ModelCurveGroups[CurveGroupIndex];
+				if (CurveGroup.CurveNames.IsEmpty())
+				{
+					continue;
+				}
+
+				NeuralInputInfo->GetCurveGroups().AddDefaulted();
+				FNeuralMorphCurveGroup& NewGroup = NeuralInputInfo->GetCurveGroups().Last();
+
+				const int32 NumCurvesInGroup = CurveGroup.CurveNames.Num();
+				NewGroup.CurveNames.AddDefaulted(NumCurvesInGroup);
+				for (int32 Index = 0; Index < NumCurvesInGroup; ++Index)
+				{
+					const FMLDeformerCurveReference& CurveRef = CurveGroup.CurveNames[Index];
+					if (!CurveRef.CurveName.IsValid() || CurveRef.CurveName.IsNone())
 					{
+						UE_LOG(LogNeuralMorphModel, Warning, TEXT("Invalid or 'None' curve detected inside curve group %d, please fix this."), CurveGroupIndex);
 						continue;
 					}
 
-					NeuralInputInfo->GetCurveGroups().AddDefaulted();
-					FNeuralMorphCurveGroup& NewGroup = NeuralInputInfo->GetCurveGroups().Last();
-
-					const int32 NumCurvesInGroup = CurveGroup.CurveNames.Num();
-					NewGroup.CurveNames.AddDefaulted(NumCurvesInGroup);
-					for (int32 Index = 0; Index < NumCurvesInGroup; ++Index)
+					if (!InputInfo->GetCurveNames().Contains(CurveRef.CurveName))
 					{
-						const FMLDeformerCurveReference& CurveRef = CurveGroup.CurveNames[Index];
-						if (!CurveRef.CurveName.IsValid() || CurveRef.CurveName.IsNone())
-						{
-							UE_LOG(LogNeuralMorphModel, Warning, TEXT("Invalid or 'None' curve detected inside curve group %d, please fix this."), CurveGroupIndex);
-							continue;
-						}
-
-						if (!SmartNameMapping->Exists(CurveRef.CurveName))
-						{
-							UE_LOG(LogNeuralMorphModel, Warning, TEXT("Curve '%s' inside curve group %d doesn't exist, please fix this."), *CurveRef.CurveName.ToString(), CurveGroupIndex);
-							continue;
-						}
-
-						if (!InputInfo->GetCurveNames().Contains(CurveRef.CurveName))
-						{
-							UE_LOG(LogNeuralMorphModel, Warning, TEXT("Curve '%s' inside curve group %d isn't included in the curve list that are input to the model."), *CurveRef.CurveName.ToString(), CurveGroupIndex);
-							continue;
-						}
-
-						NewGroup.CurveNames[Index] = CurveRef.CurveName;
+						UE_LOG(LogNeuralMorphModel, Warning, TEXT("Curve '%s' inside curve group %d isn't included in the curve list that are input to the model."), *CurveRef.CurveName.ToString(), CurveGroupIndex);
+						continue;
 					}
+
+					NewGroup.CurveNames[Index] = CurveRef.CurveName;
 				}
-			} // if SmartNameMapping
+			}
 		}
 	}
 
diff --git a/Engine/Plugins/Animation/MotionWarping/Source/MotionWarping/Private/MotionWarpingComponent.cpp b/Engine/Plugins/Animation/MotionWarping/Source/MotionWarping/Private/MotionWarpingComponent.cpp
--- a/Engine/Plugins/Animation/MotionWarping/Source/MotionWarping/Private/MotionWarpingComponent.cpp
+++ b/Engine/Plugins/Animation/MotionWarping/Source/MotionWarping/Private/MotionWarpingComponent.cpp
@@ -176,7 +176,7 @@ FTransform UMotionWarpingUtilities::CalculateRootTransformRelativeToWarpPointAtT
 				TArray<FBoneIndexType> RequiredBoneIndexArray = { 0, (FBoneIndexType)BoneIndex };
 				FullBoneContainer.GetReferenceSkeleton().EnsureParentsExistAndSort(RequiredBoneIndexArray);
 
-				FBoneContainer LimitedBoneContainer(RequiredBoneIndexArray, FCurveEvaluationOption(false), *FullBoneContainer.GetAsset());
+				FBoneContainer LimitedBoneContainer(RequiredBoneIndexArray, UE::Anim::FCurveFilterSettings(UE::Anim::ECurveFilterMode::DisallowAll), *FullBoneContainer.GetAsset());
 
 				FCSPose<FCompactPose> Pose;
 				UMotionWarpingUtilities::ExtractComponentSpacePose(Animation, LimitedBoneContainer, Time, false, Pose);
diff --git a/Engine/Plugins/Animation/MotionWarping/Source/MotionWarping/Private/RootMotionModifier_AdjustmentBlendWarp.cpp b/Engine/Plugins/Animation/MotionWarping/Source/MotionWarping/Private/RootMotionModifier_AdjustmentBlendWarp.cpp
--- a/Engine/Plugins/Animation/MotionWarping/Source/MotionWarping/Private/RootMotionModifier_AdjustmentBlendWarp.cpp
+++ b/Engine/Plugins/Animation/MotionWarping/Source/MotionWarping/Private/RootMotionModifier_AdjustmentBlendWarp.cpp
@@ -132,7 +132,7 @@ void URootMotionModifier_AdjustmentBlendWarp::PrecomputeWarpedTracks()
 	}
 
 	// Init BoneContainer
-	FBoneContainer RequiredBones(RequiredBoneIndexArray, FCurveEvaluationOption(false), *BoneContainer.GetAsset());
+	FBoneContainer RequiredBones(RequiredBoneIndexArray, UE::Anim::FCurveFilterSettings(UE::Anim::ECurveFilterMode::DisallowAll), *BoneContainer.GetAsset());
 
 	// Extract pose
 	FCSPose<FCompactPose> CSPose;
diff --git a/Engine/Plugins/Animation/RigLogic/Source/RigLogicModule/Private/AnimNode_RigLogic.cpp b/Engine/Plugins/Animation/RigLogic/Source/RigLogicModule/Private/AnimNode_RigLogic.cpp
--- a/Engine/Plugins/Animation/RigLogic/Source/RigLogicModule/Private/AnimNode_RigLogic.cpp
+++ b/Engine/Plugins/Animation/RigLogic/Source/RigLogicModule/Private/AnimNode_RigLogic.cpp
@@ -147,16 +147,16 @@ void FAnimNode_RigLogic::UpdateControlCurves(const FPoseContext& InputContext, c
 	DECLARE_SCOPE_HIERARCHICAL_COUNTER_FUNC()
 	QUICK_SCOPE_CYCLE_COUNTER(STAT_AnimNode_RigLogic_UpdateControlCurves);
 
-	const int32 RawControlCount = RigInstance->GetRawControlCount();
-	for (int32 ControlIndex = 0; ControlIndex < RawControlCount; ++ControlIndex)
-	{
-		const int32 ControlCurveUID = DNAIndexMapping->ControlAttributesMapDNAIndicesToUEUIDs[ControlIndex];
-		if (ControlCurveUID != INDEX_NONE)
+	// Combine control attribute curve with input curve to get indexed curve to apply to rig
+	// Curve elements that dont have a control mapping will have INDEX_NONE as their index
+	UE::Anim::FNamedValueArrayUtils::Union(InputContext.Curve, DNAIndexMapping->ControlAttributeCurves,
+		[this](const UE::Anim::FCurveElement& InCurveElement, const UE::Anim::FCurveElementIndexed& InControlAttributeCurveElement, UE::Anim::ENamedValueUnionFlags InFlags)
 		{
-			const float Value = InputContext.Curve.Get(static_cast<SmartName::UID_Type>(ControlCurveUID));
-			RigInstance->SetRawControl(ControlIndex, Value);
-		}
-	}
+			if(InControlAttributeCurveElement.Index != INDEX_NONE)
+ 			{
+ 				RigInstance->SetRawControl(InControlAttributeCurveElement.Index, InCurveElement.Value);
+ 			}
+		});
 }
 
 void FAnimNode_RigLogic::CalculateRigLogic(FRigLogic* RigLogic)
@@ -196,21 +196,15 @@ void FAnimNode_RigLogic::UpdateBlendShapeCurves(const FDNAIndexMapping* DNAIndex
 {
 	DECLARE_SCOPE_HIERARCHICAL_COUNTER_FUNC()
 	QUICK_SCOPE_CYCLE_COUNTER(STAT_AnimNode_RigLogic_UpdateBlendShapeCurves);
-
+	
 	const uint16 LOD = RigInstance->GetLOD();
-	const TArray<int32>& BlendShapeIndices = DNAIndexMapping->BlendShapeIndicesPerLOD[LOD].Values;
-	const TArray<int32>& MorphTargetIndices = DNAIndexMapping->MorphTargetIndicesPerLOD[LOD].Values;
-	check(BlendShapeIndices.Num() == MorphTargetIndices.Num());
-	for (uint32 MappingIndex = 0; MappingIndex < static_cast<uint32>(BlendShapeIndices.Num()); ++MappingIndex)
+	const FDNAIndexMapping::FCachedIndexedCurve& MorphTargetCurve = DNAIndexMapping->MorphTargetCurvesPerLOD[LOD];
+	UE::Anim::FNamedValueArrayUtils::Union(OutputContext.Curve, MorphTargetCurve, [&BlendShapeValues](UE::Anim::FCurveElement& InOutResult, const UE::Anim::FCurveElementIndexed& InSource, UE::Anim::ENamedValueUnionFlags InFlags)
 	{
-		const int32 BlendShapeIndex = BlendShapeIndices[MappingIndex];
-		const int32 MorphTargetCurveUID = MorphTargetIndices[MappingIndex];
-		if (MorphTargetCurveUID != INDEX_NONE)
-		{
-			const float Value = BlendShapeValues[BlendShapeIndex];
-			OutputContext.Curve.Set(static_cast<SmartName::UID_Type>(MorphTargetCurveUID), Value);
-		}
-	}
+		check(BlendShapeValues.IsValidIndex(InSource.Index));
+		InOutResult.Value = BlendShapeValues[InSource.Index];
+		InOutResult.Flags |= UE::Anim::ECurveElementFlags::MorphTarget;
+	});
 }
 
 void FAnimNode_RigLogic::UpdateAnimMapCurves(const FDNAIndexMapping* DNAIndexMapping, TArrayView<const float> AnimMapOutputs, FPoseContext& OutputContext)
@@ -219,18 +213,12 @@ void FAnimNode_RigLogic::UpdateAnimMapCurves(const FDNAIndexMapping* DNAIndexMap
 	QUICK_SCOPE_CYCLE_COUNTER(STAT_AnimNode_RigLogic_UpdateAnimMapCurves);
 
 	const uint16 LOD = RigInstance->GetLOD();
-	const TArray<int32>& AnimatedMapIndices = DNAIndexMapping->AnimatedMapIndicesPerLOD[LOD].Values;
-	const TArray<int32>& MaskMultiplierIndices = DNAIndexMapping->MaskMultiplierIndicesPerLOD[LOD].Values;
-	check(AnimatedMapIndices.Num() == MaskMultiplierIndices.Num());
-	for (uint32 MappingIndex = 0; MappingIndex < static_cast<uint32>(AnimatedMapIndices.Num()); ++MappingIndex)
+	const FDNAIndexMapping::FCachedIndexedCurve& MaskMultiplierCurve = DNAIndexMapping->MaskMultiplierCurvesPerLOD[LOD];
+	UE::Anim::FNamedValueArrayUtils::Union(OutputContext.Curve, MaskMultiplierCurve, [&AnimMapOutputs](UE::Anim::FCurveElement& InOutResult, const UE::Anim::FCurveElementIndexed& InSource, UE::Anim::ENamedValueUnionFlags InFlags)
 	{
-		const int32 AnimatedMapIndex = AnimatedMapIndices[MappingIndex];
-		const int32 MaskMultiplierUID = MaskMultiplierIndices[MappingIndex];
-		if (MaskMultiplierUID != INDEX_NONE)
-		{
-			const float Value = AnimMapOutputs[AnimatedMapIndex];
-			OutputContext.Curve.Set(static_cast<SmartName::UID_Type>(MaskMultiplierUID), Value);
-		}
-	}
+		check(AnimMapOutputs.IsValidIndex(InSource.Index));
+		InOutResult.Value = AnimMapOutputs[InSource.Index];
+		InOutResult.Flags |= UE::Anim::ECurveElementFlags::Material;
+	});
 }
 
diff --git a/Engine/Plugins/Animation/RigLogic/Source/RigLogicModule/Private/DNAIndexMapping.cpp b/Engine/Plugins/Animation/RigLogic/Source/RigLogicModule/Private/DNAIndexMapping.cpp
--- a/Engine/Plugins/Animation/RigLogic/Source/RigLogicModule/Private/DNAIndexMapping.cpp
+++ b/Engine/Plugins/Animation/RigLogic/Source/RigLogicModule/Private/DNAIndexMapping.cpp
@@ -8,6 +8,7 @@
 #include "Animation/SmartName.h"
 #include "Components/SkeletalMeshComponent.h"
 #include "Engine/SkeletalMesh.h"
+#include "Animation/AnimCurveTypes.h"
 
 
 /** Constructs curve name from nameToSplit using formatString of form x<obj>y<attr>z **/
@@ -35,7 +36,10 @@ void FDNAIndexMapping::MapControlCurves(const IBehaviorReader* DNABehavior, cons
 	LLM_SCOPE_BYNAME(TEXT("Animation/RigLogic"));
 
 	const uint32 ControlCount = DNABehavior->GetRawControlCount();
-	ControlAttributesMapDNAIndicesToUEUIDs.Reset(ControlCount);
+
+	ControlAttributeCurves.Empty();
+	ControlAttributeCurves.Reserve(ControlCount);
+
 	for (uint32_t ControlIndex = 0; ControlIndex < ControlCount; ++ControlIndex)
 	{
 		const FString DNAControlName = DNABehavior->GetRawControlName(ControlIndex);
@@ -44,12 +48,7 @@ void FDNAIndexMapping::MapControlCurves(const IBehaviorReader* DNABehavior, cons
 		{
 			return;
 		}
-		const FName ControlFName(*AnimatedControlName);
-		FSmartName ControlCurveName;
-		const bool IsControlFound = Skeleton->GetSmartNameByName(USkeleton::AnimCurveMappingName, ControlFName, ControlCurveName);
-		// SmartName::UID_Type is uint32; we use int32 to allow for NONE value
-		const int32 UID = (IsControlFound ? static_cast<int32>(ControlCurveName.UID) : INDEX_NONE);
-		ControlAttributesMapDNAIndicesToUEUIDs.Add(UID);
+		ControlAttributeCurves.Add(*AnimatedControlName, ControlIndex);
 	}
 }
 
@@ -77,16 +76,15 @@ void FDNAIndexMapping::MapMorphTargets(const IBehaviorReader* DNABehavior, const
 	const TMap<FName, int32>& MorphTargetIndexMap = SkeletalMesh->GetMorphTargetIndexMap();
 	const TArray<UMorphTarget*>& MorphTargets = SkeletalMesh->GetMorphTargets();
 
-	BlendShapeIndicesPerLOD.Reset(LODCount);
-	BlendShapeIndicesPerLOD.AddDefaulted(LODCount);
-	MorphTargetIndicesPerLOD.Reset(LODCount);
-	MorphTargetIndicesPerLOD.AddDefaulted(LODCount);
+	MorphTargetCurvesPerLOD.Reset(LODCount);
+	MorphTargetCurvesPerLOD.AddDefaulted(LODCount);
 
 	for (uint16 LODIndex = 0; LODIndex < LODCount; ++LODIndex)
 	{
 		TArrayView<const uint16> MappingIndicesForLOD = DNABehavior->GetMeshBlendShapeChannelMappingIndicesForLOD(LODIndex);
-		MorphTargetIndicesPerLOD[LODIndex].Values.Reserve(MappingIndicesForLOD.Num());
-		BlendShapeIndicesPerLOD[LODIndex].Values.Reserve(MappingIndicesForLOD.Num());
+
+		MorphTargetCurvesPerLOD[LODIndex].Reserve(MappingIndicesForLOD.Num());
+
 		for (uint16 MappingIndex : MappingIndicesForLOD)
 		{
 			const FMeshBlendShapeChannelMapping Mapping = DNABehavior->GetMeshBlendShapeChannelMapping(MappingIndex);
@@ -95,14 +93,11 @@ void FDNAIndexMapping::MapMorphTargets(const IBehaviorReader* DNABehavior, const
 			const FString MorphTargetStr = MeshName + TEXT("__") + BlendShapeName;
 			const FName MorphTargetName(*MorphTargetStr);
 			const int32* MorphTargetIndex = MorphTargetIndexMap.Find(MorphTargetName);
-			int32 UID = INDEX_NONE;
 			if ((MorphTargetIndex != nullptr) && (*MorphTargetIndex != INDEX_NONE))
 			{
 				const UMorphTarget* MorphTarget = MorphTargets[*MorphTargetIndex];
-				UID = Skeleton->GetUIDByName(USkeleton::AnimCurveMappingName, MorphTarget->GetFName());
+				MorphTargetCurvesPerLOD[LODIndex].Add(MorphTarget->GetFName(), Mapping.BlendShapeChannelIndex);
 			}
-			BlendShapeIndicesPerLOD[LODIndex].Values.Add(Mapping.BlendShapeChannelIndex);
-			MorphTargetIndicesPerLOD[LODIndex].Values.Add(UID);
 		}
 	}
 }
@@ -113,16 +108,15 @@ void FDNAIndexMapping::MapMaskMultipliers(const IBehaviorReader* DNABehavior, co
 
 	const uint16 LODCount = DNABehavior->GetLODCount();
 
-	AnimatedMapIndicesPerLOD.Reset();
-	AnimatedMapIndicesPerLOD.AddDefaulted(LODCount);
-	MaskMultiplierIndicesPerLOD.Reset();
-	MaskMultiplierIndicesPerLOD.AddDefaulted(LODCount);
+	MaskMultiplierCurvesPerLOD.Reset();
+	MaskMultiplierCurvesPerLOD.AddDefaulted(LODCount);
 
 	for (uint16 LODIndex = 0; LODIndex < LODCount; ++LODIndex)
 	{
 		TArrayView<const uint16> IndicesPerLOD = DNABehavior->GetAnimatedMapIndicesForLOD(LODIndex);
-		AnimatedMapIndicesPerLOD[LODIndex].Values.Reserve(IndicesPerLOD.Num());
-		MaskMultiplierIndicesPerLOD[LODIndex].Values.Reserve(IndicesPerLOD.Num());
+
+		MaskMultiplierCurvesPerLOD[LODIndex].Reserve(IndicesPerLOD.Num());
+
 		for (uint16 AnimMapIndex : IndicesPerLOD)
 		{
 			const FString AnimatedMapName = DNABehavior->GetAnimatedMapName(AnimMapIndex);
@@ -131,12 +125,8 @@ void FDNAIndexMapping::MapMaskMultipliers(const IBehaviorReader* DNABehavior, co
 			{
 				return;
 			}
-			const FName MaskMultiplierName(*MaskMultiplierNameStr);
-			FSmartName CurveName;
-			const bool IsControlFound = Skeleton->GetSmartNameByName(USkeleton::AnimCurveMappingName, MaskMultiplierName, CurveName);
-			const int32 UID = (IsControlFound ? CurveName.UID : INDEX_NONE);
-			AnimatedMapIndicesPerLOD[LODIndex].Values.Add(AnimMapIndex);
-			MaskMultiplierIndicesPerLOD[LODIndex].Values.Add(UID);
+
+			MaskMultiplierCurvesPerLOD[LODIndex].Add(*MaskMultiplierNameStr, AnimMapIndex);
 		}
 	}
 }
diff --git a/Engine/Plugins/Animation/RigLogic/Source/RigLogicModule/Private/DNAIndexMapping.h b/Engine/Plugins/Animation/RigLogic/Source/RigLogicModule/Private/DNAIndexMapping.h
--- a/Engine/Plugins/Animation/RigLogic/Source/RigLogicModule/Private/DNAIndexMapping.h
+++ b/Engine/Plugins/Animation/RigLogic/Source/RigLogicModule/Private/DNAIndexMapping.h
@@ -5,6 +5,7 @@
 #include "CoreMinimal.h"
 #include "BoneIndices.h"
 #include "Engine/AssetUserData.h"
+#include "Animation/AnimCurveTypes.h"
 
 #include "DNAIndexMapping.generated.h"
 
@@ -20,16 +21,16 @@ struct FDNAIndexMapping
 	{
 		TArray<T> Values;
 	};
+	
+	using FCachedIndexedCurve = TBaseBlendedCurve<FDefaultAllocator, UE::Anim::FCurveElementIndexed>; 
 
 	FGuid SkeletonGuid;
-	// For maps, we use int32 instead of SmartName::UID_Ttype directly to allow storing INDEX_NONE for missing elements
-	// if value is valid, it is cast to SmartName::UID_Type
-	TArray<int32> ControlAttributesMapDNAIndicesToUEUIDs;
+
+	// all the control attributes that we will need to extract, alongside their control index
+	FCachedIndexedCurve ControlAttributeCurves;
 	TArray<FMeshPoseBoneIndex> JointsMapDNAIndicesToMeshPoseBoneIndices;
-	TArray<TArrayWrapper<int32>> BlendShapeIndicesPerLOD;
-	TArray<TArrayWrapper<int32>> MorphTargetIndicesPerLOD;
-	TArray<TArrayWrapper<int32>> AnimatedMapIndicesPerLOD;
-	TArray<TArrayWrapper<int32>> MaskMultiplierIndicesPerLOD;
+	TArray<FCachedIndexedCurve> MorphTargetCurvesPerLOD;
+	TArray<FCachedIndexedCurve> MaskMultiplierCurvesPerLOD;
 
 	void MapControlCurves(const IBehaviorReader* DNABehavior, const USkeleton* Skeleton);
 	void MapJoints(const IBehaviorReader* DNABehavior, const USkeletalMeshComponent* SkeletalMeshComponent);
diff --git a/Engine/Plugins/Enterprise/GLTFExporter/Source/GLTFExporter/Private/Converters/GLTFBoneUtilities.cpp b/Engine/Plugins/Enterprise/GLTFExporter/Source/GLTFExporter/Private/Converters/GLTFBoneUtilities.cpp
--- a/Engine/Plugins/Enterprise/GLTFExporter/Source/GLTFExporter/Private/Converters/GLTFBoneUtilities.cpp
+++ b/Engine/Plugins/Enterprise/GLTFExporter/Source/GLTFExporter/Private/Converters/GLTFBoneUtilities.cpp
@@ -57,7 +57,7 @@ void FGLTFBoneUtilities::GetBoneTransformsByFrame(const UAnimSequence* AnimSeque
 	FMemMark Mark(FMemStack::Get());
 
 	FBoneContainer BoneContainer;
-	BoneContainer.InitializeTo(BoneIndices, FCurveEvaluationOption(true), *AnimSequence->GetSkeleton());
+	BoneContainer.InitializeTo(BoneIndices, UE::Anim::FCurveFilterSettings(UE::Anim::ECurveFilterMode::DisallowAll), *AnimSequence->GetSkeleton());
 
 	const int32 FrameCount = FrameTimestamps.Num();
 	OutBoneTransformsByFrame.AddDefaulted(FrameCount);
diff --git a/Engine/Plugins/Experimental/Animation/ContextualAnimation/Source/ContextualAnimation/Private/ContextualAnimSceneAsset.cpp b/Engine/Plugins/Experimental/Animation/ContextualAnimation/Source/ContextualAnimation/Private/ContextualAnimSceneAsset.cpp
--- a/Engine/Plugins/Experimental/Animation/ContextualAnimation/Source/ContextualAnimation/Private/ContextualAnimSceneAsset.cpp
+++ b/Engine/Plugins/Experimental/Animation/ContextualAnimation/Source/ContextualAnimation/Private/ContextualAnimSceneAsset.cpp
@@ -252,7 +252,7 @@ void FContextualAnimSceneSection::GenerateIKTargetTracks(UContextualAnimSceneAss
 				{
 					// Complete bones chain and create bone container to extract pose from my animation
 					Animation->GetSkeleton()->GetReferenceSkeleton().EnsureParentsExistAndSort(RequiredBoneIndexArray);
-					FBoneContainer BoneContainer = FBoneContainer(RequiredBoneIndexArray, FCurveEvaluationOption(false), *Animation->GetSkeleton());
+					FBoneContainer BoneContainer = FBoneContainer(RequiredBoneIndexArray, UE::Anim::FCurveFilterSettings(UE::Anim::ECurveFilterMode::DisallowAll), *Animation->GetSkeleton());
 
 					// Initialize track container
 					AnimTrack.IKTargetData.Initialize(TotalTracks, SampleInterval);
@@ -296,7 +296,7 @@ void FContextualAnimSceneSection::GenerateIKTargetTracks(UContextualAnimSceneAss
 								{
 									// Complete bones chain and create bone container to extract pose form the target animation
 									OtherAnimation->GetSkeleton()->GetReferenceSkeleton().EnsureParentsExistAndSort(OtherRequiredBoneIndexArray);
-									OtherBoneContainer = FBoneContainer(OtherRequiredBoneIndexArray, FCurveEvaluationOption(false), *OtherAnimation->GetSkeleton());
+									OtherBoneContainer = FBoneContainer(OtherRequiredBoneIndexArray, UE::Anim::FCurveFilterSettings(UE::Anim::ECurveFilterMode::DisallowAll), *OtherAnimation->GetSkeleton());
 
 									// Extract pose from target animation
 									ExtractPoseIgnoringForceRootLock(OtherAnimation, OtherBoneContainer, Time, false, OtherComponentSpacePose);
diff --git a/Engine/Plugins/Experimental/Animation/ContextualAnimation/Source/ContextualAnimation/Private/ContextualAnimUtilities.cpp b/Engine/Plugins/Experimental/Animation/ContextualAnimation/Source/ContextualAnimation/Private/ContextualAnimUtilities.cpp
--- a/Engine/Plugins/Experimental/Animation/ContextualAnimation/Source/ContextualAnimation/Private/ContextualAnimUtilities.cpp
+++ b/Engine/Plugins/Experimental/Animation/ContextualAnimation/Source/ContextualAnimation/Private/ContextualAnimUtilities.cpp
@@ -133,7 +133,7 @@ void UContextualAnimUtilities::DrawPose(const UAnimSequenceBase* Animation, floa
 		RequiredBoneIndexArray.Add(Idx);
 	}
 
-	FBoneContainer BoneContainer(RequiredBoneIndexArray, FCurveEvaluationOption(false), *Animation->GetSkeleton());
+	FBoneContainer BoneContainer(RequiredBoneIndexArray, UE::Anim::FCurveFilterSettings(UE::Anim::ECurveFilterMode::DisallowAll), *Animation->GetSkeleton());
 	FCSPose<FCompactPose> ComponentSpacePose;
 	UContextualAnimUtilities::ExtractComponentSpacePose(Animation, BoneContainer, Time, true, ComponentSpacePose);
 
diff --git a/Engine/Plugins/Experimental/Animation/CurveExpression/Source/Runtime/Private/AnimNode_RemapCurvesFromMesh.cpp b/Engine/Plugins/Experimental/Animation/CurveExpression/Source/Runtime/Private/AnimNode_RemapCurvesFromMesh.cpp
--- a/Engine/Plugins/Experimental/Animation/CurveExpression/Source/Runtime/Private/AnimNode_RemapCurvesFromMesh.cpp
+++ b/Engine/Plugins/Experimental/Animation/CurveExpression/Source/Runtime/Private/AnimNode_RemapCurvesFromMesh.cpp
@@ -7,6 +7,7 @@
 #include "Animation/ExposedValueHandler.h"
 #include "CurveExpressionModule.h"
 #include "Engine/SkeletalMesh.h"
+#include "Animation/AnimCurveTypes.h"
 
 #include UE_INLINE_GENERATED_CPP_BY_NAME(AnimNode_RemapCurvesFromMesh)
 
@@ -63,12 +64,6 @@ void FAnimNode_RemapCurvesFromMesh::VerifyExpressions(
 
 	for (const TPair<FName, FString>& ExpressionPair: CurveExpressions)
 	{
-		if (!CurveNameToUIDMap.Contains(ExpressionPair.Key))
-		{
-			ReportAndLog(FString::Printf(TEXT("Target curve '%s' does not exist."), *ExpressionPair.Key.ToString()));
-			bFoundError = true;
-		}
-
 		TOptional<FParseError> Error = VerificationEngine.Verify(ExpressionPair.Value);
 		if (Error.IsSet())
 		{
@@ -108,15 +103,12 @@ void FAnimNode_RemapCurvesFromMesh::Update_AnyThread(
 		CachedExpressions.Reset();
 		for (const TPair<FName, FString>& ExpressionPair: CurveExpressions)
 		{
-			if (CurveNameToUIDMap.Contains(ExpressionPair.Key))
-			{
-				using namespace CurveExpression::Evaluator;
+			using namespace CurveExpression::Evaluator;
 
-				TVariant<FExpressionObject, FParseError> Result = ExpressionEngine->Parse(ExpressionPair.Value);
-				if (const FExpressionObject* Expression = Result.TryGet<FExpressionObject>())
-				{
-					CachedExpressions.Add(ExpressionPair.Key, *Expression);
-				}
+			TVariant<FExpressionObject, FParseError> Result = ExpressionEngine->Parse(ExpressionPair.Value);
+			if (const FExpressionObject* Expression = Result.TryGet<FExpressionObject>())
+			{
+				CachedExpressions.Add(ExpressionPair.Key, *Expression);
 			}
 		}
 	}			
@@ -140,27 +132,30 @@ void FAnimNode_RemapCurvesFromMesh::Evaluate_AnyThread(
 	{
 		if (bExpressionsImmutable)
 		{
-			for (const TTuple<FName, CurveExpression::Evaluator::FExpressionObject>& ExpressionItem: CachedExpressions)
+			FBlendedCurve Curve;
+			Curve.Reserve(CachedExpressions.Num());
+			for (const TTuple<FName, CurveExpression::Evaluator::FExpressionObject>& ExpressionItem : CachedExpressions)
 			{
-				if (const SmartName::UID_Type* UID = CurveNameToUIDMap.Find(ExpressionItem.Key))
-				{
-					Output.Curve.Set(*UID, ExpressionEngine->Execute(ExpressionItem.Value));
-				}
+				Curve.Add(ExpressionItem.Key, ExpressionEngine->Execute(ExpressionItem.Value)); 
 			}
+
+			Output.Curve.Combine(Curve);
 		}
 		else
 		{
-			for (const TTuple<FName, FString>& ExpressionItem: CurveExpressions)
+			FBlendedCurve Curve;
+			Curve.Reserve(CurveExpressions.Num());
+
+			for (const TTuple<FName, FString>& ExpressionItem : CurveExpressions)
 			{
-				if (const SmartName::UID_Type* UID = CurveNameToUIDMap.Find(ExpressionItem.Key))
+				TOptional<float> Result = ExpressionEngine->Evaluate(ExpressionItem.Value);
+				if (Result.IsSet())
 				{
-					TOptional<float> Result = ExpressionEngine->Evaluate(ExpressionItem.Value);
-					if (Result.IsSet())
-					{
-						Output.Curve.Set(*UID, *Result);
-					}
+					Curve.Add(ExpressionItem.Key, *Result);
 				}
 			}
+
+			Output.Curve.Combine(Curve);
 		}
 	}
 }
@@ -227,7 +222,6 @@ void FAnimNode_RemapCurvesFromMesh::ReinitializeMeshComponent(
 	CurrentlyUsedTargetMesh.Reset();
 	
 	ExpressionEngine.Reset();
-	CurveNameToUIDMap.Reset();
 	CachedExpressions.Reset();
 
 	if (InTargetMeshComponent && IsValid(InNewSkeletalMeshComponent) && InNewSkeletalMeshComponent->GetSkeletalMeshAsset())
@@ -247,37 +241,14 @@ void FAnimNode_RemapCurvesFromMesh::ReinitializeMeshComponent(
 			const USkeleton* SourceSkeleton = SourceSkelMesh->GetSkeleton();
 			if (ensureMsgf(SourceSkeleton, TEXT("Invalid null source skeleton : %s"), *GetNameSafe(TargetSkelMesh)))
 			{
-				const FSmartNameMapping* SourceContainer = SourceSkeleton->GetSmartNameContainer(USkeleton::AnimCurveMappingName);
-
-				SourceContainer->Iterate([&SourceCurves](const FSmartNameMappingIterator& Iterator)
+				// TODO: is this good enough for future use? ExpressionEngine appears to have been built on the assumption that curves are always pre-defined.
+				SourceSkeleton->ForEachCurveMetaData([&SourceCurves](FName InCurveName, const FCurveMetaData& InCurveMetaData)
 				{
-					if (FName CurveName; Iterator.GetName(CurveName))
-					{
-						SourceCurves.Add(CurveName, 0.0f);
-					}
+					SourceCurves.Add(InCurveName, 0.0f);
 				});
 
 				ExpressionEngine.Emplace(MoveTemp(SourceCurves));
 			}
-
-			// Grab the target curves and build the Name -> UID mapping so that we can call
-			// FBlendCurve::Set in Evalute_AnyThread 
-			const USkeleton* TargetSkeleton = TargetSkelMesh->GetSkeleton();
-
-			if (ensureMsgf(TargetSkeleton, TEXT("Invalid null target skeleton : %s"), *GetNameSafe(TargetSkelMesh)))
-			{
-				const FSmartNameMapping* TargetContainer = TargetSkeleton->GetSmartNameContainer(USkeleton::AnimCurveMappingName);
-
-				// CurveExpressions isn't filled in here, since we haven't evaluated BP downstream.
-				// So instead, we just take a full copy of the smart mapping.
-				TargetContainer->Iterate([this](const FSmartNameMappingIterator& Iterator)
-				{
-					if (FName CurveName; Iterator.GetName(CurveName))
-					{
-						CurveNameToUIDMap.Add(CurveName, Iterator.GetIndex());
-					}
-				});
-			}
 		}
 	}
 }
diff --git a/Engine/Plugins/Experimental/Animation/CurveExpression/Source/Runtime/Public/AnimNode_RemapCurvesFromMesh.h b/Engine/Plugins/Experimental/Animation/CurveExpression/Source/Runtime/Public/AnimNode_RemapCurvesFromMesh.h
--- a/Engine/Plugins/Experimental/Animation/CurveExpression/Source/Runtime/Public/AnimNode_RemapCurvesFromMesh.h
+++ b/Engine/Plugins/Experimental/Animation/CurveExpression/Source/Runtime/Public/AnimNode_RemapCurvesFromMesh.h
@@ -59,8 +59,6 @@ private:
 
 	// target mesh 
 	TWeakObjectPtr<USkeletalMesh>			CurrentlyUsedTargetMesh;
-	
-	TMap<FName, SmartName::UID_Type> CurveNameToUIDMap;
 
 	TOptional<CurveExpression::Evaluator::FEngine> ExpressionEngine;
 	TMap<FName, CurveExpression::Evaluator::FExpressionObject> CachedExpressions;
diff --git a/Engine/Plugins/Experimental/Animation/PoseCorrectives/Source/PoseCorrectives/Private/AnimNode_CorrectPose.cpp b/Engine/Plugins/Experimental/Animation/PoseCorrectives/Source/PoseCorrectives/Private/AnimNode_CorrectPose.cpp
--- a/Engine/Plugins/Experimental/Animation/PoseCorrectives/Source/PoseCorrectives/Private/AnimNode_CorrectPose.cpp
+++ b/Engine/Plugins/Experimental/Animation/PoseCorrectives/Source/PoseCorrectives/Private/AnimNode_CorrectPose.cpp
@@ -39,19 +39,12 @@ void FAnimNode_CorrectPose::CacheBones_AnyThread(const FAnimationCacheBonesConte
 	USkeleton* Skeleton = Context.AnimInstanceProxy->GetSkeleton();
 
 	BoneCompactIndices.Reset();
-	CurveUIDs.Reset();
-	
+
 	for (const FName& BoneName : PoseCorrectivesAsset->GetBoneNames())
 	{
 		FBoneReference BoneRef(BoneName);
 		BoneRef.Initialize(BoneContainer);		
 		BoneCompactIndices.Push(BoneRef.GetCompactPoseIndex(BoneContainer));
-	}	
-
-	for (const FName& CurveName : PoseCorrectivesAsset->GetCurveNames())
-	{
-		SmartName::UID_Type CurveUID = Skeleton->GetUIDByName(USkeleton::AnimCurveMappingName, CurveName);
-		CurveUIDs.Push(CurveUID);
 	}
 }
 
@@ -124,9 +117,9 @@ void FAnimNode_CorrectPose::Evaluate_AnyThread(FPoseContext& Output)
 		}
 	}
 
-	for (const SmartName::UID_Type& DriverCurveUID : CurveUIDs)
+	for (const FName& DriverCurveName : PoseCorrectivesAsset->GetCurveNames())
 	{
-		RBFInput.AddFromScalar(SourceData.Curve.Get(DriverCurveUID));
+		RBFInput.AddFromScalar(SourceData.Curve.Get(DriverCurveName));
 	}
 
 	// Update Targets
@@ -146,9 +139,9 @@ void FAnimNode_CorrectPose::Evaluate_AnyThread(FPoseContext& Output)
 		}
 	}
 
-	for (const SmartName::UID_Type& CurveUID : CurveUIDs)
+	for (const FName& DriverCurveName : PoseCorrectivesAsset->GetCurveNames())
 	{
-		CorrectiveCurveValues.Push(Output.Curve.Get(CurveUID));
+		CorrectiveCurveValues.Push(Output.Curve.Get(DriverCurveName));
 	}
 
 	const TArray<FPoseCorrective>& Correctives = PoseCorrectivesAsset->GetCorrectives();
@@ -166,9 +159,8 @@ void FAnimNode_CorrectPose::Evaluate_AnyThread(FPoseContext& Output)
 		}
 	}
 
-	for (int32 CurveIndex = 0; CurveIndex < CurveUIDs.Num(); CurveIndex++)
+	for (int32 CurveIndex = 0; CurveIndex < PoseCorrectivesAsset->GetCurveNames().Num(); ++CurveIndex)
 	{
-		const SmartName::UID_Type& CurveUID = CurveUIDs[CurveIndex];
-		Output.Curve.Set(CurveUID, CorrectiveCurveValues[CurveIndex]);
+		Output.Curve.Set(PoseCorrectivesAsset->GetCurveNames()[CurveIndex], CorrectiveCurveValues[CurveIndex]);
 	}
 }
diff --git a/Engine/Plugins/Experimental/Animation/PoseCorrectives/Source/PoseCorrectives/Private/PoseCorrectivesAsset.cpp b/Engine/Plugins/Experimental/Animation/PoseCorrectives/Source/PoseCorrectives/Private/PoseCorrectivesAsset.cpp
--- a/Engine/Plugins/Experimental/Animation/PoseCorrectives/Source/PoseCorrectives/Private/PoseCorrectivesAsset.cpp
+++ b/Engine/Plugins/Experimental/Animation/PoseCorrectives/Source/PoseCorrectives/Private/PoseCorrectivesAsset.cpp
@@ -46,16 +46,7 @@ TArray<FName> UPoseCorrectivesAsset::GetCurveNames() const
 	if (TargetMesh)
 	{
 		const USkeleton* Skeleton = TargetMesh->GetSkeleton();
-		const FSmartNameMapping* CurveMapping = Skeleton->GetSmartNameContainer(USkeleton::AnimCurveMappingName);
-
-		CurveMapping->Iterate([&CurveNamesList](const FSmartNameMappingIterator& Iterator)
-			{
-				FName CurveName;
-				if (Iterator.GetName(CurveName))
-				{
-					CurveNamesList.Add(CurveName);
-				}
-			});
+		Skeleton->GetCurveMetaDataNames(CurveNamesList);
 	}
 
 	return CurveNamesList;
@@ -121,7 +112,6 @@ void UPoseCorrectivesAsset::AddCorrective(USkeletalMeshComponent* SourceMeshComp
 
 	const FBlendedHeapCurve& SourceMeshCurves = SourceMeshComponent->GetAnimationCurves();
 	const FBlendedHeapCurve& TargetMeshCurves = TargetMeshComponent->GetAnimationCurves();
-	const FSmartNameMapping* Mapping = TargetMesh->GetSkeleton()->GetSmartNameContainer(USkeleton::AnimCurveMappingName);
 
 	TArray<FName> CurveNames = GetCurveNames();
 	PoseCorrective.CurveData.Reset();
@@ -131,14 +121,10 @@ void UPoseCorrectivesAsset::AddCorrective(USkeletalMeshComponent* SourceMeshComp
 
 	for (int32 CurveIndex = 0; CurveIndex < CurveNames.Num(); ++CurveIndex)
 	{
-		const SmartName::UID_Type CurveUID = Mapping->FindUID(CurveNames[CurveIndex]);
-		if (CurveUID != SmartName::MaxUID)
-		{
-			const float SourceCurveValue = SourceMeshCurves.Get(CurveUID);
-			const float TargetCurveValue = TargetMeshCurves.Get(CurveUID);
-			PoseCorrective.CurveData[CurveIndex] = SourceCurveValue;
-			PoseCorrective.CorrectiveCurvesDelta[CurveIndex] = TargetCurveValue - SourceCurveValue;
-		}
+		const float SourceCurveValue = SourceMeshCurves.Get(CurveNames[CurveIndex]);
+		const float TargetCurveValue = TargetMeshCurves.Get(CurveNames[CurveIndex]);
+		PoseCorrective.CurveData[CurveIndex] = SourceCurveValue;
+		PoseCorrective.CorrectiveCurvesDelta[CurveIndex] = TargetCurveValue - SourceCurveValue;
 	}
 
 	PoseCorrectives.Add(CorrectiveName, PoseCorrective);
diff --git a/Engine/Plugins/Experimental/Animation/PoseCorrectives/Source/PoseCorrectives/Public/AnimNode_CorrectPose.h b/Engine/Plugins/Experimental/Animation/PoseCorrectives/Source/PoseCorrectives/Public/AnimNode_CorrectPose.h
--- a/Engine/Plugins/Experimental/Animation/PoseCorrectives/Source/PoseCorrectives/Public/AnimNode_CorrectPose.h
+++ b/Engine/Plugins/Experimental/Animation/PoseCorrectives/Source/PoseCorrectives/Public/AnimNode_CorrectPose.h
@@ -41,7 +41,6 @@ private:
 	TObjectPtr<UPoseCorrectivesProcessor> PoseCorrectivesProcessor = nullptr;
 	
 	TArray<FCompactPoseBoneIndex> BoneCompactIndices;
-	TArray<SmartName::UID_Type> CurveUIDs;
 
 	TArray<FCorrectivesRBFTarget> RBFTargets;
 	FCorrectivesRBFEntry RBFInput;
diff --git a/Engine/Plugins/Experimental/Animation/PoseCorrectives/Source/PoseCorrectivesEditor/Private/AnimNode_CorrectivesSource.cpp b/Engine/Plugins/Experimental/Animation/PoseCorrectives/Source/PoseCorrectivesEditor/Private/AnimNode_CorrectivesSource.cpp
--- a/Engine/Plugins/Experimental/Animation/PoseCorrectives/Source/PoseCorrectivesEditor/Private/AnimNode_CorrectivesSource.cpp
+++ b/Engine/Plugins/Experimental/Animation/PoseCorrectives/Source/PoseCorrectivesEditor/Private/AnimNode_CorrectivesSource.cpp
@@ -4,6 +4,7 @@
 
 #include "Animation/AnimInstanceProxy.h"
 #include "PoseCorrectivesAsset.h"
+#include "Animation/AnimCurveUtils.h"
 
 
 void FAnimNode_CorrectivesSource::CacheBones_AnyThread(const FAnimationCacheBonesContext& Context)
@@ -16,22 +17,14 @@ void FAnimNode_CorrectivesSource::CacheBones_AnyThread(const FAnimationCacheBone
 	}
 
 	const FBoneContainer& BoneContainer = Context.AnimInstanceProxy->GetRequiredBones();
-	USkeleton* Skeleton = Context.AnimInstanceProxy->GetSkeleton();
 
 	BoneCompactIndices.Reset();
-	CurveUIDs.Reset();
-	
+
 	for (const FName& BoneName : PoseCorrectivesAsset->GetBoneNames())
 	{
 		FBoneReference BoneRef(BoneName);
-		BoneRef.Initialize(BoneContainer);		
+		BoneRef.Initialize(BoneContainer);
 		BoneCompactIndices.Push(BoneRef.GetCompactPoseIndex(BoneContainer));
-	}	
-
-	for (const FName& CurveName : PoseCorrectivesAsset->GetCurveNames())
-	{
-		SmartName::UID_Type CurveUID = Skeleton->GetUIDByName(USkeleton::AnimCurveMappingName, CurveName);
-		CurveUIDs.Push(CurveUID);
 	}
 }
 
@@ -58,12 +51,22 @@ void FAnimNode_CorrectivesSource::Evaluate_AnyThread(FPoseContext& Output)
 			}
 		}
 
-		for (int32 CurveIndex = 0; CurveIndex < CurveUIDs.Num(); ++CurveIndex)
+		FBlendedCurve Curve;
+		auto GetNameFromIndex = [PoseCorrectivesAsset = PoseCorrectivesAsset](int32 InCurveIndex)
+		{
+			return PoseCorrectivesAsset->GetCurveNames()[InCurveIndex];
+		};
+		
+		auto GetValueFromIndex = [PoseCorrective, bUseSourcePose = bUseSourcePose](int32 InCurveIndex)
 		{
-			float CurveValue = PoseCorrective->CurveData[CurveIndex];
+			float CurveValue = PoseCorrective->CurveData[InCurveIndex];
 			if (!bUseSourcePose)
-				CurveValue += PoseCorrective->CorrectiveCurvesDelta[CurveIndex];
-			Output.Curve.Set(CurveUIDs[CurveIndex], CurveValue);
-		}
+				CurveValue += PoseCorrective->CorrectiveCurvesDelta[InCurveIndex];
+			return CurveValue;
+		};
+
+		UE::Anim::FCurveUtils::BuildUnsorted(Curve, PoseCorrectivesAsset->GetCurveNames().Num(), GetNameFromIndex, GetValueFromIndex);
+
+		Output.Curve.Override(Curve);
 	}
 }
diff --git a/Engine/Plugins/Experimental/Animation/PoseCorrectives/Source/PoseCorrectivesEditor/Public/AnimNode_CorrectivesSource.h b/Engine/Plugins/Experimental/Animation/PoseCorrectives/Source/PoseCorrectivesEditor/Public/AnimNode_CorrectivesSource.h
--- a/Engine/Plugins/Experimental/Animation/PoseCorrectives/Source/PoseCorrectivesEditor/Public/AnimNode_CorrectivesSource.h
+++ b/Engine/Plugins/Experimental/Animation/PoseCorrectives/Source/PoseCorrectivesEditor/Public/AnimNode_CorrectivesSource.h
@@ -26,5 +26,5 @@ struct FAnimNode_CorrectivesSource : public FAnimNode_Base
 private:
 
 	TArray<FCompactPoseBoneIndex> BoneCompactIndices;
-	TArray<SmartName::UID_Type> CurveUIDs;
+	
 };
\ No newline at end of file
diff --git a/Engine/Plugins/Experimental/Animation/PoseSearch/Source/Runtime/Private/AnimNode_BlendStack.cpp b/Engine/Plugins/Experimental/Animation/PoseSearch/Source/Runtime/Private/AnimNode_BlendStack.cpp
--- a/Engine/Plugins/Experimental/Animation/PoseSearch/Source/Runtime/Private/AnimNode_BlendStack.cpp
+++ b/Engine/Plugins/Experimental/Animation/PoseSearch/Source/Runtime/Private/AnimNode_BlendStack.cpp
@@ -98,11 +98,7 @@ void FPoseSearchAnimPlayer::StorePoseContext(const FPoseContext& PoseContext)
 		StoredPose.CopyBonesFrom(PoseContext.Pose);
 	}
 
-	if (PoseContext.Curve.IsValid())
-	{
-		StoredCurve.CopyFrom(PoseContext.Curve);
-	}
-
+	StoredCurve.CopyFrom(PoseContext.Curve);
 	StoredAttributes.CopyFrom(PoseContext.CustomAttributes);
 }
 
@@ -118,12 +114,8 @@ void FPoseSearchAnimPlayer::RestorePoseContext(FPoseContext& PoseContext) const
 	{
 		PoseContext.Pose.ResetToRefPose();
 	}
-
-	if (StoredCurve.IsValid())
-	{
-		PoseContext.Curve.CopyFrom(StoredCurve);
-	}
-
+	
+	PoseContext.Curve.CopyFrom(StoredCurve);
 	PoseContext.CustomAttributes.CopyFrom(StoredAttributes);
 }
 
diff --git a/Engine/Plugins/Experimental/Animation/PoseSearch/Source/Runtime/Private/PoseSearchDatabaseIndexingContext.cpp b/Engine/Plugins/Experimental/Animation/PoseSearch/Source/Runtime/Private/PoseSearchDatabaseIndexingContext.cpp
--- a/Engine/Plugins/Experimental/Animation/PoseSearch/Source/Runtime/Private/PoseSearchDatabaseIndexingContext.cpp
+++ b/Engine/Plugins/Experimental/Animation/PoseSearch/Source/Runtime/Private/PoseSearchDatabaseIndexingContext.cpp
@@ -22,7 +22,7 @@ bool FDatabaseIndexingContext::IndexDatabase(FPoseSearchIndexBase& SearchIndexBa
 	check(Schema);
 
 	FBoneContainer BoneContainer;
-	BoneContainer.InitializeTo(Schema->BoneIndicesWithParents, FCurveEvaluationOption(false), *Schema->Skeleton);
+	BoneContainer.InitializeTo(Schema->BoneIndicesWithParents, UE::Anim::FCurveFilterSettings(UE::Anim::ECurveFilterMode::DisallowAll), *Schema->Skeleton);
 
 	SamplingContext.Init(Schema->MirrorDataTable, BoneContainer);
 
diff --git a/Engine/Plugins/Experimental/MeshModelingToolsetExp/Source/MeshModelingToolsExp/Private/SkinWeightsBindingTool.cpp b/Engine/Plugins/Experimental/MeshModelingToolsetExp/Source/MeshModelingToolsExp/Private/SkinWeightsBindingTool.cpp
--- a/Engine/Plugins/Experimental/MeshModelingToolsetExp/Source/MeshModelingToolsExp/Private/SkinWeightsBindingTool.cpp
+++ b/Engine/Plugins/Experimental/MeshModelingToolsetExp/Source/MeshModelingToolsExp/Private/SkinWeightsBindingTool.cpp
@@ -81,18 +81,14 @@ void USkinWeightsBindingTool::Setup()
 		ReferenceSkeleton = SkeletalMesh->GetRefSkeleton();
 		
 		// Initialize the bone browser
-		const FCurveEvaluationOption CurveEvalOption(
-				SkeletalMeshComponent->GetAllowedAnimCurveEvaluate(),
-				&SkeletalMeshComponent->GetDisallowedAnimCurvesEvaluation(),
-				0 /* Always use the highest LOD */
-				);
+		const UE::Anim::FCurveFilterSettings CurveFilterSettings = SkeletalMeshComponent->GetCurveFilterSettings(0 /* Always use the highest LOD */);
 		TArray<FBoneIndexType> AllBones;
 		AllBones.Reserve(ReferenceSkeleton.GetRawBoneNum());
 		for (int32 BoneIndex = 0; BoneIndex < ReferenceSkeleton.GetRawBoneNum(); BoneIndex++)
 		{
 			AllBones.Add(static_cast<FBoneIndexType>(BoneIndex));
 		}
-		BoneContainer.InitializeTo(AllBones, CurveEvalOption, *SkeletalMesh);
+		BoneContainer.InitializeTo(AllBones, CurveFilterSettings, *SkeletalMesh);
 
 		Properties->SkeletalMesh = SkeletalMesh;
 		Properties->CurrentBone.Initialize(BoneContainer);
diff --git a/Engine/Plugins/Experimental/MeshModelingToolsetExp/Source/MeshModelingToolsExp/Private/SkinWeightsPaintTool.cpp b/Engine/Plugins/Experimental/MeshModelingToolsetExp/Source/MeshModelingToolsExp/Private/SkinWeightsPaintTool.cpp
--- a/Engine/Plugins/Experimental/MeshModelingToolsetExp/Source/MeshModelingToolsExp/Private/SkinWeightsPaintTool.cpp
+++ b/Engine/Plugins/Experimental/MeshModelingToolsetExp/Source/MeshModelingToolsExp/Private/SkinWeightsPaintTool.cpp
@@ -373,13 +373,8 @@ void USkinWeightsPaintTool::Setup()
 
 	// initialize the bone browser
 	// TODO replace this and synchronize selection & rendering with skeletal mesh editor
-	FCurveEvaluationOption CurveEvalOption(
-			Component->GetAllowedAnimCurveEvaluate(),
-			&Component->GetDisallowedAnimCurvesEvaluation(),
-			0 /* Always use the highest LOD */
-			);
 	USkeletalMesh& SkeletalMesh = *Component->GetSkeletalMeshAsset();
-	BoneContainer.InitializeTo(Component->RequiredBones, CurveEvalOption, SkeletalMesh);
+	BoneContainer.InitializeTo(Component->RequiredBones, Component->GetCurveFilterSettings(0)/* Always use the highest LOD */, SkeletalMesh);
 	// Update the skeleton drawing information from the original bind pose
 	MaxDrawRadius = Component->Bounds.SphereRadius * 0.0025f;
 	UpdateBonePositionInfos(MaxDrawRadius);
diff --git a/Engine/Plugins/Experimental/Mutable/Source/CustomizableObjectEditor/Private/MuCOE/CustomizableObjectEditorViewportClient.cpp b/Engine/Plugins/Experimental/Mutable/Source/CustomizableObjectEditor/Private/MuCOE/CustomizableObjectEditorViewportClient.cpp
--- a/Engine/Plugins/Experimental/Mutable/Source/CustomizableObjectEditor/Private/MuCOE/CustomizableObjectEditorViewportClient.cpp
+++ b/Engine/Plugins/Experimental/Mutable/Source/CustomizableObjectEditor/Private/MuCOE/CustomizableObjectEditorViewportClient.cpp
@@ -2051,10 +2051,10 @@ void FCustomizableObjectEditorViewportClient::SetAnimation(UAnimationAsset* Anim
 				UAnimSingleNodeInstance* SingleNodeInstance = Cast<UAnimSingleNodeInstance>(SkeletalMeshComponent->GetAnimInstance());
 				if (SingleNodeInstance)
 				{
-					TArray<FSmartName> ArrayPoseSmartNames = PoseAsset->GetPoseNames();
-					for (int32 i = 0; i < ArrayPoseSmartNames.Num(); ++i)
+					TArray<FName> ArrayPoseNames = PoseAsset->GetPoseFNames();
+					for (int32 i = 0; i < ArrayPoseNames.Num(); ++i)
 					{
-						SingleNodeInstance->SetPreviewCurveOverride(ArrayPoseSmartNames[i].DisplayName, 1.0f, false);
+						SingleNodeInstance->SetPreviewCurveOverride(ArrayPoseNames[i], 1.0f, false);
 					}
 				}
 			}
diff --git a/Engine/Plugins/Experimental/Mutable/Source/CustomizableObjectEditor/Private/MuCOE/CustomizableObjectInstanceEditor.cpp b/Engine/Plugins/Experimental/Mutable/Source/CustomizableObjectEditor/Private/MuCOE/CustomizableObjectInstanceEditor.cpp
--- a/Engine/Plugins/Experimental/Mutable/Source/CustomizableObjectEditor/Private/MuCOE/CustomizableObjectInstanceEditor.cpp
+++ b/Engine/Plugins/Experimental/Mutable/Source/CustomizableObjectEditor/Private/MuCOE/CustomizableObjectInstanceEditor.cpp
@@ -597,10 +597,10 @@ void FCustomizableObjectInstanceEditor::SetPoseAsset(class UPoseAsset* PoseAsset
 			UAnimSingleNodeInstance* SingleNodeInstance = Cast<UAnimSingleNodeInstance>(PreviewSkeletalMeshComponent->GetAnimInstance());
 			if (SingleNodeInstance)
 			{
-				TArray<FSmartName> ArrayPoseSmartNames = PoseAsset->GetPoseNames();
+				const TArray<FName>& ArrayPoseSmartNames = PoseAsset->GetPoseFNames();
 				for (int32 i = 0; i < ArrayPoseSmartNames.Num(); ++i)
 				{
-					SingleNodeInstance->SetPreviewCurveOverride(ArrayPoseSmartNames[i].DisplayName, 1.0f, false);
+					SingleNodeInstance->SetPreviewCurveOverride(ArrayPoseSmartNames[i], 1.0f, false);
 				}
 			}
 		}
diff --git a/Engine/Plugins/Experimental/Mutable/Source/CustomizableObjectEditor/Private/MuCOE/Nodes/CustomizableObjectNodeAnimationPose.cpp b/Engine/Plugins/Experimental/Mutable/Source/CustomizableObjectEditor/Private/MuCOE/Nodes/CustomizableObjectNodeAnimationPose.cpp
--- a/Engine/Plugins/Experimental/Mutable/Source/CustomizableObjectEditor/Private/MuCOE/Nodes/CustomizableObjectNodeAnimationPose.cpp
+++ b/Engine/Plugins/Experimental/Mutable/Source/CustomizableObjectEditor/Private/MuCOE/Nodes/CustomizableObjectNodeAnimationPose.cpp
@@ -102,7 +102,7 @@ void UCustomizableObjectNodeAnimationPose::StaticRetrievePoseInformation(UPoseAs
 		}
 
 		FBoneContainer BoneContainer;
-		BoneContainer.InitializeTo(RequiredBones, FCurveEvaluationOption(SkeletalMeshComponent->GetAllowedAnimCurveEvaluate()), *PoseSkeleton);
+		BoneContainer.InitializeTo(RequiredBones, SkeletalMeshComponent->GetCurveFilterSettings(), *PoseSkeleton);
 
 		// Needs a FMemMark declared before in the stack context so that the memory allocated by the FCompactPose is freed correctly
 		FCompactPose OutPose;
@@ -115,18 +115,15 @@ void UCustomizableObjectNodeAnimationPose::StaticRetrievePoseInformation(UPoseAs
 		OutPose = OutAnimData.GetPose();
 		OutCurve = OutAnimData.GetCurve();
 
-		FBlendedHeapCurve AnimCurves = SkeletalMeshComponent->AnimCurves;
-		OutCurve.CurveWeights = AnimCurves.CurveWeights;
-		OutCurve.bInitialized = AnimCurves.bInitialized;
-		OutCurve.UIDToArrayIndexLUT = AnimCurves.UIDToArrayIndexLUT;
+		OutCurve.CopyFrom(SkeletalMeshComponent->AnimCurves);
 
 		// Assuming one single pose, with a weigth set to 1.0
 		FAnimExtractContext ExtractionContext;
 		ExtractionContext.bExtractRootMotion = false;
 		ExtractionContext.CurrentTime = 0.0f;
 		
-		const TArray<FSmartName>& PoseNames = PoseAsset->GetPoseNames();
-		ExtractionContext.PoseCurves.Add(FPoseCurve(0, PoseNames[0].UID, 1.0f));
+		const TArray<FName>& PoseNames = PoseAsset->GetPoseFNames();
+		ExtractionContext.PoseCurves.Add(FPoseCurve(0, PoseNames[0], 1.0f));
 		FAnimationPoseData SecondOutAnimData(OutPose, OutCurve, OutAttributes);
 		PoseAsset->GetAnimationPose(SecondOutAnimData, ExtractionContext);
 		OutPose = OutAnimData.GetPose();
diff --git a/Engine/Plugins/Experimental/SkeletalReduction/Source/Private/SkeletalMeshReductionPlugin.cpp b/Engine/Plugins/Experimental/SkeletalReduction/Source/Private/SkeletalMeshReductionPlugin.cpp
--- a/Engine/Plugins/Experimental/SkeletalReduction/Source/Private/SkeletalMeshReductionPlugin.cpp
+++ b/Engine/Plugins/Experimental/SkeletalReduction/Source/Private/SkeletalMeshReductionPlugin.cpp
@@ -2121,7 +2121,7 @@ void FQuadricSkeletalMeshReduction::ReduceSkeletalMesh(USkeletalMesh& SkeletalMe
 			RequiredBoneIndexArray[BoneIndex] = BoneIndex;
 		}
 
-		FBoneContainer RequiredBones(RequiredBoneIndexArray, false, SkeletalMesh);
+		FBoneContainer RequiredBones(RequiredBoneIndexArray, UE::Anim::ECurveFilterMode::DisallowAll, SkeletalMesh);
 		RequiredBones.SetUseRAWData(true);
 
 		FCompactPose Pose;
diff --git a/Engine/Plugins/Importers/AlembicImporter/Source/AlembicLibrary/Private/AbcImporter.cpp b/Engine/Plugins/Importers/AlembicImporter/Source/AlembicLibrary/Private/AbcImporter.cpp
--- a/Engine/Plugins/Importers/AlembicImporter/Source/AlembicLibrary/Private/AbcImporter.cpp
+++ b/Engine/Plugins/Importers/AlembicImporter/Source/AlembicLibrary/Private/AbcImporter.cpp
@@ -792,7 +792,6 @@ TArray<UObject*> FAbcImporter::ImportAsSkeletalMesh(UObject* InParent, EObjectFl
 
 		SkeletalMesh->MarkPackageDirty();
 
-		Controller.UpdateCurveNamesFromSkeleton(Skeleton, ERawCurveTrackTypes::RCT_Float);
 		Controller.NotifyPopulated();
 
 		Controller.CloseBracket();
@@ -826,10 +825,7 @@ TArray<UObject*> FAbcImporter::ImportAsSkeletalMesh(UObject* InParent, EObjectFl
 
 void FAbcImporter::SetupMorphTargetCurves(USkeleton* Skeleton, FName ConstCurveName, UAnimSequence* Sequence, const TArray<float> &CurveValues, const TArray<float>& TimeValues, IAnimationDataController& Controller)
 {
-	FSmartName NewName;
-	Skeleton->AddSmartNameAndModify(USkeleton::AnimCurveMappingName, ConstCurveName, NewName);
-
-	FAnimationCurveIdentifier CurveId(NewName, ERawCurveTrackTypes::RCT_Float);
+	FAnimationCurveIdentifier CurveId(ConstCurveName, ERawCurveTrackTypes::RCT_Float);
 	Controller.AddCurve(CurveId);
 
 	const FFloatCurve* NewCurve = Sequence->GetDataModel()->FindFloatCurve(CurveId);
diff --git a/Engine/Plugins/Importers/USDImporter/Source/USDUtilities/Private/USDSkeletalDataConversion.cpp b/Engine/Plugins/Importers/USDImporter/Source/USDUtilities/Private/USDSkeletalDataConversion.cpp
--- a/Engine/Plugins/Importers/USDImporter/Source/USDUtilities/Private/USDSkeletalDataConversion.cpp
+++ b/Engine/Plugins/Importers/USDImporter/Source/USDUtilities/Private/USDSkeletalDataConversion.cpp
@@ -426,21 +426,12 @@ namespace SkelDataConversionImpl
 		{
 			return;
 		}
-
-		const FSmartNameMapping* NameMapping = Skeleton->GetSmartNameContainer( USkeleton::AnimCurveMappingName );
-		if ( !NameMapping )
-		{
-			return;
-		}
-
-		FSmartName NewName;
-		Skeleton->AddSmartNameAndModify( USkeleton::AnimCurveMappingName, CurveName, NewName );
-
+		
 		const bool bShouldTransact = false;
 		const IAnimationDataModel* DataModel = Sequence->GetDataModel();
 		IAnimationDataController& Controller = Sequence->GetController();
 
-		FAnimationCurveIdentifier CurveId(NewName, ERawCurveTrackTypes::RCT_Float);
+		FAnimationCurveIdentifier CurveId(CurveName, ERawCurveTrackTypes::RCT_Float);
 		const FFloatCurve* Curve = DataModel->FindFloatCurve(CurveId);
 		if ( !Curve )
 		{
@@ -461,9 +452,7 @@ namespace SkelDataConversionImpl
 
 			Controller.SetCurveFlags( CurveId, Curve->GetCurveTypeFlags() | AACF_DefaultCurve, bShouldTransact );
 		}
-
-		Controller.UpdateCurveNamesFromSkeleton( Skeleton, ERawCurveTrackTypes::RCT_Float, bShouldTransact );
-
+		
 		if ( Curve )
 		{
 			Controller.SetCurveKeys( CurveId, SourceData.GetConstRefOfKeys(), bShouldTransact );
@@ -2015,13 +2004,6 @@ bool UsdToUnreal::ConvertSkelAnim(
 		return false;
 	}
 
-	// In a regular import workflow this NameMapping will exist and be populated with the blend shape names we imported, if any
-	const FSmartNameMapping* NameMapping = Skeleton->GetSmartNameContainer( USkeleton::AnimCurveMappingName );
-	if ( !NameMapping )
-	{
-		return false;
-	}
-
 	TUsdStore<std::vector<double>> UsdJointTransformTimeSamples;
 	AnimQuery.Get().GetJointTransformTimeSamples( &( UsdJointTransformTimeSamples.Get() ) );
 	int32 NumJointTransformSamples = UsdJointTransformTimeSamples.Get().size();
@@ -3206,89 +3188,39 @@ bool UnrealToUsd::ConvertAnimSequence( UAnimSequence* AnimSequence, pxr::UsdPrim
 		pxr::VtArray< pxr::TfToken > BlendShapeNames;
 		pxr::VtArray< float > BlendShapeWeights;
 
-		const FSmartNameMapping* AnimCurveMapping = AnimSkeleton->GetSmartNameContainer( USkeleton::AnimCurveMappingName );
+		// We need to make sure we have at least one mark on the memstack allocator because FBlendedCurve
+		// will allocate using one and will assert if there aren't any marks yet
+		FMemMark Mark( FMemStack::Get() );
 
-		if ( AnimCurveMapping )
+		// Blend shape weights
+		for ( int32 TimeCode = 0; TimeCode < NumTimeCodes; ++TimeCode )
 		{
-			TArray< FName > AnimCurveNames;
-			AnimCurveMapping->FillNameArray( AnimCurveNames );
+			const float AnimTime = TimeCode / TimeCodesPerSecond;
 
-			TArray< SmartName::UID_Type > AnimCurveUIDs;
-			{
-				TArray< FName > UID_ToNameArray;
-				AnimCurveMapping->FillUIDToNameArray( UID_ToNameArray );
-				AnimCurveUIDs.Reserve( UID_ToNameArray.Num() );
-				for ( int32 NameIndex = 0; NameIndex < UID_ToNameArray.Num(); ++NameIndex )
-				{
-					AnimCurveUIDs.Add( NameIndex );
-				}
-			}
+			FBlendedCurve BlendedCurve;
+			const bool bForceUseRawData = true;
+			AnimSequence->EvaluateCurveData( BlendedCurve, AnimTime, bForceUseRawData );
 
-			// Blend shape names
-			for ( const FName& AnimCurveName : AnimCurveNames )
-			{
-				const FCurveMetaData* CurveMetaData = AnimCurveMapping->GetCurveMetaData( AnimCurveName );
+			BlendShapeNames.clear();
+			BlendShapeNames.reserve( BlendedCurve.Num() );
+			BlendShapeWeights.clear();
+			BlendShapeWeights.reserve( BlendedCurve.Num() );
 
-				if ( CurveMetaData && CurveMetaData->Type.bMorphtarget )
+			BlendedCurve.ForEachElement([&BlendShapeWeights, &BlendShapeNames](const UE::Anim::FCurveElement& InElement)
+			{
+				if (EnumHasAnyFlags(InElement.Flags, UE::Anim::ECurveElementFlags::MorphTarget))
 				{
-					BlendShapeNames.push_back( UnrealToUsd::ConvertToken( *AnimCurveName.ToString() ).Get() );
+					BlendShapeNames.push_back( UnrealToUsd::ConvertToken( *InElement.Name.ToString() ).Get() );
+					BlendShapeWeights.push_back( InElement.Value );
 				}
-			}
+			});
 
-			// We need to make sure we have at least one mark on the memstack allocator because FBlendedCurve
-			// will allocate using one and will assert if there aren't any marks yet
-			FMemMark Mark( FMemStack::Get() );
-
-			// Blend shape weights
-			for ( int32 TimeCode = 0; TimeCode < NumTimeCodes; ++TimeCode )
+			if(BlendShapeWeights.size() > 0 && BlendShapeNames.size() > 0)
 			{
-				BlendShapeWeights.clear();
-				BlendShapeWeights.reserve( AnimCurveNames.Num() );
-
-				const float AnimTime = TimeCode / TimeCodesPerSecond;
-
-				FBlendedCurve BlendedCurve;
-				BlendedCurve.InitFrom( &AnimCurveUIDs );
-				const bool bForceUseRawData = true;
-				AnimSequence->EvaluateCurveData( BlendedCurve, AnimTime, bForceUseRawData );
-
-				if ( BlendedCurve.IsValid() )
-				{
-					for ( const FName& AnimCurveName : AnimCurveNames )
-					{
-						const FCurveMetaData* CurveMetaData = AnimCurveMapping->GetCurveMetaData( AnimCurveName );
-
-						if ( CurveMetaData && CurveMetaData->Type.bMorphtarget )
-						{
-							SmartName::UID_Type NameUID = AnimSkeleton->GetUIDByName( USkeleton::AnimCurveMappingName, AnimCurveName );
-							if ( NameUID != SmartName::MaxUID )
-							{
-								BlendShapeWeights.push_back( BlendedCurve.Get( NameUID ) );
-							}
-						}
-					}
-				};
-
+				UsdSkelAnim.CreateBlendShapesAttr().Set( BlendShapeNames, pxr::UsdTimeCode( TimeCode ) );
 				UsdSkelAnim.CreateBlendShapeWeightsAttr().Set( BlendShapeWeights, pxr::UsdTimeCode( TimeCode ) );
 			}
 		}
-
-		if ( !BlendShapeNames.empty() )
-		{
-			UsdSkelAnim.CreateBlendShapesAttr().Set( BlendShapeNames );
-		}
-		else
-		{
-			if ( pxr::UsdAttribute BlendShapesAttr = UsdSkelAnim.GetBlendShapesAttr() )
-			{
-				BlendShapesAttr.Clear();
-			}
-
-			if ( pxr::UsdAttribute BlendShapeWeightsAttr = UsdSkelAnim.GetBlendShapeWeightsAttr() )
-			{
-				BlendShapeWeightsAttr.Clear();
-			}
-		}
 	}
 
 	// Joints
