Fix Rewind Debugger filtering when the server instance of an actor is shown in the tree view

diff --git a/Engine/Plugins/Animation/GameplayInsights/Source/RewindDebugger/Private/SRewindDebugger.cpp b/Engine/Plugins/Animation/GameplayInsights/Source/RewindDebugger/Private/SRewindDebugger.cpp
--- a/Engine/Plugins/Animation/GameplayInsights/Source/RewindDebugger/Private/SRewindDebugger.cpp
+++ b/Engine/Plugins/Animation/GameplayInsights/Source/RewindDebugger/Private/SRewindDebugger.cpp
@@ -499,9 +499,12 @@ bool FilterTrack(TSharedPtr<RewindDebugger::FRewindDebuggerTrack>& Track, const
 
 void SRewindDebugger::RefreshDebugComponents()
 {
-	if (DebugComponents && DebugComponents->Num() > 0)
+	if (DebugComponents)
 	{
-		FilterTrack((*DebugComponents)[0], TrackFilterBox->GetText().ToString(), !ShouldDisplayEmptyTracks());
+		for(TSharedPtr<RewindDebugger::FRewindDebuggerTrack>& DebugComponent : *DebugComponents)
+		{
+			FilterTrack(DebugComponent, TrackFilterBox->GetText().ToString(), !ShouldDisplayEmptyTracks());
+		}
 	}
 	
 	ComponentTreeView->Refresh();
