Rewind Debugger component tree improvements

diff --git a/Engine/Plugins/Animation/GameplayInsights/Source/RewindDebugger/Private/RewindDebugger.cpp b/Engine/Plugins/Animation/GameplayInsights/Source/RewindDebugger/Private/RewindDebugger.cpp
--- a/Engine/Plugins/Animation/GameplayInsights/Source/RewindDebugger/Private/RewindDebugger.cpp
+++ b/Engine/Plugins/Animation/GameplayInsights/Source/RewindDebugger/Private/RewindDebugger.cpp
@@ -288,73 +288,101 @@ void FRewindDebugger::GetTargetObjectIds(TArray<uint64>& OutTargetObjectIds) con
 #endif
 }
 
+class FRewindDebuggerDummyTrack : public RewindDebugger::FRewindDebuggerTrack
+{
+public:
+
+	FRewindDebuggerDummyTrack(const FName& InObjectName, const FText& InDisplayName)
+		: ObjectName(InObjectName), DisplayName(InDisplayName)
+	{
+	}
 
+private:
+	virtual FName GetNameInternal() const override { return ObjectName; }
+	virtual FText GetDisplayNameInternal() const override { return DisplayName; }
+	
+	FName ObjectName;
+	FText DisplayName;
+};
 
 void FRewindDebugger::RefreshDebugTracks()
 {
+	static const FName DebugMessageTrackName = "DebugMessageDummyTrack";
 	TRACE_CPUPROFILER_EVENT_SCOPE(FRewindDebugger::RefreshDebugTracks);
 
 	if (TargetObjectIds.Num() == 0)
 	{
 		GetTargetObjectIds(TargetObjectIds);
-
-		if (TargetObjectIds.Num() == 0)
-		{
-			if (DebugTracks.Num() != 0)
-			{
-				// clear tracks so we don't show data from previous recordings
-				DebugTracks.SetNum(0);
-				ComponentListChangedDelegate.ExecuteIfBound();
-			}
-			return;
-		}
 	}
-
-	if (const TraceServices::IAnalysisSession* Session = GetAnalysisSession())
+	
+	FString DebugTargetActorName = DebugTargetActor.Get();
+	
+	if (TargetObjectIds.Num() == 0 && !DebugTargetActorName.IsEmpty())
 	{
-		TraceServices::FAnalysisSessionReadScope SessionReadScope(*Session);
-
-		UWorld* World = GetWorldToVisualize();
+		// fallback codepath for when the target object is not found
 
-		if (const IGameplayProvider* GameplayProvider = Session->ReadProvider<IGameplayProvider>("GameplayProvider"))
+		if (DebugTracks.Num() != 2)
 		{
+			// clear tracks so we don't show data from previous recordings
+			DebugTracks.SetNum(0);
+			DebugTracks.SetNum(2);
+		}
 
-			bool bChanged = false;
-
-			// remove any existing tracks that don't match the current list of object ids
-			for (int TrackIndex = DebugTracks.Num()-1; TrackIndex>=0; TrackIndex--)
-			{
-				uint64* FoundId = TargetObjectIds.FindByPredicate([this, TrackIndex](const uint64& TrackId) { return DebugTracks[TrackIndex]->GetObjectId() == TrackId; });
-				if (!FoundId)
-				{
-					DebugTracks.RemoveAt(TrackIndex);
-				}
-			}
-
-			// add new tracks for current list of object ids if they don't already exist
-			for (uint64 TargetObjectId : TargetObjectIds)
-			{
-				TSharedPtr<RewindDebugger::FRewindDebuggerTrack>* FoundTrack = DebugTracks.FindByPredicate([TargetObjectId](const TSharedPtr<RewindDebugger::FRewindDebuggerTrack>& Track) { return Track->GetObjectId() == TargetObjectId; });
-
-				if (!FoundTrack)
-				{
-					DebugTracks.Add(MakeShared<RewindDebugger::FRewindDebuggerObjectTrack>(TargetObjectId, DebugTargetActor.Get(), true));
-					bChanged = true;
-				}
-			}
-
-			// update all tracks
-			for (TSharedPtr<RewindDebugger::FRewindDebuggerTrack>& DebugTrack : DebugTracks )
-			{
-				bChanged = bChanged || DebugTrack->Update();
-			}
-
-			if (bChanged)
-			{
-				ComponentListChangedDelegate.ExecuteIfBound();
-			}
+		if (DebugTracks[1] == nullptr || DebugTracks[0] == nullptr || DebugTracks[0]->GetName().ToString() != DebugTargetActor.Get() )
+		{
+			DebugTracks[0] = MakeShared<FRewindDebuggerDummyTrack>(FName(DebugTargetActorName), FText::FromString(DebugTargetActorName)); 
+			DebugTracks[1] = MakeShared<FRewindDebuggerDummyTrack>(DebugMessageTrackName, NSLOCTEXT("RewindDebugger", "No Debug Data", " - Start a recording to debug"));
+			ComponentListChangedDelegate.ExecuteIfBound();
 		}
 	}
+	else
+	{
+		if (const TraceServices::IAnalysisSession* Session = GetAnalysisSession())
+    	{
+    		TraceServices::FAnalysisSessionReadScope SessionReadScope(*Session);
+    
+    		UWorld* World = GetWorldToVisualize();
+    
+    		if (const IGameplayProvider* GameplayProvider = Session->ReadProvider<IGameplayProvider>("GameplayProvider"))
+    		{
+    
+    			bool bChanged = false;
+    
+    			// remove any existing tracks that don't match the current list of object ids
+    			for (int TrackIndex = DebugTracks.Num()-1; TrackIndex>=0; TrackIndex--)
+    			{
+    				uint64* FoundId = TargetObjectIds.FindByPredicate([this, TrackIndex](const uint64& TrackId) { return DebugTracks[TrackIndex]->GetObjectId() == TrackId; });
+    				if (!FoundId)
+    				{
+    					DebugTracks.RemoveAt(TrackIndex);
+    				}
+    			}
+    
+    			// add new tracks for current list of object ids if they don't already exist
+    			for (uint64 TargetObjectId : TargetObjectIds)
+    			{
+    				TSharedPtr<RewindDebugger::FRewindDebuggerTrack>* FoundTrack = DebugTracks.FindByPredicate([TargetObjectId](const TSharedPtr<RewindDebugger::FRewindDebuggerTrack>& Track) { return Track->GetObjectId() == TargetObjectId; });
+    
+    				if (!FoundTrack)
+    				{
+    					DebugTracks.Add(MakeShared<RewindDebugger::FRewindDebuggerObjectTrack>(TargetObjectId, DebugTargetActor.Get(), true));
+    					bChanged = true;
+    				}
+    			}
+    
+    			// update all tracks
+    			for (TSharedPtr<RewindDebugger::FRewindDebuggerTrack>& DebugTrack : DebugTracks )
+    			{
+    				bChanged = bChanged || DebugTrack->Update();
+    			}
+    
+    			if (bChanged)
+    			{
+    				ComponentListChangedDelegate.ExecuteIfBound();
+    			}
+    		}
+    	}	
+	}
 }
 
 namespace
diff --git a/Engine/Plugins/Animation/GameplayInsights/Source/RewindDebugger/Private/RewindDebuggerObjectTrack.cpp b/Engine/Plugins/Animation/GameplayInsights/Source/RewindDebugger/Private/RewindDebuggerObjectTrack.cpp
--- a/Engine/Plugins/Animation/GameplayInsights/Source/RewindDebugger/Private/RewindDebuggerObjectTrack.cpp
+++ b/Engine/Plugins/Animation/GameplayInsights/Source/RewindDebugger/Private/RewindDebuggerObjectTrack.cpp
@@ -281,7 +281,10 @@ bool FRewindDebuggerObjectTrack::UpdateInternal()
 
 	for (auto& Child : Children)
 	{
-		bChanged = bChanged || Child->Update();
+		if (Child->Update())
+		{
+			bChanged = true;
+		}
 	}
 
 	return bChanged;
diff --git a/Engine/Plugins/Animation/GameplayInsights/Source/RewindDebugger/Private/SRewindDebuggerTimelines.cpp b/Engine/Plugins/Animation/GameplayInsights/Source/RewindDebugger/Private/SRewindDebuggerTimelines.cpp
--- a/Engine/Plugins/Animation/GameplayInsights/Source/RewindDebugger/Private/SRewindDebuggerTimelines.cpp
+++ b/Engine/Plugins/Animation/GameplayInsights/Source/RewindDebugger/Private/SRewindDebuggerTimelines.cpp
@@ -161,10 +161,13 @@ int32 SRewindDebuggerTimelines::OnPaint( const FPaintArgs& Args, const FGeometry
 static void RestoreTimelineExpansion(TSharedPtr<RewindDebugger::FRewindDebuggerTrack> Track, TSharedPtr<STreeView<TSharedPtr<RewindDebugger::FRewindDebuggerTrack>>>& TreeView)
 {
 	TreeView->SetItemExpansion(Track, Track->GetIsExpanded());
-	Track->IterateSubTracks([&TreeView](TSharedPtr<RewindDebugger::FRewindDebuggerTrack> SubTrack)
-		{
-			RestoreTimelineExpansion(SubTrack, TreeView);
-		});
+	if (Track->IsVisible())
+	{
+		Track->IterateSubTracks([&TreeView](TSharedPtr<RewindDebugger::FRewindDebuggerTrack> SubTrack)
+			{
+				RestoreTimelineExpansion(SubTrack, TreeView);
+			});
+	}
 }
 
 void SRewindDebuggerTimelines::RestoreExpansion()
