gameplay insights
- making SCurveTimelineView public
- allow normalization of curves with small numbers in SCurveTimelineView::PaintCurve

diff --git a/Engine/Plugins/Animation/GameplayInsights/Source/GameplayInsights/Private/SCurveTimelineView.cpp b/Engine/Plugins/Animation/GameplayInsights/Source/GameplayInsights/Private/SCurveTimelineView.cpp
--- a/Engine/Plugins/Animation/GameplayInsights/Source/GameplayInsights/Private/SCurveTimelineView.cpp
+++ b/Engine/Plugins/Animation/GameplayInsights/Source/GameplayInsights/Private/SCurveTimelineView.cpp
@@ -151,7 +151,7 @@ int32 SCurveTimelineView::PaintCurve(const FGeometry& AllottedGeometry, const FS
 			}
 		}
 		
-		float Range = FMath::Max(1.0,MaxValue - MinValue);
+		const float Range = FMath::Max(UE_KINDA_SMALL_NUMBER, MaxValue - MinValue);
 		
 		auto &CurvePoints = Curve->Points;
 		const int NumPoints = CurvePoints.Num();
diff --git a/Engine/Plugins/Animation/GameplayInsights/Source/GameplayInsights/Private/SCurveTimelineView.h b/Engine/Plugins/Animation/GameplayInsights/Source/GameplayInsights/Private/SCurveTimelineView.h
deleted file mode 100644
diff --git a/Engine/Plugins/Animation/GameplayInsights/Source/GameplayInsights/Public/SCurveTimelineView.h b/Engine/Plugins/Animation/GameplayInsights/Source/GameplayInsights/Public/SCurveTimelineView.h
new file mode 100644
--- /dev/null
+++ b/Engine/Plugins/Animation/GameplayInsights/Source/GameplayInsights/Public/SCurveTimelineView.h
@@ -0,0 +1,99 @@
+// Copyright Epic Games, Inc. All Rights Reserved.
+
+#pragma once
+
+#include "CoreMinimal.h"
+#include "Widgets/SCompoundWidget.h"
+#include "Widgets/DeclarativeSyntaxSupport.h"
+
+class FPaintArgs;
+class FSlateWindowElementList;
+class FTimeSliderController;
+class SToolTip;
+
+class GAMEPLAYINSIGHTS_API SCurveTimelineView : public SCompoundWidget
+{
+public:
+	struct FTimelineCurveData
+	{
+	public:
+		struct CurvePoint
+		{
+			double Time;
+			float Value;
+		};
+
+		TArray<CurvePoint> Points;
+	};
+	
+	SLATE_BEGIN_ARGS(SCurveTimelineView)
+		: _ViewRange(TRange<double>(0,10))
+		, _DesiredSize(FVector2D(100.f,20.f))
+		, _RenderFill(false)
+	{}
+    	/** View time range */
+     	SLATE_ATTRIBUTE(TRange<double>, ViewRange);
+
+		/** Data for curve to render */
+		SLATE_ATTRIBUTE(TSharedPtr<FTimelineCurveData>, CurveData);
+
+		/** Desired widget size */
+		SLATE_ATTRIBUTE(FVector2D, DesiredSize);
+	
+		/** Curve Color */
+		SLATE_ATTRIBUTE(FLinearColor, CurveColor);
+	
+		/** Fill Color */
+		SLATE_ATTRIBUTE(FLinearColor, FillColor);
+	
+		/** Enable Fill rendering */
+		SLATE_ATTRIBUTE(bool, RenderFill);
+     
+	SLATE_END_ARGS()
+
+
+	/**
+	 * Construct the widget
+	 * 
+	 * @param InArgs   A declaration from which to construct the widget
+	 */
+	void Construct( const FArguments& InArgs );
+
+	// render the curve with a fixed min/max range, rather than using the min/max from data
+	void SetFixedRange(float Min, float Max)
+	{
+		bUseFixedRange = true;
+		FixedRangeMin = Min;
+		FixedRangeMax = Max;
+	}
+
+	FText GetCurveToolTipInputText() const { return CurveToolTipInputText; }
+	FText GetCurveToolTipOutputText() const { return CurveToolTipOutputText; }
+	
+protected:
+	// SWidget interface
+	virtual FVector2D ComputeDesiredSize(float) const override;
+	virtual int32 OnPaint( const FPaintArgs& Args, const FGeometry& AllottedGeometry, const FSlateRect& MyCullingRect, FSlateWindowElementList& OutDrawElements, int32 LayerId, const FWidgetStyle& InWidgetStyle, bool bParentEnabled ) const override;
+	virtual FReply OnMouseMove(const FGeometry& MyGeometry, const FPointerEvent& MouseEvent) override;
+	
+	int32 PaintCurve(const FGeometry& AllottedGeometry, const FSlateRect& MyCullingRect, FSlateWindowElementList& OutDrawElements, int32 LayerId, const FWidgetStyle& InWidgetStyle, bool bParentEnabled ) const;
+	void UpdateCurveToolTip(const FGeometry& MyGeometry, const FPointerEvent& MouseEvent);
+	
+	TAttribute<TRange<double>> ViewRange;
+	TAttribute<TSharedPtr<FTimelineCurveData>> CurveData;
+	TAttribute<FVector2D> DesiredSize;
+	TAttribute<FLinearColor> CurveColor;
+	TAttribute<FLinearColor> FillColor;
+	TAttribute<bool> RenderFill;
+	
+	TSharedPtr<SToolTip> CurveToolTip;	/** The tooltip control for the curve. */
+	FText CurveToolTipInputText;		
+	FText CurveToolTipOutputText;
+	
+	const FSlateBrush* FillBrush;
+
+	bool bUseFixedRange;
+	float FixedRangeMin;
+	float FixedRangeMax;
+};
+
