fix rewind debugger frame stepping not working with client+server pie sessions

diff --git a/Engine/Plugins/Animation/GameplayInsights/Source/GameplayInsights/Private/GameplayProvider.cpp b/Engine/Plugins/Animation/GameplayInsights/Source/GameplayInsights/Private/GameplayProvider.cpp
--- a/Engine/Plugins/Animation/GameplayInsights/Source/GameplayInsights/Private/GameplayProvider.cpp
+++ b/Engine/Plugins/Animation/GameplayInsights/Source/GameplayInsights/Private/GameplayProvider.cpp
@@ -662,22 +662,33 @@ void FGameplayProvider::AppendRecordingInfo(uint64 InWorldId, double InProfileTi
 {
 	Session.WriteAccessCheck();
 
-	FRecordingInfoMessage NewRecordingInfo;
-	NewRecordingInfo.WorldId = InWorldId;
-	NewRecordingInfo.ProfileTime = InProfileTime;
-	NewRecordingInfo.RecordingIndex = InRecordingIndex;
-	NewRecordingInfo.FrameIndex = InFrameIndex;
-	NewRecordingInfo.ElapsedTime = InElapsedTime;
-
-	if(TSharedRef<TraceServices::TPointTimeline<FRecordingInfoMessage>>* ExistingRecording = Recordings.Find(InRecordingIndex))
-	{
-		(*ExistingRecording)->AppendEvent(InProfileTime, NewRecordingInfo);
-	}
-	else
+	if (int32* WorldIndex = WorldIdToIndexMap.Find(InWorldId))
 	{
-		TSharedPtr<TraceServices::TPointTimeline<FRecordingInfoMessage>> NewRecording = MakeShared<TraceServices::TPointTimeline<FRecordingInfoMessage>>(Session.GetLinearAllocator());
-		NewRecording->AppendEvent(InProfileTime, NewRecordingInfo);
-		Recordings.Add(InRecordingIndex, NewRecording.ToSharedRef());
+		if (WorldInfos.IsValidIndex(*WorldIndex))
+		{
+			// Try to only keep track of the client world for now.
+			// We should put each world into its own timeline in future, to support frame stepping in sync with server frames
+			if (WorldInfos[*WorldIndex].NetMode == FWorldInfo::ENetMode::Standalone || WorldInfos[*WorldIndex].NetMode == FWorldInfo::ENetMode::Client)
+			{
+				FRecordingInfoMessage NewRecordingInfo;
+				NewRecordingInfo.WorldId = InWorldId;
+				NewRecordingInfo.ProfileTime = InProfileTime;
+				NewRecordingInfo.RecordingIndex = InRecordingIndex;
+				NewRecordingInfo.FrameIndex = InFrameIndex;
+				NewRecordingInfo.ElapsedTime = InElapsedTime;
+
+				if (TSharedRef<TraceServices::TPointTimeline<FRecordingInfoMessage>>* ExistingRecording = Recordings.Find(InRecordingIndex))
+				{
+					(*ExistingRecording)->AppendEvent(InProfileTime, NewRecordingInfo);
+				}
+				else
+				{
+					TSharedPtr<TraceServices::TPointTimeline<FRecordingInfoMessage>> NewRecording = MakeShared<TraceServices::TPointTimeline<FRecordingInfoMessage>>(Session.GetLinearAllocator());
+					NewRecording->AppendEvent(InProfileTime, NewRecordingInfo);
+					Recordings.Add(InRecordingIndex, NewRecording.ToSharedRef());
+				}
+			}
+		}
 	}
 }
 
