Fix rewind debugger follow target actor mode when there is a server instance of the actor

diff --git a/Engine/Plugins/Animation/GameplayInsights/Source/RewindDebugger/Private/RewindDebugger.cpp b/Engine/Plugins/Animation/GameplayInsights/Source/RewindDebugger/Private/RewindDebugger.cpp
--- a/Engine/Plugins/Animation/GameplayInsights/Source/RewindDebugger/Private/RewindDebugger.cpp
+++ b/Engine/Plugins/Animation/GameplayInsights/Source/RewindDebugger/Private/RewindDebugger.cpp
@@ -870,9 +870,9 @@ void FRewindDebugger::Tick(float DeltaTime)
 						{
 							{
 								TRACE_CPUPROFILER_EVENT_SCOPE(FRewindDebugger::Tick_UpdateActorPosition);
-								// until we have actor transforms traced out, the first skeletal mesh component transform on the target actor be used as as the actor position 
-								uint64 TargetActorId = GetTargetActorId();
-								if (TargetActorId != 0)
+								// until we have actor transforms traced out, the first skeletal mesh component transform on the target actor be used as as the actor position
+
+								for(uint64 TargetActorId : TargetObjectIds)
 								{
 #if OBJECT_TRACE_ENABLED
 									if(UObject* ObjectInstance = FObjectTrace::GetObjectFromId(TargetActorId))
