Added inertialization track to rewind debugger. Made it so that inertialization requests track the origin of the request for debugging purposes.

diff --git a/Engine/Plugins/Animation/GameplayInsights/Source/GameplayInsights/Private/AnimationAnalyzer.cpp b/Engine/Plugins/Animation/GameplayInsights/Source/GameplayInsights/Private/AnimationAnalyzer.cpp
--- a/Engine/Plugins/Animation/GameplayInsights/Source/GameplayInsights/Private/AnimationAnalyzer.cpp
+++ b/Engine/Plugins/Animation/GameplayInsights/Source/GameplayInsights/Private/AnimationAnalyzer.cpp
@@ -207,12 +207,13 @@ bool FAnimationAnalyzer::OnEvent(uint16 RouteId, EStyle Style, const FOnEventCon
 		uint64 AnimInstanceId = EventData.GetValue<uint64>("AnimInstanceId");
 		int32 NodeId = EventData.GetValue<int32>("NodeId");
 		int32 PreviousNodeId = EventData.GetValue<int32>("PreviousNodeId");
+		FString NodeName = FTraceAnalyzerUtils::LegacyAttachmentString<TCHAR>("Name", Context);
 		float Weight = EventData.GetValue<float>("Weight");
 		float RootMotionWeight = EventData.GetValue<float>("RootMotionWeight");
 		uint16 FrameCounter = EventData.GetValue<uint16>("FrameCounter");
 		uint8 Phase = EventData.GetValue<uint8>("Phase");
-		FString TargetNodeName = FTraceAnalyzerUtils::LegacyAttachmentString<TCHAR>("DisplayName", Context);
-		AnimationProvider.AppendAnimNodeStart(AnimInstanceId, Context.EventTime.AsSeconds(StartCycle), FrameCounter, NodeId, PreviousNodeId, Weight, RootMotionWeight, *TargetNodeName, Phase);
+		FString NodeDisplayName = FTraceAnalyzerUtils::LegacyAttachmentString<TCHAR>("DisplayName", Context);
+		AnimationProvider.AppendAnimNodeStart(AnimInstanceId, Context.EventTime.AsSeconds(StartCycle), FrameCounter, NodeId, PreviousNodeId, *NodeName, Weight, RootMotionWeight, *NodeDisplayName, Phase);
 		break;
 	}
 	case RouteId_AnimNodeAttribute:
@@ -236,6 +237,8 @@ bool FAnimationAnalyzer::OnEvent(uint16 RouteId, EStyle Style, const FOnEventCon
 	case RouteId_AnimNodeValueClass:
 	{
 		uint64 Cycle = EventData.GetValue<uint64>("Cycle");
+		double RecordingTime = EventData.GetValue<double>("RecordingTime");
+
 		uint64 AnimInstanceId = EventData.GetValue<uint64>("AnimInstanceId");
 		int32 NodeId = EventData.GetValue<int32>("NodeId");
 		uint16 FrameCounter = EventData.GetValue<uint16>("FrameCounter");
@@ -246,31 +249,31 @@ bool FAnimationAnalyzer::OnEvent(uint16 RouteId, EStyle Style, const FOnEventCon
 		case RouteId_AnimNodeValueBool:
 		{
 			bool Value = EventData.GetValue<bool>("Value");
-			AnimationProvider.AppendAnimNodeValue(AnimInstanceId, Context.EventTime.AsSeconds(Cycle), FrameCounter, NodeId, *Key, Value);
+			AnimationProvider.AppendAnimNodeValue(AnimInstanceId, Context.EventTime.AsSeconds(Cycle), RecordingTime, FrameCounter, NodeId, *Key, Value);
 			break;
 		}
 		case RouteId_AnimNodeValueInt:
 		{
 			int32 Value = EventData.GetValue<int32>("Value");
-			AnimationProvider.AppendAnimNodeValue(AnimInstanceId, Context.EventTime.AsSeconds(Cycle), FrameCounter, NodeId, *Key, Value);
+			AnimationProvider.AppendAnimNodeValue(AnimInstanceId, Context.EventTime.AsSeconds(Cycle), RecordingTime, FrameCounter, NodeId, *Key, Value);
 			break;
 		}
 		case RouteId_AnimNodeValueFloat:
 		{
 			float Value = EventData.GetValue<float>("Value");
-			AnimationProvider.AppendAnimNodeValue(AnimInstanceId, Context.EventTime.AsSeconds(Cycle), FrameCounter, NodeId, *Key, Value);
+			AnimationProvider.AppendAnimNodeValue(AnimInstanceId, Context.EventTime.AsSeconds(Cycle), RecordingTime, FrameCounter, NodeId, *Key, Value);
 			break;
 		}
 		case RouteId_AnimNodeValueVector2D:
 		{
 			FVector2D Value(EventData.GetValue<float>("ValueX"), EventData.GetValue<float>("ValueY"));
-			AnimationProvider.AppendAnimNodeValue(AnimInstanceId, Context.EventTime.AsSeconds(Cycle), FrameCounter, NodeId, *Key, Value);
+			AnimationProvider.AppendAnimNodeValue(AnimInstanceId, Context.EventTime.AsSeconds(Cycle), RecordingTime, FrameCounter, NodeId, *Key, Value);
 			break;
 		}
 		case RouteId_AnimNodeValueVector:
 		{
 			FVector Value(EventData.GetValue<float>("ValueX"), EventData.GetValue<float>("ValueY"), EventData.GetValue<float>("ValueZ"));
-			AnimationProvider.AppendAnimNodeValue(AnimInstanceId, Context.EventTime.AsSeconds(Cycle), FrameCounter, NodeId, *Key, Value);
+			AnimationProvider.AppendAnimNodeValue(AnimInstanceId, Context.EventTime.AsSeconds(Cycle), RecordingTime, FrameCounter, NodeId, *Key, Value);
 			break;
 		}
 		case RouteId_AnimNodeValueString:
@@ -281,19 +284,19 @@ bool FAnimationAnalyzer::OnEvent(uint16 RouteId, EStyle Style, const FOnEventCon
 				int32 KeyLength = EventData.GetValue<uint32>("KeyLength");
 				Value = reinterpret_cast<const TCHAR*>(EventData.GetAttachment()) + KeyLength;
 			}
-			AnimationProvider.AppendAnimNodeValue(AnimInstanceId, Context.EventTime.AsSeconds(Cycle), FrameCounter, NodeId, *Key, *Value);
+			AnimationProvider.AppendAnimNodeValue(AnimInstanceId, Context.EventTime.AsSeconds(Cycle), RecordingTime, FrameCounter, NodeId, *Key, *Value);
 			break;
 		}
 		case RouteId_AnimNodeValueObject:
 		{
 			uint64 Value = EventData.GetValue<uint64>("Value");
-			AnimationProvider.AppendAnimNodeValueObject(AnimInstanceId, Context.EventTime.AsSeconds(Cycle), FrameCounter, NodeId, *Key, Value);
+			AnimationProvider.AppendAnimNodeValueObject(AnimInstanceId, Context.EventTime.AsSeconds(Cycle), RecordingTime, FrameCounter, NodeId, *Key, Value);
 			break;
 		}
 		case RouteId_AnimNodeValueClass:
 		{
 			uint64 Value = EventData.GetValue<uint64>("Value");
-			AnimationProvider.AppendAnimNodeValueClass(AnimInstanceId, Context.EventTime.AsSeconds(Cycle), FrameCounter, NodeId, *Key, Value);
+			AnimationProvider.AppendAnimNodeValueClass(AnimInstanceId, Context.EventTime.AsSeconds(Cycle), RecordingTime, FrameCounter, NodeId, *Key, Value);
 			break;
 		}
 		}
diff --git a/Engine/Plugins/Animation/GameplayInsights/Source/GameplayInsights/Private/AnimationProvider.cpp b/Engine/Plugins/Animation/GameplayInsights/Source/GameplayInsights/Private/AnimationProvider.cpp
--- a/Engine/Plugins/Animation/GameplayInsights/Source/GameplayInsights/Private/AnimationProvider.cpp
+++ b/Engine/Plugins/Animation/GameplayInsights/Source/GameplayInsights/Private/AnimationProvider.cpp
@@ -854,7 +854,7 @@ void FAnimationProvider::AppendAnimGraph(uint64 InAnimInstanceId, double InStart
 	Session.UpdateDurationSeconds(InStartTime);
 }
 
-void FAnimationProvider::AppendAnimNodeStart(uint64 InAnimInstanceId, double InStartTime, uint16 InFrameCounter, int32 InNodeId, int32 PreviousNodeId, float InWeight, float InRootMotionWeight, const TCHAR* InTargetNodeName, uint8 InPhase)
+void FAnimationProvider::AppendAnimNodeStart(uint64 InAnimInstanceId, double InStartTime, uint16 InFrameCounter, int32 InNodeId, int32 PreviousNodeId, const TCHAR* InTargetNodeName, float InWeight, float InRootMotionWeight, const TCHAR* InTargetNodeDisplayName, uint8 InPhase)
 {
 	Session.WriteAccessCheck();
 
@@ -875,7 +875,8 @@ void FAnimationProvider::AppendAnimNodeStart(uint64 InAnimInstanceId, double InS
 	}
 
 	FAnimNodeMessage Message;
-	Message.NodeName = Session.StoreString(InTargetNodeName);
+	Message.NodeName = Session.StoreString(InTargetNodeDisplayName);
+	Message.NodeTypeName = Session.StoreString(InTargetNodeName);
 	Message.AnimInstanceId = InAnimInstanceId;
 	Message.PreviousNodeId = PreviousNodeId;
 	Message.NodeId = InNodeId;
@@ -957,79 +958,79 @@ void FAnimationProvider::AppendBlendSpacePlayer(uint64 InAnimInstanceId, double
 	Session.UpdateDurationSeconds(InTime);
 }
 
-void FAnimationProvider::AppendAnimNodeValue(uint64 InAnimInstanceId, double InTime, uint16 InFrameCounter, int32 InNodeId, const TCHAR* InKey, bool bInValue)
+void FAnimationProvider::AppendAnimNodeValue(uint64 InAnimInstanceId, double InTime, double InRecordingTime, uint16 InFrameCounter, int32 InNodeId, const TCHAR* InKey, bool bInValue)
 {
 	FAnimNodeValueMessage Message;
 	Message.Value.Bool.bValue = bInValue;
 	Message.Value.Type = EAnimNodeValueType::Bool;
 
-	AppendAnimNodeValue(InAnimInstanceId, InTime, InFrameCounter, InNodeId, InKey, Message);
+	AppendAnimNodeValue(InAnimInstanceId, InTime, InRecordingTime, InFrameCounter, InNodeId, InKey, Message);
 }
 
-void FAnimationProvider::AppendAnimNodeValue(uint64 InAnimInstanceId, double InTime, uint16 InFrameCounter, int32 InNodeId, const TCHAR* InKey, int32 InValue)
+void FAnimationProvider::AppendAnimNodeValue(uint64 InAnimInstanceId, double InTime, double InRecordingTime, uint16 InFrameCounter, int32 InNodeId, const TCHAR* InKey, int32 InValue)
 {
 	FAnimNodeValueMessage Message;
 	Message.Value.Int32.Value = InValue;
 	Message.Value.Type = EAnimNodeValueType::Int32;
 
-	AppendAnimNodeValue(InAnimInstanceId, InTime, InFrameCounter, InNodeId, InKey, Message);
+	AppendAnimNodeValue(InAnimInstanceId, InTime, InRecordingTime, InFrameCounter, InNodeId, InKey, Message);
 }
 
-void FAnimationProvider::AppendAnimNodeValue(uint64 InAnimInstanceId, double InTime, uint16 InFrameCounter, int32 InNodeId, const TCHAR* InKey, float InValue)
+void FAnimationProvider::AppendAnimNodeValue(uint64 InAnimInstanceId, double InTime, double InRecordingTime, uint16 InFrameCounter, int32 InNodeId, const TCHAR* InKey, float InValue)
 {
 	FAnimNodeValueMessage Message;
 	Message.Value.Float.Value = InValue;
 	Message.Value.Type = EAnimNodeValueType::Float;
 
-	AppendAnimNodeValue(InAnimInstanceId, InTime, InFrameCounter, InNodeId, InKey, Message);
+	AppendAnimNodeValue(InAnimInstanceId, InTime, InRecordingTime, InFrameCounter, InNodeId, InKey, Message);
 }
 
-void FAnimationProvider::AppendAnimNodeValue(uint64 InAnimInstanceId, double InTime, uint16 InFrameCounter, int32 InNodeId, const TCHAR* InKey, const FVector2D& InValue)
+void FAnimationProvider::AppendAnimNodeValue(uint64 InAnimInstanceId, double InTime, double InRecordingTime, uint16 InFrameCounter, int32 InNodeId, const TCHAR* InKey, const FVector2D& InValue)
 {
 	FAnimNodeValueMessage Message;
 	Message.Value.Vector2D.Value = InValue;
 	Message.Value.Type = EAnimNodeValueType::Vector2D;
 
-	AppendAnimNodeValue(InAnimInstanceId, InTime, InFrameCounter, InNodeId, InKey, Message);
+	AppendAnimNodeValue(InAnimInstanceId, InTime, InRecordingTime, InFrameCounter, InNodeId, InKey, Message);
 }
 
-void FAnimationProvider::AppendAnimNodeValue(uint64 InAnimInstanceId, double InTime, uint16 InFrameCounter, int32 InNodeId, const TCHAR* InKey, const FVector& InValue)
+void FAnimationProvider::AppendAnimNodeValue(uint64 InAnimInstanceId, double InTime, double InRecordingTime, uint16 InFrameCounter, int32 InNodeId, const TCHAR* InKey, const FVector& InValue)
 {
 	FAnimNodeValueMessage Message;
 	Message.Value.Vector.Value = InValue;
 	Message.Value.Type = EAnimNodeValueType::Vector;
 
-	AppendAnimNodeValue(InAnimInstanceId, InTime, InFrameCounter, InNodeId, InKey, Message);
+	AppendAnimNodeValue(InAnimInstanceId, InTime, InRecordingTime, InFrameCounter, InNodeId, InKey, Message);
 }
 
-void FAnimationProvider::AppendAnimNodeValue(uint64 InAnimInstanceId, double InTime, uint16 InFrameCounter, int32 InNodeId, const TCHAR* InKey, const TCHAR* InValue)
+void FAnimationProvider::AppendAnimNodeValue(uint64 InAnimInstanceId, double InTime, double InRecordingTime, uint16 InFrameCounter, int32 InNodeId, const TCHAR* InKey, const TCHAR* InValue)
 {
 	FAnimNodeValueMessage Message;
 	Message.Value.String.Value = Session.StoreString(InValue);
 	Message.Value.Type = EAnimNodeValueType::String;
 
-	AppendAnimNodeValue(InAnimInstanceId, InTime, InFrameCounter, InNodeId, InKey, Message);
+	AppendAnimNodeValue(InAnimInstanceId, InTime, InRecordingTime, InFrameCounter, InNodeId, InKey, Message);
 }
 
-void FAnimationProvider::AppendAnimNodeValueObject(uint64 InAnimInstanceId, double InTime, uint16 InFrameCounter, int32 InNodeId, const TCHAR* InKey, uint64 InValue)
+void FAnimationProvider::AppendAnimNodeValueObject(uint64 InAnimInstanceId, double InTime, double InRecordingTime, uint16 InFrameCounter, int32 InNodeId, const TCHAR* InKey, uint64 InValue)
 {
 	FAnimNodeValueMessage Message;
 	Message.Value.Object.Value = InValue;
 	Message.Value.Type = EAnimNodeValueType::Object;
 
-	AppendAnimNodeValue(InAnimInstanceId, InTime, InFrameCounter, InNodeId, InKey, Message);
+	AppendAnimNodeValue(InAnimInstanceId, InTime, InRecordingTime, InFrameCounter, InNodeId, InKey, Message);
 }
 
-void FAnimationProvider::AppendAnimNodeValueClass(uint64 InAnimInstanceId, double InTime, uint16 InFrameCounter, int32 InNodeId, const TCHAR* InKey, uint64 InValue)
+void FAnimationProvider::AppendAnimNodeValueClass(uint64 InAnimInstanceId, double InTime, double InRecordingTime, uint16 InFrameCounter, int32 InNodeId, const TCHAR* InKey, uint64 InValue)
 {
 	FAnimNodeValueMessage Message;
 	Message.Value.Class.Value = InValue;
 	Message.Value.Type = EAnimNodeValueType::Class;
 
-	AppendAnimNodeValue(InAnimInstanceId, InTime, InFrameCounter, InNodeId, InKey, Message);
+	AppendAnimNodeValue(InAnimInstanceId, InTime, InRecordingTime, InFrameCounter, InNodeId, InKey, Message);
 }
 
-void FAnimationProvider::AppendAnimNodeValue(uint64 InAnimInstanceId, double InTime, uint16 InFrameCounter, int32 InNodeId, const TCHAR* InKey, FAnimNodeValueMessage& InMessage)
+void FAnimationProvider::AppendAnimNodeValue(uint64 InAnimInstanceId, double InTime, double InRecordingTime, uint16 InFrameCounter, int32 InNodeId, const TCHAR* InKey, FAnimNodeValueMessage& InMessage)
 {
 	Session.WriteAccessCheck();
 
@@ -1052,6 +1053,7 @@ void FAnimationProvider::AppendAnimNodeValue(uint64 InAnimInstanceId, double InT
 	InMessage.AnimInstanceId = InAnimInstanceId;
 	InMessage.NodeId = InNodeId;
 	InMessage.FrameCounter = InFrameCounter;
+	InMessage.RecordingTime = InRecordingTime;
 
 	Timeline->AppendEvent(InTime, InMessage);
 
diff --git a/Engine/Plugins/Animation/GameplayInsights/Source/GameplayInsights/Private/AnimationProvider.h b/Engine/Plugins/Animation/GameplayInsights/Source/GameplayInsights/Private/AnimationProvider.h
--- a/Engine/Plugins/Animation/GameplayInsights/Source/GameplayInsights/Private/AnimationProvider.h
+++ b/Engine/Plugins/Animation/GameplayInsights/Source/GameplayInsights/Private/AnimationProvider.h
@@ -77,17 +77,17 @@ public:
 	void AppendAnimGraph(uint64 InAnimInstanceId, double InStartTime, double InEndTime, int32 InNodeCount, uint16 InFrameCounter, uint8 InPhase);
 
 	/** Add an anim node */
-	void AppendAnimNodeStart(uint64 InAnimInstanceId, double InStartTime, uint16 InFrameCounter, int32 InTargetLinkId, int32 InSourceLinkId, float InWeight, float InRootMotionWeight, const TCHAR* InTargetNodeName, uint8 InPhase);
+	void AppendAnimNodeStart(uint64 InAnimInstanceId, double InStartTime, uint16 InFrameCounter, int32 InTargetLinkId, int32 InSourceLinkId, const TCHAR* InTargetNodeName, float InWeight, float InRootMotionWeight, const TCHAR* InTargetNodeDisplayName, uint8 InPhase);
 
 	/** Add anim node values */
-	void AppendAnimNodeValue(uint64 InAnimInstanceId, double InTime, uint16 InFrameCounter, int32 InNodeId, const TCHAR* InKey, bool bInValue);
-	void AppendAnimNodeValue(uint64 InAnimInstanceId, double InTime, uint16 InFrameCounter, int32 InNodeId, const TCHAR* InKey, int32 InValue);
-	void AppendAnimNodeValue(uint64 InAnimInstanceId, double InTime, uint16 InFrameCounter, int32 InNodeId, const TCHAR* InKey, float InValue);
-	void AppendAnimNodeValue(uint64 InAnimInstanceId, double InTime, uint16 InFrameCounter, int32 InNodeId, const TCHAR* InKey, const FVector2D& InValue);
-	void AppendAnimNodeValue(uint64 InAnimInstanceId, double InTime, uint16 InFrameCounter, int32 InNodeId, const TCHAR* InKey, const FVector& InValue);
-	void AppendAnimNodeValue(uint64 InAnimInstanceId, double InTime, uint16 InFrameCounter, int32 InNodeId, const TCHAR* InKey, const TCHAR* InValue);
-	void AppendAnimNodeValueObject(uint64 InAnimInstanceId, double InTime, uint16 InFrameCounter, int32 InNodeId, const TCHAR* InKey, uint64 InValue);
-	void AppendAnimNodeValueClass(uint64 InAnimInstanceId, double InTime, uint16 InFrameCounter, int32 InNodeId, const TCHAR* InKey, uint64 InValue);
+	void AppendAnimNodeValue(uint64 InAnimInstanceId, double InTime, double InRecordingTime, uint16 InFrameCounter, int32 InNodeId, const TCHAR* InKey, bool bInValue);
+	void AppendAnimNodeValue(uint64 InAnimInstanceId, double InTime, double InRecordingTime, uint16 InFrameCounter, int32 InNodeId, const TCHAR* InKey, int32 InValue);
+	void AppendAnimNodeValue(uint64 InAnimInstanceId, double InTime, double InRecordingTime, uint16 InFrameCounter, int32 InNodeId, const TCHAR* InKey, float InValue);
+	void AppendAnimNodeValue(uint64 InAnimInstanceId, double InTime, double InRecordingTime, uint16 InFrameCounter, int32 InNodeId, const TCHAR* InKey, const FVector2D& InValue);
+	void AppendAnimNodeValue(uint64 InAnimInstanceId, double InTime, double InRecordingTime, uint16 InFrameCounter, int32 InNodeId, const TCHAR* InKey, const FVector& InValue);
+	void AppendAnimNodeValue(uint64 InAnimInstanceId, double InTime, double InRecordingTime, uint16 InFrameCounter, int32 InNodeId, const TCHAR* InKey, const TCHAR* InValue);
+	void AppendAnimNodeValueObject(uint64 InAnimInstanceId, double InTime, double InRecordingTime, uint16 InFrameCounter, int32 InNodeId, const TCHAR* InKey, uint64 InValue);
+	void AppendAnimNodeValueClass(uint64 InAnimInstanceId, double InTime, double InRecordingTime, uint16 InFrameCounter, int32 InNodeId, const TCHAR* InKey, uint64 InValue);
 
 	/** Add anim graph attribute */
 	void AppendAnimGraphAttribute(uint64 InSourceAnimInstanceId, uint64 InTargetAnimInstanceId, double InTime, int32 InSourceNodeId, int32 InTargetNodeId, uint32 InAttributeNameId);
@@ -116,7 +116,7 @@ public:
 
 private:
 	/** Add anim node values helper */
-	void AppendAnimNodeValue(uint64 InAnimInstanceId, double InTime, uint16 InFrameCounter, int32 InNodeId, const TCHAR* InKey, FAnimNodeValueMessage& InMessage);
+	void AppendAnimNodeValue(uint64 InAnimInstanceId, double InTime, double InRecordingTime, uint16 InFrameCounter, int32 InNodeId, const TCHAR* InKey, FAnimNodeValueMessage& InMessage);
 
 	/** Handle an object ending play - terminate existing scopes */
 	void HandleObjectEndPlay(uint64 InObjectId, double InTime, const FObjectInfo& InObjectInfo);
diff --git a/Engine/Plugins/Animation/GameplayInsights/Source/GameplayInsights/Private/GameplayInsightsModule.cpp b/Engine/Plugins/Animation/GameplayInsights/Source/GameplayInsights/Private/GameplayInsightsModule.cpp
--- a/Engine/Plugins/Animation/GameplayInsights/Source/GameplayInsights/Private/GameplayInsightsModule.cpp
+++ b/Engine/Plugins/Animation/GameplayInsights/Source/GameplayInsights/Private/GameplayInsightsModule.cpp
@@ -18,6 +18,7 @@
 #include "ObjectPropertyTrace.h"
 #include "SMontageView.h"
 #include "AnimCurvesTrack.h"
+#include "InertializationsTrack.h"
 #include "BlendWeightsTrack.h"
 #include "MontagesTrack.h"
 #include "NotifiesTrack.h"
@@ -79,6 +80,8 @@ void FGameplayInsightsModule::StartupModule()
 	
 	static RewindDebugger::FAnimationCurvesTrackCreator AnimationCurvesTrackCreator;
 	IModularFeatures::Get().RegisterModularFeature(RewindDebugger::IRewindDebuggerTrackCreator::ModularFeatureName, &AnimationCurvesTrackCreator);
+	static RewindDebugger::FInertializationsTrackCreator InertializationsTrackCreator;
+	IModularFeatures::Get().RegisterModularFeature(RewindDebugger::IRewindDebuggerTrackCreator::ModularFeatureName, &InertializationsTrackCreator);
 	static RewindDebugger::FBlendWeightsTrackCreator BlendWeightsTrackCreator;
 	IModularFeatures::Get().RegisterModularFeature(RewindDebugger::IRewindDebuggerTrackCreator::ModularFeatureName, &BlendWeightsTrackCreator);
 	static RewindDebugger::FMontagesTrackCreator MontagesTrackCreator;
diff --git a/Engine/Plugins/Animation/GameplayInsights/Source/GameplayInsights/Private/InertializationsTrack.cpp b/Engine/Plugins/Animation/GameplayInsights/Source/GameplayInsights/Private/InertializationsTrack.cpp
new file mode 100644
--- /dev/null
+++ b/Engine/Plugins/Animation/GameplayInsights/Source/GameplayInsights/Private/InertializationsTrack.cpp
@@ -0,0 +1,250 @@
+// Copyright Epic Games, Inc. All Rights Reserved.
+
+#include "InertializationsTrack.h"
+#include "SCurveTimelineView.h"
+#include "IRewindDebugger.h"
+#include "GameplayProvider.h"
+#include "AnimationProvider.h"
+#include "Styling/SlateIconFinder.h"
+#include "SInertializationDetailsView.h"
+
+#if WITH_EDITOR
+#include "Animation/AnimInstance.h"
+#include "Editor.h"
+#include "Subsystems/AssetEditorSubsystem.h"
+#endif
+
+#define LOCTEXT_NAMESPACE "InertializationsTrack"
+
+namespace RewindDebugger
+{
+
+FInertializationsTrack::FInertializationsTrack(uint64 InObjectId) : ObjectId(InObjectId)
+{
+	Icon = FSlateIcon("EditorStyle", "AnimGraph.Attribute.InertialBlending.Icon", "AnimGraph.Attribute.InertialBlending.Icon");
+}
+
+void FInertializationsTrack::IterateSubTracksInternal(TFunction<void(TSharedPtr<FRewindDebuggerTrack> SubTrack)> IteratorFunction)
+{
+	for(TSharedPtr<FInertializationTrack>& Track : Children)
+	{
+		IteratorFunction(Track);
+	}
+};
+
+bool FInertializationsTrack::UpdateInternal()
+{
+	IRewindDebugger* RewindDebugger = IRewindDebugger::Instance();
+
+	TSortedMap<int32, const TCHAR*, TInlineAllocator<8>> NodeNameMap;
+
+	const TraceServices::IAnalysisSession* AnalysisSession = RewindDebugger->GetAnalysisSession();
+
+	bool bChanged = false;
+	
+	// convert time range to from rewind debugger times to profiler times
+	TRange<double> TraceTimeRange = RewindDebugger->GetCurrentTraceRange();
+	double StartTime = TraceTimeRange.GetLowerBoundValue();
+	double EndTime = TraceTimeRange.GetUpperBoundValue();
+	
+	// count number of unique inertialization nodes in the current time range
+	const FGameplayProvider* GameplayProvider = AnalysisSession->ReadProvider<FGameplayProvider>(FGameplayProvider::ProviderName);
+	const FAnimationProvider* AnimationProvider = AnalysisSession->ReadProvider<FAnimationProvider>(FAnimationProvider::ProviderName);
+
+	if (GameplayProvider && AnimationProvider)
+	{
+		TraceServices::FAnalysisSessionReadScope SessionReadScope(*AnalysisSession);
+
+		AnimationProvider->ReadAnimNodesTimeline(ObjectId, [&NodeNameMap, &AnimationProvider, StartTime, EndTime](const FAnimationProvider::AnimNodesTimeline& InTimeline)
+		{
+			// this isn't very efficient, and it gets called every frame.  will need optimizing
+			InTimeline.EnumerateEvents(StartTime, EndTime, [&NodeNameMap, &AnimationProvider, StartTime, EndTime](double InStartTime, double InEndTime, uint32 InDepth, const FAnimNodeMessage& InMessage)
+			{
+				if (InEndTime > StartTime && InStartTime < EndTime)
+				{
+					for (const TCHAR* InertializationNodeType : { TEXT("AnimNode_DeadBlending"), TEXT("AnimNode_Inertialization") })
+					{
+						if (FCString::Strcmp(InertializationNodeType, InMessage.NodeTypeName) == 0)
+						{
+							NodeNameMap.Add(InMessage.NodeId, InMessage.NodeName);
+							break;
+						}
+					}
+				}
+
+				return TraceServices::EEventEnumerate::Continue;
+			});
+		});
+
+		TArray<int32, TInlineAllocator<8>> NodeIds;
+		NodeNameMap.GetKeys(NodeIds);
+
+		if (Children.Num() != NodeIds.Num())
+		{
+			bChanged = true;
+		}
+
+		Children.SetNum(NodeIds.Num());
+		for(int32 NodeIdx = 0; NodeIdx < NodeIds.Num(); NodeIdx++)
+		{
+			if (!Children[NodeIdx].IsValid())
+			{
+				Children[NodeIdx] = MakeShared<FInertializationTrack>(ObjectId, NodeIds[NodeIdx], FText::FromString(NodeNameMap[NodeIds[NodeIdx]]));
+				bChanged = true;
+			}
+
+			bChanged = bChanged || Children[NodeIdx]->Update();
+		}
+	}
+
+	return bChanged;
+}
+
+TSharedPtr<SWidget> FInertializationsTrack::GetDetailsViewInternal()
+{
+	IRewindDebugger* RewindDebugger = IRewindDebugger::Instance();
+	return SNew(SInertializationDetailsView, ObjectId, RewindDebugger->CurrentTraceTime(), *RewindDebugger->GetAnalysisSession())
+				.CurrentTime_Lambda([RewindDebugger]{ return RewindDebugger->CurrentTraceTime(); });
+}
+
+
+FInertializationTrack::FInertializationTrack(uint64 InObjectId, int32 InNodeId, const FText& Name)
+	: NodeId(InNodeId)
+	, CurveName(Name)
+	, ObjectId(InObjectId)
+{
+	Icon = FSlateIcon("EditorStyle", "AnimGraph.Attribute.InertialBlending.Icon", "AnimGraph.Attribute.InertialBlending.Icon");
+	SetIsExpanded(false);
+}
+
+TSharedPtr<SCurveTimelineView::FTimelineCurveData> FInertializationTrack::GetCurveData() const
+{
+	if (!CurveData.IsValid())
+	{
+		CurveData = MakeShared<SCurveTimelineView::FTimelineCurveData>();
+	}
+	
+	CurvesUpdateRequested++;
+	
+	return CurveData;
+}
+
+bool FInertializationTrack::UpdateInternal()
+{
+	IRewindDebugger* RewindDebugger = IRewindDebugger::Instance();
+	// compute curve points
+	//
+	TRange<double> TraceTimeRange = RewindDebugger->GetCurrentTraceRange();
+	double StartTime = TraceTimeRange.GetLowerBoundValue();
+	double EndTime = TraceTimeRange.GetUpperBoundValue();
+
+	const TraceServices::IAnalysisSession* AnalysisSession = RewindDebugger->GetAnalysisSession();
+	const FGameplayProvider* GameplayProvider = AnalysisSession->ReadProvider<FGameplayProvider>(FGameplayProvider::ProviderName);
+	const FAnimationProvider* AnimationProvider = AnalysisSession->ReadProvider<FAnimationProvider>(FAnimationProvider::ProviderName);
+
+	if (CurvesUpdateRequested > 10 && GameplayProvider && AnimationProvider)
+	{
+		auto& CurvePoints = CurveData->Points;
+		CurvePoints.SetNum(0,false);
+		
+		TraceServices::FAnalysisSessionReadScope SessionReadScope(*AnalysisSession);
+		
+		AnimationProvider->ReadAnimNodeValuesTimeline(ObjectId, [AnalysisSession, StartTime, EndTime, &CurvePoints, this](const FAnimationProvider::AnimNodeValuesTimeline& InTimeline)
+		{
+			InTimeline.EnumerateEvents(StartTime, EndTime, [&CurvePoints, AnalysisSession, this](double InStartTime, double InEndTime, uint32 InDepth, const FAnimNodeValueMessage& InMessage)
+			{
+				if (InMessage.NodeId == NodeId && FCString::Strcmp(InMessage.Key, TEXT("Inertialization Weight")) == 0)
+				{
+					CurvePoints.Add({ InMessage.RecordingTime,	InMessage.Value.Float.Value });
+				}
+				
+				return TraceServices::EEventEnumerate::Continue;
+			});
+		});
+
+		CurvesUpdateRequested = 0;
+	}
+
+	return false;
+}
+
+static FLinearColor MakeInertializationCurveColor(uint32 InSeed, bool bInLine = false)
+{
+	FRandomStream Stream(InSeed);
+	const uint8 Hue = (uint8)(Stream.FRand() * 255.0f);
+	const uint8 SatVal = bInLine ? 196 : 128;
+	return FLinearColor::MakeFromHSV8(Hue, SatVal, SatVal);
+}
+
+TSharedPtr<SWidget> FInertializationTrack::GetTimelineViewInternal()
+{
+	FLinearColor Color;
+	Color = MakeInertializationCurveColor(CityHash32(reinterpret_cast<char*>(&NodeId), 8));
+	Color.A = 0.5f;
+
+	FLinearColor CurveColor = Color;
+	CurveColor.R *= 0.5;
+	CurveColor.G *= 0.5;
+	CurveColor.B *= 0.5;
+
+	TSharedPtr<SCurveTimelineView> CurveTimelineView = SNew(SCurveTimelineView)
+		.FillColor(Color)
+		.CurveColor(CurveColor)
+		.ViewRange_Lambda([]() { return IRewindDebugger::Instance()->GetCurrentViewRange(); })
+		.RenderFill(true)
+		.CurveData_Raw(this, &FInertializationTrack::GetCurveData);
+
+	CurveTimelineView->SetFixedRange(0, 1);
+
+	return CurveTimelineView;
+}
+
+TSharedPtr<SWidget> FInertializationTrack::GetDetailsViewInternal()
+{
+	IRewindDebugger* RewindDebugger = IRewindDebugger::Instance();
+	TSharedPtr<SInertializationDetailsView> InertializationDetailsView = SNew(SInertializationDetailsView, ObjectId, RewindDebugger->CurrentTraceTime(), *RewindDebugger->GetAnalysisSession())
+				.CurrentTime_Lambda([RewindDebugger]{ return RewindDebugger->CurrentTraceTime(); });
+	InertializationDetailsView->SetFilter(NodeId);
+	return InertializationDetailsView;
+}
+
+bool FInertializationTrack::HandleDoubleClickInternal()
+{
+	return false;
+}
+
+FName FInertializationsTrackCreator::GetTargetTypeNameInternal() const
+{
+	static const FName TargetTypeName = "AnimInstance";
+	return TargetTypeName;
+}
+
+FName FInertializationsTrackCreator::GetNameInternal() const
+{
+	static const FName InertializationsName("Inertializations");
+	return InertializationsName;
+}
+
+TSharedPtr<RewindDebugger::FRewindDebuggerTrack> FInertializationsTrackCreator::CreateTrackInternal(uint64 ObjectId) const
+{
+ 	return MakeShared<RewindDebugger::FInertializationsTrack>(ObjectId);
+}
+
+bool FInertializationsTrackCreator::HasDebugInfoInternal(uint64 ObjectId) const
+{
+	const TraceServices::IAnalysisSession* AnalysisSession = IRewindDebugger::Instance()->GetAnalysisSession();
+	TraceServices::FAnalysisSessionReadScope SessionReadScope(*AnalysisSession);
+	bool bHasData = false;
+	if (const FAnimationProvider* AnimationProvider = AnalysisSession->ReadProvider<FAnimationProvider>(FAnimationProvider::ProviderName))
+	{
+		AnimationProvider->ReadAnimGraphTimeline(ObjectId, [&bHasData](const FAnimationProvider::AnimGraphTimeline& InGraphTimeline)
+		{
+			bHasData = true;
+		});
+	}
+	
+	return bHasData;
+}
+
+}
+#undef LOCTEXT_NAMESPACE
diff --git a/Engine/Plugins/Animation/GameplayInsights/Source/GameplayInsights/Private/InertializationsTrack.h b/Engine/Plugins/Animation/GameplayInsights/Source/GameplayInsights/Private/InertializationsTrack.h
new file mode 100644
--- /dev/null
+++ b/Engine/Plugins/Animation/GameplayInsights/Source/GameplayInsights/Private/InertializationsTrack.h
@@ -0,0 +1,75 @@
+// Copyright Epic Games, Inc. All Rights Reserved.
+#pragma once
+
+#include "IAnimationProvider.h"
+#include "IRewindDebugger.h"
+#include "RewindDebuggerTrack.h"
+#include "SCurveTimelineView.h"
+#include "Textures/SlateIcon.h"
+#include "IRewindDebuggerTrackCreator.h"
+
+namespace RewindDebugger
+{
+
+class FInertializationTrack : public FRewindDebuggerTrack
+{
+public:
+	
+	FInertializationTrack(uint64 InObjectId, int32 NodeId, const FText& Name);
+
+	int32 GetNodeId() const { return NodeId; };
+
+private:
+	virtual bool UpdateInternal() override;
+	virtual TSharedPtr<SWidget> GetTimelineViewInternal() override;
+	virtual TSharedPtr<SWidget> GetDetailsViewInternal() override;
+	
+	virtual FName GetNameInternal() const override { return "Inertialization"; }
+	virtual FSlateIcon GetIconInternal() override { return Icon; }
+	virtual FText GetDisplayNameInternal() const override { return CurveName; }
+	virtual uint64 GetObjectIdInternal() const override { return ObjectId; }
+	virtual bool HandleDoubleClickInternal() override;
+
+	TSharedPtr<SCurveTimelineView::FTimelineCurveData> GetCurveData() const;
+	
+	int32 NodeId;
+
+	FSlateIcon Icon;
+	FText CurveName;
+	uint64 ObjectId;
+
+	mutable TSharedPtr<SCurveTimelineView::FTimelineCurveData> CurveData;
+	mutable int CurvesUpdateRequested = 0;
+};
+
+class FInertializationsTrack : public FRewindDebuggerTrack
+{
+public:
+	FInertializationsTrack(uint64 InObjectId);
+
+private:
+	virtual bool UpdateInternal() override;
+	virtual TSharedPtr<SWidget> GetDetailsViewInternal() override;
+	virtual void IterateSubTracksInternal(TFunction<void(TSharedPtr<FRewindDebuggerTrack> SubTrack)> IteratorFunction) override;
+	
+	virtual FName GetNameInternal() const override { return "Inertializations"; }
+	virtual FText GetDisplayNameInternal() const override { return NSLOCTEXT("RewindDebugger", "Inertializations Track Name", "Inertializations"); }
+	virtual uint64 GetObjectIdInternal() const override { return ObjectId; }
+
+	virtual FSlateIcon GetIconInternal() override { return Icon; }
+
+	FSlateIcon Icon;
+	uint64 ObjectId;
+	TArray<TSharedPtr<FInertializationTrack>> Children;
+};
+	
+class FInertializationsTrackCreator : public IRewindDebuggerTrackCreator
+{
+public:
+	virtual FName GetTargetTypeNameInternal() const;
+	virtual FName GetNameInternal() const override;
+	virtual TSharedPtr<FRewindDebuggerTrack> CreateTrackInternal(uint64 ObjectId) const override;
+	virtual bool HasDebugInfoInternal(uint64 ObjectId) const override;
+};
+
+}
\ No newline at end of file
diff --git a/Engine/Plugins/Animation/GameplayInsights/Source/GameplayInsights/Private/SInertializationDetailsView.cpp b/Engine/Plugins/Animation/GameplayInsights/Source/GameplayInsights/Private/SInertializationDetailsView.cpp
new file mode 100644
--- /dev/null
+++ b/Engine/Plugins/Animation/GameplayInsights/Source/GameplayInsights/Private/SInertializationDetailsView.cpp
@@ -0,0 +1,126 @@
+// Copyright Epic Games, Inc. All Rights Reserved.
+
+#include "SInertializationDetailsView.h"
+#include "AnimationProvider.h"
+#include "TraceServices/Model/AnalysisSession.h"
+#include "GameplayProvider.h"
+#include "TraceServices/Model/Frames.h"
+#include "VariantTreeNode.h"
+
+#define LOCTEXT_NAMESPACE "SInertializationDetailsView"
+
+struct FInertializationDetailsNodeItem
+{
+	const TCHAR* Type = nullptr;
+	float ElapsedTime = 0.0f;
+	float Duration = 0.0f;
+	float MaxDuration = 0.0f;
+	float InertializationWeight = 0.0f;
+	bool bActive = false;
+	const TCHAR* Request = nullptr;
+};
+
+void SInertializationDetailsView::GetVariantsAtFrame(const TraceServices::FFrame& InFrame, TArray<TSharedRef<FVariantTreeNode>>& OutVariants) const
+{
+	const FGameplayProvider* GameplayProvider = AnalysisSession->ReadProvider<FGameplayProvider>(FGameplayProvider::ProviderName);
+	const FAnimationProvider* AnimationProvider = AnalysisSession->ReadProvider<FAnimationProvider>(FAnimationProvider::ProviderName);
+
+	TSortedMap<int32, FInertializationDetailsNodeItem, TInlineAllocator<8>> NodeMap;
+
+	if(GameplayProvider && AnimationProvider)
+	{
+		TraceServices::FAnalysisSessionReadScope SessionReadScope(*AnalysisSession);
+
+		// Find all inertialization nodes within window
+
+		AnimationProvider->ReadAnimNodesTimeline(ObjectId, [this, &NodeMap, &AnimationProvider, &InFrame](const FAnimationProvider::AnimNodesTimeline& InTimeline)
+		{
+			InTimeline.EnumerateEvents(InFrame.StartTime, InFrame.EndTime, [this, &NodeMap, &AnimationProvider, &InFrame](double InStartTime, double InEndTime, uint32 InDepth, const FAnimNodeMessage& InMessage)
+			{
+				if (InStartTime >= InFrame.StartTime && InEndTime <= InFrame.EndTime)
+				{
+					if (bFilterSet && InMessage.NodeId != NodeIdFilter)
+					{
+						return TraceServices::EEventEnumerate::Continue;
+					}
+
+					for (const TCHAR* InertializationNodeType : { TEXT("AnimNode_DeadBlending"), TEXT("AnimNode_Inertialization") })
+					{
+						if (FCString::Strcmp(InertializationNodeType, InMessage.NodeTypeName) == 0)
+						{
+							FInertializationDetailsNodeItem& NodeItemRef = NodeMap.FindOrAdd(InMessage.NodeId);
+							NodeItemRef.Type = InMessage.NodeName;
+							break;
+						}
+					}
+				}
+
+				return TraceServices::EEventEnumerate::Continue;
+			});
+		});
+
+		// Fill in Node Trace Details
+
+		AnimationProvider->ReadAnimNodeValuesTimeline(ObjectId, [this, &NodeMap, &AnimationProvider, &InFrame](const FAnimationProvider::AnimNodeValuesTimeline& InTimeline)
+		{
+			InTimeline.EnumerateEvents(InFrame.StartTime, InFrame.EndTime, [this, &NodeMap, &AnimationProvider, &InFrame](double InStartTime, double InEndTime, uint32 InDepth, const FAnimNodeValueMessage& InMessage)
+			{
+				if (InStartTime >= InFrame.StartTime && InEndTime <= InFrame.EndTime)
+				{
+					if (FInertializationDetailsNodeItem* NodeItem = NodeMap.Find(InMessage.NodeId))
+					{
+						if (FCString::Strcmp(InMessage.Key, TEXT("State")) == 0)
+						{
+							NodeItem->bActive = FCString::Strcmp(InMessage.Value.String.Value, TEXT("EInertializationState::Active")) == 0;
+						}
+						else if (FCString::Strcmp(InMessage.Key, TEXT("Elapsed Time")) == 0)
+						{
+							NodeItem->ElapsedTime = InMessage.Value.Float.Value;
+						}
+						else if (FCString::Strcmp(InMessage.Key, TEXT("Duration")) == 0)
+						{
+							NodeItem->Duration = InMessage.Value.Float.Value;
+						}
+						else if (FCString::Strcmp(InMessage.Key, TEXT("Max Duration")) == 0)
+						{
+							NodeItem->MaxDuration = InMessage.Value.Float.Value;
+						}
+						else if (FCString::Strcmp(InMessage.Key, TEXT("Inertialization Weight")) == 0)
+						{
+							NodeItem->InertializationWeight = InMessage.Value.Float.Value;
+						}
+						else if (FCString::Strcmp(InMessage.Key, TEXT("Request")) == 0)
+						{
+							NodeItem->Request = InMessage.Value.String.Value;
+						}
+					}
+				}
+				return TraceServices::EEventEnumerate::Continue;
+			});
+		});
+
+		// Make UI Variant Items
+
+		for (TPair<int32, FInertializationDetailsNodeItem>& NodePair : NodeMap)
+		{
+			if (NodePair.Value.bActive)
+			{
+				TSharedRef<FVariantTreeNode> Header = OutVariants.Add_GetRef(FVariantTreeNode::MakeHeader(NodePair.Value.Type ? FText::FromString(NodePair.Value.Type) : LOCTEXT("UnknownNode", "Unknown Node"), NodePair.Key));
+				Header->AddChild(FVariantTreeNode::MakeFloat(LOCTEXT("ElapsedTime", "Elapsed Time"), NodePair.Value.ElapsedTime));
+				Header->AddChild(FVariantTreeNode::MakeFloat(LOCTEXT("Duration", "Duration"), NodePair.Value.Duration));
+				Header->AddChild(FVariantTreeNode::MakeFloat(LOCTEXT("MaxDuration", "Max Duration"), NodePair.Value.MaxDuration));
+				Header->AddChild(FVariantTreeNode::MakeFloat(LOCTEXT("InertializationWeight", "Inertialization Weight"), NodePair.Value.InertializationWeight));
+				Header->AddChild(FVariantTreeNode::MakeString(LOCTEXT("RequestNode", "Request Node"), NodePair.Value.Request));
+			}
+		}
+	}
+}
+
+static const FName InertializationDetailsName("Inertialization");
+
+FName SInertializationDetailsView::GetName() const
+{
+	return InertializationDetailsName;
+}
+
+#undef LOCTEXT_NAMESPACE
diff --git a/Engine/Plugins/Animation/GameplayInsights/Source/GameplayInsights/Private/SInertializationDetailsView.h b/Engine/Plugins/Animation/GameplayInsights/Source/GameplayInsights/Private/SInertializationDetailsView.h
new file mode 100644
--- /dev/null
+++ b/Engine/Plugins/Animation/GameplayInsights/Source/GameplayInsights/Private/SInertializationDetailsView.h
@@ -0,0 +1,30 @@
+// Copyright Epic Games, Inc. All Rights Reserved.
+
+#pragma once
+
+#include "CoreMinimal.h"
+#include "SPropertiesDebugViewBase.h"
+
+namespace TraceServices { class IAnalysisSession; }
+
+class SInertializationDetailsView : public SPropertiesDebugViewBase
+{
+public:
+	virtual void GetVariantsAtFrame(const TraceServices::FFrame& InFrame, TArray<TSharedRef<FVariantTreeNode>>& OutVariants) const override;
+	virtual FName GetName() const override;
+
+	void SetFilter(uint32 InNodeIdFilter)
+	{
+		NodeIdFilter = InNodeIdFilter;
+		bFilterSet = true;
+	}
+
+	void ClearFilter()
+	{
+		bFilterSet = false;
+	}
+
+private:
+	uint32 NodeIdFilter = 0;
+	bool bFilterSet = false;
+};
\ No newline at end of file
diff --git a/Engine/Plugins/Animation/GameplayInsights/Source/GameplayInsights/Public/IAnimationProvider.h b/Engine/Plugins/Animation/GameplayInsights/Source/GameplayInsights/Public/IAnimationProvider.h
--- a/Engine/Plugins/Animation/GameplayInsights/Source/GameplayInsights/Public/IAnimationProvider.h
+++ b/Engine/Plugins/Animation/GameplayInsights/Source/GameplayInsights/Public/IAnimationProvider.h
@@ -107,6 +107,7 @@ struct FAnimNodeMessage
 	float RootMotionWeight = 0.0f;
 	uint16 FrameCounter = 0;
 	EAnimGraphPhase Phase = EAnimGraphPhase::Initialize;
+	const TCHAR* NodeTypeName = nullptr;
 };
 
 enum class EAnimNodeValueType : uint8
@@ -191,6 +192,7 @@ struct FAnimNodeValueMessage
 	FVariantValue Value;
 	int32 NodeId = -1;
 	uint16 FrameCounter = 0;
+	double RecordingTime = 0.0f;
 };
 
 struct FAnimSequencePlayerMessage
diff --git a/Engine/Plugins/Experimental/Animation/PoseSearch/Source/Runtime/Private/PoseSearchLibrary.cpp b/Engine/Plugins/Experimental/Animation/PoseSearch/Source/Runtime/Private/PoseSearchLibrary.cpp
--- a/Engine/Plugins/Experimental/Animation/PoseSearch/Source/Runtime/Private/PoseSearchLibrary.cpp
+++ b/Engine/Plugins/Experimental/Animation/PoseSearch/Source/Runtime/Private/PoseSearchLibrary.cpp
@@ -11,6 +11,7 @@
 #include "Animation/AnimSubsystem_Tag.h"
 #include "Animation/BlendSpace.h"
 #include "Animation/BuiltInAttributeTypes.h"
+#include "Animation/AnimTrace.h"
 #include "InstancedStruct.h"
 #include "PoseSearch/AnimNode_MotionMatching.h"
 #include "PoseSearch/AnimNode_PoseSearchHistoryCollector.h"
@@ -62,7 +63,13 @@ static void RequestInertialBlend(const FAnimationUpdateContext& Context, float B
 		UE::Anim::IInertializationRequester* InertializationRequester = Context.GetMessage<UE::Anim::IInertializationRequester>();
 		if (InertializationRequester)
 		{
-			InertializationRequester->RequestInertialization(BlendTime);
+			FInertializationRequest Request;
+			Request.Duration = BlendTime;
+#if ANIM_TRACE_ENABLED
+			Request.Description = LOCTEXT("InertializationRequestDescription", "Motion Matching");
+#endif
+
+			InertializationRequester->RequestInertialization(Request);
 		}
 	}
 }
diff --git a/Engine/Source/Runtime/AnimGraphRuntime/Private/AnimNodes/AnimNode_BlendListBase.cpp b/Engine/Source/Runtime/AnimGraphRuntime/Private/AnimNodes/AnimNode_BlendListBase.cpp
--- a/Engine/Source/Runtime/AnimGraphRuntime/Private/AnimNodes/AnimNode_BlendListBase.cpp
+++ b/Engine/Source/Runtime/AnimGraphRuntime/Private/AnimNodes/AnimNode_BlendListBase.cpp
@@ -125,7 +125,17 @@ void FAnimNode_BlendListBase::Update_AnyThread(const FAnimationUpdateContext& Co
 				UE::Anim::IInertializationRequester* InertializationRequester = Context.GetMessage<UE::Anim::IInertializationRequester>();
 				if (InertializationRequester)
 				{
-					InertializationRequester->RequestInertializationWithBlendMode(CurrentBlendTimes[ChildIndex], CurrentBlendProfile, true, GetBlendType(), GetCustomBlendCurve());
+					FInertializationRequest Request;
+					Request.Duration = CurrentBlendTimes[ChildIndex];
+					Request.BlendProfile = CurrentBlendProfile;
+					Request.bUseBlendMode = true;
+					Request.BlendMode = GetBlendType();
+					Request.CustomBlendCurve = GetCustomBlendCurve();
+#if ANIM_TRACE_ENABLED
+					Request.Description = NSLOCTEXT("AnimNode_BlendListBase", "InertializationRequestDescription", "Blend List");
+#endif
+
+					InertializationRequester->RequestInertialization(Request);
 					InertializationRequester->AddDebugRecord(*Context.AnimInstanceProxy, Context.GetCurrentNodeId());
 				}
 				else
diff --git a/Engine/Source/Runtime/AnimGraphRuntime/Private/AnimNodes/AnimNode_Mirror.cpp b/Engine/Source/Runtime/AnimGraphRuntime/Private/AnimNodes/AnimNode_Mirror.cpp
--- a/Engine/Source/Runtime/AnimGraphRuntime/Private/AnimNodes/AnimNode_Mirror.cpp
+++ b/Engine/Source/Runtime/AnimGraphRuntime/Private/AnimNodes/AnimNode_Mirror.cpp
@@ -143,7 +143,13 @@ void FAnimNode_MirrorBase::Update_AnyThread(const FAnimationUpdateContext& Conte
 			UE::Anim::IInertializationRequester* InertializationRequester = Context.GetMessage<UE::Anim::IInertializationRequester>();
 			if (InertializationRequester)
 			{
-				InertializationRequester->RequestInertialization(GetBlendTimeOnMirrorStateChange());
+				FInertializationRequest Request;
+				Request.Duration = GetBlendTimeOnMirrorStateChange();
+#if ANIM_TRACE_ENABLED
+				Request.Description = LOCTEXT("InertializationRequestDescription", "Mirror");
+#endif
+
+				InertializationRequester->RequestInertialization(Request);
 				InertializationRequester->AddDebugRecord(*Context.AnimInstanceProxy, Context.GetCurrentNodeId());
 			}
 			else
diff --git a/Engine/Source/Runtime/AnimGraphRuntime/Private/AnimNodes/AnimNode_Slot.cpp b/Engine/Source/Runtime/AnimGraphRuntime/Private/AnimNodes/AnimNode_Slot.cpp
--- a/Engine/Source/Runtime/AnimGraphRuntime/Private/AnimNodes/AnimNode_Slot.cpp
+++ b/Engine/Source/Runtime/AnimGraphRuntime/Private/AnimNodes/AnimNode_Slot.cpp
@@ -48,7 +48,14 @@ void FAnimNode_Slot::Update_AnyThread(const FAnimationUpdateContext& Context)
 		UE::Anim::IInertializationRequester* InertializationRequester = Context.GetMessage<UE::Anim::IInertializationRequester>();
 		if (InertializationRequester)
 		{
-			InertializationRequester->RequestInertialization(InertializationRequest.Get<0>(), InertializationRequest.Get<1>());
+			FInertializationRequest Request;
+			Request.Duration = InertializationRequest.Get<0>();
+			Request.BlendProfile = InertializationRequest.Get<1>();
+#if ANIM_TRACE_ENABLED
+			Request.Description = NSLOCTEXT("AnimNode_Slot", "InertializationRequestDescription", "Slot");
+#endif
+
+			InertializationRequester->RequestInertialization(Request);
 		}
 		else
 		{
diff --git a/Engine/Source/Runtime/AnimGraphRuntime/Private/BlendSpacePlayerLibrary.cpp b/Engine/Source/Runtime/AnimGraphRuntime/Private/BlendSpacePlayerLibrary.cpp
--- a/Engine/Source/Runtime/AnimGraphRuntime/Private/BlendSpacePlayerLibrary.cpp
+++ b/Engine/Source/Runtime/AnimGraphRuntime/Private/BlendSpacePlayerLibrary.cpp
@@ -3,6 +3,7 @@
 #include "BlendSpacePlayerLibrary.h"
 #include "Animation/AnimNode_Inertialization.h"
 #include "AnimNodes/AnimNode_BlendSpacePlayer.h"
+#include "Animation/AnimTrace.h"
 
 #include UE_INLINE_GENERATED_CPP_BY_NAME(BlendSpacePlayerLibrary)
 
@@ -52,7 +53,13 @@ FBlendSpacePlayerReference UBlendSpacePlayerLibrary::SetBlendSpaceWithInertialBl
 				{
 					if (UE::Anim::IInertializationRequester* InertializationRequester = AnimationUpdateContext->GetMessage<UE::Anim::IInertializationRequester>())
 					{
-						InertializationRequester->RequestInertialization(BlendTime);
+						FInertializationRequest Request;
+						Request.Duration = BlendTime;
+#if ANIM_TRACE_ENABLED
+						Request.Description = NSLOCTEXT("BlendSpacePlayerLibrary", "InertializationRequestDescription", "Blend Space Player");
+#endif
+
+						InertializationRequester->RequestInertialization(Request);
 					}
 				}
 				else
diff --git a/Engine/Source/Runtime/AnimGraphRuntime/Private/SequenceEvaluatorLibrary.cpp b/Engine/Source/Runtime/AnimGraphRuntime/Private/SequenceEvaluatorLibrary.cpp
--- a/Engine/Source/Runtime/AnimGraphRuntime/Private/SequenceEvaluatorLibrary.cpp
+++ b/Engine/Source/Runtime/AnimGraphRuntime/Private/SequenceEvaluatorLibrary.cpp
@@ -5,6 +5,7 @@
 #include "Animation/AnimNode_Inertialization.h"
 #include "AnimNodes/AnimNode_SequenceEvaluator.h"
 #include "AnimationRuntime.h"
+#include "Animation/AnimTrace.h"
 
 #include UE_INLINE_GENERATED_CPP_BY_NAME(SequenceEvaluatorLibrary)
 
@@ -90,7 +91,13 @@ FSequenceEvaluatorReference USequenceEvaluatorLibrary::SetSequenceWithInertialBl
 				{
 					if (UE::Anim::IInertializationRequester* InertializationRequester = AnimationUpdateContext->GetMessage<UE::Anim::IInertializationRequester>())
 					{
-						InertializationRequester->RequestInertialization(BlendTime);
+						FInertializationRequest Request;
+						Request.Duration = BlendTime;
+#if ANIM_TRACE_ENABLED
+						Request.Description = NSLOCTEXT("SequenceEvaluatorLibrary", "InertializationRequestDescription", "Sequence Evaluator");
+#endif
+
+						InertializationRequester->RequestInertialization(Request);
 					}
 				}
 				else
diff --git a/Engine/Source/Runtime/AnimGraphRuntime/Private/SequencePlayerLibrary.cpp b/Engine/Source/Runtime/AnimGraphRuntime/Private/SequencePlayerLibrary.cpp
--- a/Engine/Source/Runtime/AnimGraphRuntime/Private/SequencePlayerLibrary.cpp
+++ b/Engine/Source/Runtime/AnimGraphRuntime/Private/SequencePlayerLibrary.cpp
@@ -4,6 +4,7 @@
 
 #include "Animation/AnimNode_Inertialization.h"
 #include "Animation/AnimNode_SequencePlayer.h"
+#include "Animation/AnimTrace.h"
 
 #include UE_INLINE_GENERATED_CPP_BY_NAME(SequencePlayerLibrary)
 
@@ -91,7 +92,13 @@ FSequencePlayerReference USequencePlayerLibrary::SetSequenceWithInertialBlending
 				{
 					if (UE::Anim::IInertializationRequester* InertializationRequester = AnimationUpdateContext->GetMessage<UE::Anim::IInertializationRequester>())
 					{
-						InertializationRequester->RequestInertialization(BlendTime);
+						FInertializationRequest Request;
+						Request.Duration = BlendTime;
+#if ANIM_TRACE_ENABLED
+						Request.Description = NSLOCTEXT("SequencePlayerLibrary", "InertializationRequestDescription", "Sequence Player");
+#endif
+
+						InertializationRequester->RequestInertialization(Request);
 					}
 				}
 				else
diff --git a/Engine/Source/Runtime/Engine/Classes/Animation/AnimNode_DeadBlending.h b/Engine/Source/Runtime/Engine/Classes/Animation/AnimNode_DeadBlending.h
--- a/Engine/Source/Runtime/Engine/Classes/Animation/AnimNode_DeadBlending.h
+++ b/Engine/Source/Runtime/Engine/Classes/Animation/AnimNode_DeadBlending.h
@@ -103,15 +103,10 @@ public: // FAnimNode_DeadBlending
 	FAnimNode_DeadBlending();
 
 	/**
-	 * Request to activate inertialization. If multiple requests are made on the same inertialization node, the minimum 
-	 * requested time will be used. If multiple blend modes are given, the blend mode of the final request will be used.
+	 * Request to activate inertialization. If multiple requests are made on the same inertialization node, the request 
+	 * with the minimum blend time will be used.
 	 */
-	virtual void RequestInertialization(
-		float Duration,
-		const UBlendProfile* BlendProfile,
-		const bool bInUseBlendMode,
-		const EAlphaBlendOption InBlendMode,
-		UCurveFloat* InCustomBlendCurve);
+	virtual void RequestInertialization(const FInertializationRequest& Request);
 
 public: // FAnimNode_Base
 
@@ -203,6 +198,9 @@ private:
 	// Current inertialization duration (used for curves).
 	float InertializationDuration = 0.0f;
 
+	// Description for the current inertialization request - used for debugging
+	FText InertializationRequestDescription;
+
 	// Current inertialization durations for each bone, indexed by skeleton bone index (used for per-bone blending).
 	TCustomBoneIndexArray<float, FSkeletonPoseBoneIndex> InertializationDurationPerBone;
 
diff --git a/Engine/Source/Runtime/Engine/Classes/Animation/AnimNode_Inertialization.h b/Engine/Source/Runtime/Engine/Classes/Animation/AnimNode_Inertialization.h
--- a/Engine/Source/Runtime/Engine/Classes/Animation/AnimNode_Inertialization.h
+++ b/Engine/Source/Runtime/Engine/Classes/Animation/AnimNode_Inertialization.h
@@ -34,21 +34,9 @@ public:
 	// If multiple requests are made on the same inertialization node, the minimum requested time will be used.
 	virtual void RequestInertialization(float InRequestedDuration, const UBlendProfile* InBlendProfile = nullptr) = 0;
 
-	// Request to activate inertialization for a duration with the given blend mode.
+	// Request to activate inertialization.
 	// If multiple requests are made on the same inertialization node, the minimum requested time will be used.
-	// 
-	// InRequestedDuration				Duration of the blend
-	// InBlendProfile					Blend profile to use
-	// bUseBlendMode					If to use the provided Blend Mode
-	// InBlendMode						Blend mode to use
-	// InCustomBlendCurve				Custom Blend Curve to use for the Blend Mode
-	//
-	virtual void RequestInertializationWithBlendMode(
-		float InRequestedDuration,
-		const UBlendProfile* InBlendProfile,
-		const bool bUseBlendMode,
-		const EAlphaBlendOption InBlendMode,
-		UCurveFloat* InCustomBlendCurve) = 0;
+	virtual void RequestInertialization(const FInertializationRequest& InInertializationRequest);
 
 	// Add a record of this request
 	virtual void AddDebugRecord(const FAnimInstanceProxy& InSourceProxy, int32 InSourceNodeId) = 0;
@@ -133,21 +121,6 @@ struct FInertializationRequest
 	{
 	}
 
-	FInertializationRequest(
-		float InDuration,
-		const UBlendProfile* InBlendProfile,
-		const bool bInUseBlendMode,
-		const EAlphaBlendOption InBlendMode,
-		UCurveFloat* InCustomBlendCurve)
-		: Duration(InDuration)
-		, BlendProfile(InBlendProfile)
-		, bUseBlendMode(bInUseBlendMode)
-		, BlendMode(InBlendMode)
-		, CustomBlendCurve(InCustomBlendCurve)
-	{
-	}
-
-
 	void Clear()
 	{
 		Duration = -1.0f;
@@ -155,6 +128,7 @@ struct FInertializationRequest
 		bUseBlendMode = false;
 		BlendMode = EAlphaBlendOption::Linear;
 		CustomBlendCurve = nullptr;
+		Description = FText::GetEmpty();
 	}
 
 	friend bool operator==(const FInertializationRequest& A, const FInertializationRequest& B)
@@ -164,7 +138,8 @@ struct FInertializationRequest
 			(A.BlendProfile == B.BlendProfile) &&
 			(A.bUseBlendMode == B.bUseBlendMode) &&
 			(A.BlendMode == B.BlendMode) &&
-			(A.CustomBlendCurve == B.CustomBlendCurve);
+			(A.CustomBlendCurve == B.CustomBlendCurve) &&
+			(A.Description.EqualTo(B.Description));
 	}
 
 	friend bool operator!=(const FInertializationRequest& A, const FInertializationRequest& B)
@@ -172,20 +147,29 @@ struct FInertializationRequest
 		return !(A == B);
 	}
 
+	// Blend duration of the inertialization request.
 	UPROPERTY(Transient)
 	float Duration = -1.0f;
 
+	// Blend profile to control per-joint blend times.
 	UPROPERTY(Transient)
 	TObjectPtr<const UBlendProfile> BlendProfile = nullptr;
 
+	// If to use the provided blend mode.
 	UPROPERTY(Transient)
 	bool bUseBlendMode = false;
 
+	// Blend mode to use.
 	UPROPERTY(Transient)
 	EAlphaBlendOption BlendMode = EAlphaBlendOption::Linear;
 
+	// Custom blend curve to use when use of the blend mode is active.
 	UPROPERTY(Transient)
 	TObjectPtr<UCurveFloat> CustomBlendCurve = nullptr;
+
+	// Description of the request - used for debugging.
+	UPROPERTY(Transient)
+	FText Description;
 };
 
 
@@ -385,6 +369,11 @@ public: // FAnimNode_Inertialization
 	//
 	virtual void RequestInertialization(float Duration, const UBlendProfile* BlendProfile);
 
+	// Request to activate inertialization.
+	// If multiple requests are made on the same inertialization node, the minimum requested time will be used.
+	//
+	virtual void RequestInertialization(const FInertializationRequest& InertializationRequest);
+
 	// Log an error when a node wants to inertialize but no inertialization ancestor node exists
 	//
 	static void LogRequestError(const FAnimationUpdateContext& Context, const FPoseLinkBase& RequesterPoseLink);
@@ -440,6 +429,9 @@ private:
 	// Inertialization duration for the main inertialization request (used for curve blending and deficit tracking)
 	float InertializationDuration;
 
+	// Description for the current inertialization request - used for debugging
+	FText InertializationRequestDescription;
+
 	// Inertialization durations indexed by skeleton bone index (used for per-bone blending)
 	TCustomBoneIndexArray<float, FSkeletonPoseBoneIndex> InertializationDurationPerBone;
 
diff --git a/Engine/Source/Runtime/Engine/Private/Animation/AnimNode_DeadBlending.cpp b/Engine/Source/Runtime/Engine/Private/Animation/AnimNode_DeadBlending.cpp
--- a/Engine/Source/Runtime/Engine/Private/Animation/AnimNode_DeadBlending.cpp
+++ b/Engine/Source/Runtime/Engine/Private/Animation/AnimNode_DeadBlending.cpp
@@ -34,17 +34,12 @@ namespace UE::Anim {
 			float InRequestedDuration,
 			const UBlendProfile* InBlendProfile) override
 		{
-			Node.RequestInertialization(InRequestedDuration, InBlendProfile, false, EAlphaBlendOption::Linear, nullptr);
+			Node.RequestInertialization(FInertializationRequest(InRequestedDuration, InBlendProfile));
 		}
 
-		virtual void RequestInertializationWithBlendMode(
-			float InRequestedDuration,
-			const UBlendProfile* InBlendProfile,
-			const bool bUseBlendMode,
-			const EAlphaBlendOption InBlendMode,
-			UCurveFloat* InCustomBlendCurve)
+		virtual void RequestInertialization(const FInertializationRequest& Request)
 		{
-			Node.RequestInertialization(InRequestedDuration, InBlendProfile, bUseBlendMode, InBlendMode, InCustomBlendCurve);
+			Node.RequestInertialization(Request);
 		}
 
 		virtual void AddDebugRecord(const FAnimInstanceProxy& InSourceProxy, int32 InSourceNodeId)
@@ -499,16 +494,11 @@ void FAnimNode_DeadBlending::ApplyTo(FCompactPose& InOutPose, FBlendedCurve& InO
 
 FAnimNode_DeadBlending::FAnimNode_DeadBlending() {}
 
-void FAnimNode_DeadBlending::RequestInertialization(
-	float Duration,
-	const UBlendProfile* BlendProfile,
-	const bool bInUseBlendMode,
-	const EAlphaBlendOption InBlendMode,
-	UCurveFloat* InCustomBlendCurve)
+void FAnimNode_DeadBlending::RequestInertialization(const FInertializationRequest& Request)
 {
-	if (Duration >= 0.0f)
+	if (Request.Duration >= 0.0f)
 	{
-		RequestQueue.AddUnique(FInertializationRequest(Duration, BlendProfile, bInUseBlendMode, InBlendMode, InCustomBlendCurve));
+		RequestQueue.AddUnique(Request);
 	}
 }
 
@@ -598,7 +588,7 @@ void FAnimNode_DeadBlending::Update_AnyThread(const FAnimationUpdateContext& Con
 				{
 					for (const FInertializationRequest& Request : RequestQueue)
 					{
-						InMessage.RequestInertializationWithBlendMode(Request.Duration, Request.BlendProfile, Request.bUseBlendMode, Request.BlendMode, Request.CustomBlendCurve);
+						InMessage.RequestInertialization(Request);
 					}
 					InMessage.AddDebugRecord(Proxy, NodeId);
 
@@ -666,70 +656,69 @@ void FAnimNode_DeadBlending::Evaluate_AnyThread(FPoseContext& Output)
 	{
 		const int32 NumSkeletonBones = UE::Anim::DeadBlending::Private::GetNumSkeletonBones(Output.AnimInstanceProxy->GetRequiredBones());
 
-		InertializationTime = 0.0f;
+		// Find shortest request
 
-		if (bAlwaysUseDefaultBlendSettings)
+		int32 ShortestRequestIdx = INDEX_NONE;
+		float ShortestRequestDuration = UE_MAX_FLT;
+		for (int32 RequestIdx = 0; RequestIdx < RequestQueue.Num(); RequestIdx++)
 		{
-			InertializationDuration = BlendTimeMultiplier * DefaultBlendDuration;
-			InertializationDurationPerBone.Init(BlendTimeMultiplier * DefaultBlendDuration, NumSkeletonBones);
-			InertializationMaxDuration = BlendTimeMultiplier * DefaultBlendDuration;
-			InertializationBlendMode = DefaultBlendMode;
-			InertializationCustomBlendCurve = DefaultCustomBlendCurve;
+			if (RequestQueue[RequestIdx].Duration < ShortestRequestDuration)
+			{
+				ShortestRequestIdx = RequestIdx;
+				ShortestRequestDuration = RequestQueue[RequestIdx].Duration;
+			}
 		}
-		else
-		{
-			// Process Request Durations by taking min of all requests
-			// For Blend Mode and Custom Curve we will just take whichever
-			// request is last in the stack (if they are provided).
 
-			InertializationDuration = UE_MAX_FLT;
-			InertializationDurationPerBone.Init(UE_MAX_FLT, NumSkeletonBones);
-			InertializationMaxDuration = UE_MAX_FLT;
-			InertializationBlendMode = DefaultBlendMode;
-			InertializationCustomBlendCurve = DefaultCustomBlendCurve;
+		// Record Request
 
-			UE::Anim::TTypedIndexArray<FSkeletonPoseBoneIndex, float, FAnimStackAllocator> RequestDurationPerBone;
+		InertializationTime = 0.0f;
+		InertializationDuration = BlendTimeMultiplier * RequestQueue[ShortestRequestIdx].Duration;
 
-			for (const FInertializationRequest Request : RequestQueue)
-			{
-				// Duration is min of requests
+		UE::Anim::TTypedIndexArray<FSkeletonPoseBoneIndex, float, FAnimStackAllocator> RequestDurationPerBone;
 
-				InertializationDuration = BlendTimeMultiplier * FMath::Min(InertializationDuration, Request.Duration);
+		if (RequestQueue[ShortestRequestIdx].BlendProfile)
+		{
+			RequestQueue[ShortestRequestIdx].BlendProfile->FillSkeletonBoneDurationsArray(RequestDurationPerBone, InertializationDuration);
+		}
+		else if (DefaultBlendProfile)
+		{
+			DefaultBlendProfile->FillSkeletonBoneDurationsArray(RequestDurationPerBone, InertializationDuration);
+		}
+		else
+		{
+			RequestDurationPerBone.Init(InertializationDuration, NumSkeletonBones);
+		}
 
-				// Per-bone durations as min of requests accounting for blend profile
+		InertializationDurationPerBone.SetNumUninitialized(NumSkeletonBones);
+		for (int32 BoneIndex = 0; BoneIndex < NumSkeletonBones; BoneIndex++)
+		{
+			InertializationDurationPerBone[BoneIndex] = RequestDurationPerBone[BoneIndex];
+		}
 
-				if (Request.BlendProfile)
-				{
-					Request.BlendProfile->FillSkeletonBoneDurationsArray(RequestDurationPerBone, Request.Duration);
-				}
-				else if (DefaultBlendProfile)
-				{
-					DefaultBlendProfile->FillSkeletonBoneDurationsArray(RequestDurationPerBone, Request.Duration);
-				}
-				else
-				{
-					RequestDurationPerBone.Init(Request.Duration, NumSkeletonBones);
-				}
+		if (RequestQueue[ShortestRequestIdx].bUseBlendMode)
+		{
+			InertializationBlendMode = RequestQueue[ShortestRequestIdx].BlendMode;
+			InertializationCustomBlendCurve = RequestQueue[ShortestRequestIdx].CustomBlendCurve;
+		}
 
-				check(RequestDurationPerBone.Num() == InertializationDurationPerBone.Num());
+		// Cache the maximum duration across all bones (so we know when to deactivate the inertialization request)
+		InertializationMaxDuration = FMath::Max(InertializationDuration, *Algo::MaxElement(InertializationDurationPerBone));
 
-				for (int32 BoneIndex = 0; BoneIndex < NumSkeletonBones; BoneIndex++)
-				{
-					InertializationDurationPerBone[BoneIndex] = BlendTimeMultiplier * FMath::Min(InertializationDurationPerBone[BoneIndex], RequestDurationPerBone[BoneIndex]);
-				}
+#if ANIM_TRACE_ENABLED
+		InertializationRequestDescription = RequestQueue[ShortestRequestIdx].Description;
+#endif
 
-				// Process blend mode - will take the last one given
+		// Override with defaults
 
-				if (Request.bUseBlendMode)
-				{
-					InertializationBlendMode = Request.BlendMode;
-					InertializationCustomBlendCurve = Request.CustomBlendCurve;
-				}
-			}
+		if (bAlwaysUseDefaultBlendSettings)
+		{
+			InertializationDuration = BlendTimeMultiplier * DefaultBlendDuration;
+			InertializationDurationPerBone.Init(BlendTimeMultiplier * DefaultBlendDuration, NumSkeletonBones);
+			InertializationMaxDuration = BlendTimeMultiplier * DefaultBlendDuration;
+			InertializationBlendMode = DefaultBlendMode;
+			InertializationCustomBlendCurve = DefaultCustomBlendCurve;
 		}
 
-		InertializationMaxDuration = FMath::Max(InertializationDuration, *Algo::MaxElement(InertializationDurationPerBone));
-
 		check(InertializationDuration != UE_MAX_FLT);
 		check(InertializationMaxDuration != UE_MAX_FLT);
 
@@ -824,6 +813,19 @@ void FAnimNode_DeadBlending::Evaluate_AnyThread(FPoseContext& Output)
 	// Reset Delta Time
 
 	DeltaTime = 0.0f;
+
+	const float InertializationWeight = InertializationState == EInertializationState::Active ?
+		1.0f - FAlphaBlend::AlphaToBlendOption(
+			InertializationTime / FMath::Max(InertializationDuration, UE_SMALL_NUMBER),
+			InertializationBlendMode, InertializationCustomBlendCurve) : 0.0f;
+
+	TRACE_ANIM_NODE_VALUE_WITH_ID(Output, GetNodeIndex(), TEXT("State"), *UEnum::GetValueAsString(InertializationState));
+	TRACE_ANIM_NODE_VALUE_WITH_ID(Output, GetNodeIndex(), TEXT("Elapsed Time"), InertializationTime);
+	TRACE_ANIM_NODE_VALUE_WITH_ID(Output, GetNodeIndex(), TEXT("Duration"), InertializationDuration);
+	TRACE_ANIM_NODE_VALUE_WITH_ID(Output, GetNodeIndex(), TEXT("Max Duration"), InertializationMaxDuration);
+	TRACE_ANIM_NODE_VALUE_WITH_ID(Output, GetNodeIndex(), TEXT("Normalized Time"), InertializationDuration > UE_KINDA_SMALL_NUMBER ? (InertializationTime / InertializationDuration) : 0.0f);
+	TRACE_ANIM_NODE_VALUE_WITH_ID(Output, GetNodeIndex(), TEXT("Inertialization Weight"), InertializationWeight);
+	TRACE_ANIM_NODE_VALUE_WITH_ID(Output, GetNodeIndex(), TEXT("Request"), *InertializationRequestDescription.ToString());
 }
 
 bool FAnimNode_DeadBlending::NeedsDynamicReset() const
diff --git a/Engine/Source/Runtime/Engine/Private/Animation/AnimNode_Inertialization.cpp b/Engine/Source/Runtime/Engine/Private/Animation/AnimNode_Inertialization.cpp
--- a/Engine/Source/Runtime/Engine/Private/Animation/AnimNode_Inertialization.cpp
+++ b/Engine/Source/Runtime/Engine/Private/Animation/AnimNode_Inertialization.cpp
@@ -30,6 +30,11 @@ static constexpr float INERTIALIZATION_TIME_EPSILON = 1.0e-7f;
 
 namespace UE { namespace Anim {
 
+void IInertializationRequester::RequestInertialization(const FInertializationRequest& InInertializationRequest)
+{
+	RequestInertialization(InInertializationRequest.Duration, InInertializationRequest.BlendProfile);
+}
+
 // Inertialization request event bound to a node
 class FInertializationRequester : public IInertializationRequester
 {
@@ -47,15 +52,10 @@ private:
 		Node.RequestInertialization(InRequestedDuration, InBlendProfile); 
 	}
 
-	virtual void RequestInertializationWithBlendMode(
-		float InRequestedDuration,
-		const UBlendProfile* InBlendProfile,
-		const bool bUseBlendMode,
-		const EAlphaBlendOption InBlendMode,
-		UCurveFloat* InCustomBlendCurve)
+	virtual void RequestInertialization(const FInertializationRequest& InInertializationRequest) override
 	{
-		// The Blend Mode parameters are ignored as FAnimNode_Inertialization does not support them.
-		Node.RequestInertialization(InRequestedDuration, InBlendProfile);
+		// The Blend Mode parameters will be ignored as FAnimNode_Inertialization does not support them.
+		Node.RequestInertialization(InInertializationRequest);
 	}
 
 	virtual void AddDebugRecord(const FAnimInstanceProxy& InSourceProxy, int32 InSourceNodeId)
@@ -109,6 +109,13 @@ void FAnimNode_Inertialization::RequestInertialization(float Duration, const UBl
 	}
 }
 
+void FAnimNode_Inertialization::RequestInertialization(const FInertializationRequest& Request)
+{
+	if (Request.Duration >= 0.0f)
+	{
+		RequestQueue.AddUnique(Request);
+	}
+}
 
 /*static*/ void FAnimNode_Inertialization::LogRequestError(const FAnimationUpdateContext& Context, const FPoseLinkBase& RequesterPoseLink)
 {
@@ -186,7 +193,7 @@ void FAnimNode_Inertialization::Update_AnyThread(const FAnimationUpdateContext&
 				{
 					for (const FInertializationRequest& Request : RequestQueue)
 					{
-						InMessage.RequestInertializationWithBlendMode(Request.Duration, Request.BlendProfile, Request.bUseBlendMode, Request.BlendMode, Request.CustomBlendCurve);
+						InMessage.RequestInertialization(Request);
 					}
  					InMessage.AddDebugRecord(Proxy, NodeId);
 
@@ -267,6 +274,9 @@ void FAnimNode_Inertialization::Evaluate_AnyThread(FPoseContext& Output)
 
 		// Handle the first inertialization request in the queue
 		InertializationDuration = FMath::Max(RequestQueue[0].Duration - AppliedDeficit, 0.0f);
+#if ANIM_TRACE_ENABLED
+		InertializationRequestDescription = RequestQueue[0].Description;
+#endif
 		FillSkeletonBoneDurationsArray(InertializationDurationPerBone, InertializationDuration, RequestQueue[0].BlendProfile);
 
 		// Handle all subsequent inertialization requests (often there will be only a single request)
@@ -279,7 +289,14 @@ void FAnimNode_Inertialization::Evaluate_AnyThread(FPoseContext& Output)
 				const float RequestDuration = FMath::Max(Request.Duration - AppliedDeficit, 0.0f);
 
 				// Merge this request in with the previous requests (using the minimum requested time per bone)
-				InertializationDuration = FMath::Min(InertializationDuration, RequestDuration);
+				if (RequestDuration < InertializationDuration)
+				{
+					InertializationDuration = RequestDuration;
+#if ANIM_TRACE_ENABLED
+					InertializationRequestDescription = Request.Description;
+#endif
+				}
+
 				if (Request.BlendProfile != nullptr)
 				{
 					FillSkeletonBoneDurationsArray(RequestDurationPerBone, RequestDuration, Request.BlendProfile);
@@ -420,11 +437,16 @@ void FAnimNode_Inertialization::Evaluate_AnyThread(FPoseContext& Output)
 	DeltaTime = 0.0f;
 	TeleportType = ETeleportType::None;
 
+	const float NormalizedInertializationTime = InertializationDuration > UE_KINDA_SMALL_NUMBER ? (InertializationElapsedTime / InertializationDuration) : 0.0f;
+	const float InertializationWeight = InertializationState == EInertializationState::Active ? 1.0f - NormalizedInertializationTime : 0.0f;
+
 	TRACE_ANIM_NODE_VALUE_WITH_ID(Output, GetNodeIndex(), TEXT("State"), *UEnum::GetValueAsString(InertializationState));
 	TRACE_ANIM_NODE_VALUE_WITH_ID(Output, GetNodeIndex(), TEXT("Elapsed Time"), InertializationElapsedTime);
 	TRACE_ANIM_NODE_VALUE_WITH_ID(Output, GetNodeIndex(), TEXT("Duration"), InertializationDuration);
 	TRACE_ANIM_NODE_VALUE_WITH_ID(Output, GetNodeIndex(), TEXT("Max Duration"), InertializationMaxDuration);
 	TRACE_ANIM_NODE_VALUE_WITH_ID(Output, GetNodeIndex(), TEXT("Normalized Time"), InertializationDuration > UE_KINDA_SMALL_NUMBER ? (InertializationElapsedTime / InertializationDuration) : 0.0f);
+	TRACE_ANIM_NODE_VALUE_WITH_ID(Output, GetNodeIndex(), TEXT("Inertialization Weight"), InertializationWeight);
+	TRACE_ANIM_NODE_VALUE_WITH_ID(Output, GetNodeIndex(), TEXT("Request"), *InertializationRequestDescription.ToString());
 }
 
 
diff --git a/Engine/Source/Runtime/Engine/Private/Animation/AnimNode_LinkedAnimGraph.cpp b/Engine/Source/Runtime/Engine/Private/Animation/AnimNode_LinkedAnimGraph.cpp
--- a/Engine/Source/Runtime/Engine/Private/Animation/AnimNode_LinkedAnimGraph.cpp
+++ b/Engine/Source/Runtime/Engine/Private/Animation/AnimNode_LinkedAnimGraph.cpp
@@ -119,12 +119,28 @@ void FAnimNode_LinkedAnimGraph::Update_AnyThread(const FAnimationUpdateContext&
 			// Issue the pending inertialization requests (which will get merged together by the inertialization node itself)
 			if (PendingBlendOutDuration >= 0.0f)
 			{
-				InertializationRequester->RequestInertialization(PendingBlendOutDuration, PendingBlendOutProfile);
+				FInertializationRequest Request;
+				Request.Duration = PendingBlendOutDuration;
+				Request.BlendProfile = PendingBlendOutProfile;
+#if ANIM_TRACE_ENABLED
+				Request.Description = NSLOCTEXT("AnimNode_LinkedAnimGraph", "InertializationRequestDescription", "Linked Anim Graph Out");
+#endif
+
+				InertializationRequester->RequestInertialization(Request);
 			}
+
 			if (PendingBlendInDuration >= 0.0f)
 			{
-				InertializationRequester->RequestInertialization(PendingBlendInDuration, PendingBlendInProfile);
+				FInertializationRequest Request;
+				Request.Duration = PendingBlendInDuration;
+				Request.BlendProfile = PendingBlendInProfile;
+#if ANIM_TRACE_ENABLED
+				Request.Description = NSLOCTEXT("AnimNode_LinkedAnimGraph", "InertializationRequestDescription", "Linked Anim Graph In");
+#endif
+
+				InertializationRequester->RequestInertialization(Request);
 			}
+
 			InertializationRequester->AddDebugRecord(*InContext.AnimInstanceProxy, InContext.GetCurrentNodeId());
 		}
 		else if ((PendingBlendOutDuration != 0.0f) && (PendingBlendInDuration != 0.0f) && (InputPoses.Num() > 0))
diff --git a/Engine/Source/Runtime/Engine/Private/Animation/AnimNode_StateMachine.cpp b/Engine/Source/Runtime/Engine/Private/Animation/AnimNode_StateMachine.cpp
--- a/Engine/Source/Runtime/Engine/Private/Animation/AnimNode_StateMachine.cpp
+++ b/Engine/Source/Runtime/Engine/Private/Animation/AnimNode_StateMachine.cpp
@@ -1236,7 +1236,21 @@ void FAnimNode_StateMachine::TransitionToState(const FAnimationUpdateContext& Co
 				UE::Anim::IInertializationRequester* InertializationRequester = Context.GetMessage<UE::Anim::IInertializationRequester>();
 				if (InertializationRequester)
 				{
-					InertializationRequester->RequestInertializationWithBlendMode(TransitionInfo.CrossfadeDuration, TransitionInfo.BlendProfile, true, TransitionInfo.BlendMode, TransitionInfo.CustomCurve);
+					FInertializationRequest Request;
+					Request.Duration = TransitionInfo.CrossfadeDuration;
+					Request.BlendProfile = TransitionInfo.BlendProfile;
+					Request.bUseBlendMode = true;
+					Request.BlendMode = TransitionInfo.BlendMode;
+					Request.CustomBlendCurve = TransitionInfo.CustomCurve;
+#if ANIM_TRACE_ENABLED
+					Request.Description = FText::Format(LOCTEXT("InertializationRequestDescription", 
+						"State Machine \"{0}\" Transition \"{1}\" to \"{2}\""), 
+						FText::FromName(GetMachineDescription()->MachineName),
+						FText::FromName(GetStateInfo(TransitionInfo.PreviousState).StateName),
+						FText::FromName(GetStateInfo(TransitionInfo.NextState).StateName));
+#endif
+
+					InertializationRequester->RequestInertialization(Request);
 					InertializationRequester->AddDebugRecord(*Context.AnimInstanceProxy, Context.GetCurrentNodeId());
 				}
 				else
diff --git a/Engine/Source/Runtime/Engine/Private/Animation/AnimTrace.cpp b/Engine/Source/Runtime/Engine/Private/Animation/AnimTrace.cpp
--- a/Engine/Source/Runtime/Engine/Private/Animation/AnimTrace.cpp
+++ b/Engine/Source/Runtime/Engine/Private/Animation/AnimTrace.cpp
@@ -78,6 +78,7 @@ UE_TRACE_EVENT_BEGIN(Animation, AnimNodeStart)
 	UE_TRACE_EVENT_FIELD(float, Weight)
 	UE_TRACE_EVENT_FIELD(uint16, FrameCounter)
 	UE_TRACE_EVENT_FIELD(uint8, Phase)
+	UE_TRACE_EVENT_FIELD(UE::Trace::WideString, Name)	
 	UE_TRACE_EVENT_FIELD(UE::Trace::WideString, DisplayName)
 UE_TRACE_EVENT_END()
 
@@ -97,6 +98,7 @@ UE_TRACE_EVENT_END()
 
 UE_TRACE_EVENT_BEGIN(Animation, AnimNodeValueBool)
 	UE_TRACE_EVENT_FIELD(uint64, Cycle)
+	UE_TRACE_EVENT_FIELD(double, RecordingTime)
 	UE_TRACE_EVENT_FIELD(uint64, AnimInstanceId)
 	UE_TRACE_EVENT_FIELD(int32, NodeId)
 	UE_TRACE_EVENT_FIELD(uint16, FrameCounter)
@@ -106,6 +108,7 @@ UE_TRACE_EVENT_END()
 
 UE_TRACE_EVENT_BEGIN(Animation, AnimNodeValueInt)
 	UE_TRACE_EVENT_FIELD(uint64, Cycle)
+	UE_TRACE_EVENT_FIELD(double, RecordingTime)
 	UE_TRACE_EVENT_FIELD(uint64, AnimInstanceId)
 	UE_TRACE_EVENT_FIELD(int32, NodeId)
 	UE_TRACE_EVENT_FIELD(int32, Value)
@@ -115,6 +118,7 @@ UE_TRACE_EVENT_END()
 
 UE_TRACE_EVENT_BEGIN(Animation, AnimNodeValueFloat)
 	UE_TRACE_EVENT_FIELD(uint64, Cycle)
+	UE_TRACE_EVENT_FIELD(double, RecordingTime)
 	UE_TRACE_EVENT_FIELD(uint64, AnimInstanceId)
 	UE_TRACE_EVENT_FIELD(int32, NodeId)
 	UE_TRACE_EVENT_FIELD(float, Value)
@@ -124,6 +128,7 @@ UE_TRACE_EVENT_END()
 
 UE_TRACE_EVENT_BEGIN(Animation, AnimNodeValueVector2D)
 	UE_TRACE_EVENT_FIELD(uint64, Cycle)
+	UE_TRACE_EVENT_FIELD(double, RecordingTime)
 	UE_TRACE_EVENT_FIELD(uint64, AnimInstanceId)
 	UE_TRACE_EVENT_FIELD(int32, NodeId)
 	UE_TRACE_EVENT_FIELD(float, ValueX)
@@ -134,6 +139,7 @@ UE_TRACE_EVENT_END()
 
 UE_TRACE_EVENT_BEGIN(Animation, AnimNodeValueVector)
 	UE_TRACE_EVENT_FIELD(uint64, Cycle)
+	UE_TRACE_EVENT_FIELD(double, RecordingTime)
 	UE_TRACE_EVENT_FIELD(uint64, AnimInstanceId)
 	UE_TRACE_EVENT_FIELD(int32, NodeId)
 	UE_TRACE_EVENT_FIELD(float, ValueX)
@@ -145,6 +151,7 @@ UE_TRACE_EVENT_END()
 
 UE_TRACE_EVENT_BEGIN(Animation, AnimNodeValueString)
 	UE_TRACE_EVENT_FIELD(uint64, Cycle)
+	UE_TRACE_EVENT_FIELD(double, RecordingTime)
 	UE_TRACE_EVENT_FIELD(uint64, AnimInstanceId)
 	UE_TRACE_EVENT_FIELD(int32, NodeId)
 	UE_TRACE_EVENT_FIELD(UE::Trace::WideString, Value)
@@ -154,6 +161,7 @@ UE_TRACE_EVENT_END()
 
 UE_TRACE_EVENT_BEGIN(Animation, AnimNodeValueObject)
 	UE_TRACE_EVENT_FIELD(uint64, Cycle)
+	UE_TRACE_EVENT_FIELD(double, RecordingTime)
 	UE_TRACE_EVENT_FIELD(uint64, AnimInstanceId)
 	UE_TRACE_EVENT_FIELD(uint64, Value)
 	UE_TRACE_EVENT_FIELD(int32, NodeId)
@@ -163,6 +171,7 @@ UE_TRACE_EVENT_END()
 
 UE_TRACE_EVENT_BEGIN(Animation, AnimNodeValueClass)
 	UE_TRACE_EVENT_FIELD(uint64, Cycle)
+	UE_TRACE_EVENT_FIELD(double, RecordingTime)
 	UE_TRACE_EVENT_FIELD(uint64, AnimInstanceId)
 	UE_TRACE_EVENT_FIELD(uint64, Value)
 	UE_TRACE_EVENT_FIELD(int32, NodeId)
@@ -649,7 +658,7 @@ void FAnimTrace::OutputAnimNodeStart(const FAnimationBaseContext& InContext, uin
 	UObject* AnimInstance = InContext.AnimInstanceProxy->GetAnimInstanceObject();
 	TRACE_OBJECT(AnimInstance);
 
-	FString DisplayNameString;
+	FString NameString, DisplayNameString;
 	IAnimClassInterface* AnimBlueprintClass = InContext.GetAnimClass();
 	if(AnimBlueprintClass)
 	{
@@ -658,12 +667,13 @@ void FAnimTrace::OutputAnimNodeStart(const FAnimationBaseContext& InContext, uin
 		FStructProperty* LinkedProperty = AnimNodeProperties[InNodeId];
 		check(LinkedProperty->Struct);
 
+		NameString = LinkedProperty->Struct->GetFName().ToString();
+
 #if WITH_EDITOR
 		DisplayNameString = LinkedProperty->Struct->GetDisplayNameText().ToString();
 #else
 		DisplayNameString = LinkedProperty->Struct->GetName();
 #endif
-
 		DisplayNameString.RemoveFromStart(TEXT("Anim Node "));
 	}
 	else
@@ -681,6 +691,7 @@ void FAnimTrace::OutputAnimNodeStart(const FAnimationBaseContext& InContext, uin
 		<< AnimNodeStart.Weight(InBlendWeight)
 		<< AnimNodeStart.FrameCounter(FObjectTrace::GetObjectWorldTickCounter(AnimInstance))
 		<< AnimNodeStart.Phase(InPhase)
+		<< AnimNodeStart.Name(*NameString, NameString.Len())	
 		<< AnimNodeStart.DisplayName(*DisplayNameString, DisplayNameString.Len());
 }
 
@@ -756,6 +767,7 @@ void FAnimTrace::OutputAnimNodeValue(const FAnimationBaseContext& InContext, uin
 
 	UE_TRACE_LOG(Animation, AnimNodeValueBool, AnimationChannel)
 		<< AnimNodeValueBool.Cycle(FPlatformTime::Cycles64())
+		<< AnimNodeValueBool.RecordingTime(FObjectTrace::GetWorldElapsedTime(AnimInstance->GetWorld()))
 		<< AnimNodeValueBool.AnimInstanceId(FObjectTrace::GetObjectId(AnimInstance))
 		<< AnimNodeValueBool.NodeId(NodeIndex)
 		<< AnimNodeValueBool.Value(InValue)
@@ -783,6 +795,7 @@ void FAnimTrace::OutputAnimNodeValue(const FAnimationBaseContext& InContext, uin
 
 	UE_TRACE_LOG(Animation, AnimNodeValueInt, AnimationChannel)
 		<< AnimNodeValueInt.Cycle(FPlatformTime::Cycles64())
+		<< AnimNodeValueInt.RecordingTime(FObjectTrace::GetWorldElapsedTime(AnimInstance->GetWorld()))
 		<< AnimNodeValueInt.AnimInstanceId(FObjectTrace::GetObjectId(AnimInstance))
 		<< AnimNodeValueInt.NodeId(NodeIndex)
 		<< AnimNodeValueInt.Value(InValue)
@@ -810,6 +823,7 @@ void FAnimTrace::OutputAnimNodeValue(const FAnimationBaseContext& InContext, uin
 
 	UE_TRACE_LOG(Animation, AnimNodeValueFloat, AnimationChannel)
 		<< AnimNodeValueFloat.Cycle(FPlatformTime::Cycles64())
+		<< AnimNodeValueFloat.RecordingTime(FObjectTrace::GetWorldElapsedTime(AnimInstance->GetWorld()))
 		<< AnimNodeValueFloat.AnimInstanceId(FObjectTrace::GetObjectId(AnimInstance))
 		<< AnimNodeValueFloat.NodeId(NodeIndex)
 		<< AnimNodeValueFloat.Value(InValue)
@@ -832,6 +846,7 @@ void FAnimTrace::OutputAnimNodeValue(const FAnimationBaseContext& InContext, uin
 
 	UE_TRACE_LOG(Animation, AnimNodeValueVector2D, AnimationChannel)
 		<< AnimNodeValueVector2D.Cycle(FPlatformTime::Cycles64())
+		<< AnimNodeValueVector2D.RecordingTime(FObjectTrace::GetWorldElapsedTime(AnimInstance->GetWorld()))
 		<< AnimNodeValueVector2D.AnimInstanceId(FObjectTrace::GetObjectId(AnimInstance))
 		<< AnimNodeValueVector2D.NodeId(NodeIndex)
 		<< AnimNodeValueVector2D.ValueX(InValue.X)
@@ -867,6 +882,7 @@ void FAnimTrace::OutputAnimNodeValue(const FAnimationBaseContext& InContext, uin
 
 	UE_TRACE_LOG(Animation, AnimNodeValueVector, AnimationChannel)
 		<< AnimNodeValueVector.Cycle(FPlatformTime::Cycles64())
+		<< AnimNodeValueVector.RecordingTime(FObjectTrace::GetWorldElapsedTime(AnimInstance->GetWorld()))
 		<< AnimNodeValueVector.AnimInstanceId(FObjectTrace::GetObjectId(InContext.AnimInstanceProxy->GetAnimInstanceObject()))
 		<< AnimNodeValueVector.NodeId(NodeIndex)
 		<< AnimNodeValueVector.ValueX(InValue.X)
@@ -899,6 +915,7 @@ void FAnimTrace::OutputAnimNodeValue(const FAnimationBaseContext& InContext, uin
 
 	UE_TRACE_LOG(Animation, AnimNodeValueString, AnimationChannel)
 		<< AnimNodeValueString.Cycle(FPlatformTime::Cycles64())
+		<< AnimNodeValueString.RecordingTime(FObjectTrace::GetWorldElapsedTime(AnimInstance->GetWorld()))
 		<< AnimNodeValueString.AnimInstanceId(FObjectTrace::GetObjectId(AnimInstance))
 		<< AnimNodeValueString.NodeId(NodeIndex)
 		<< AnimNodeValueString.FrameCounter(FObjectTrace::GetObjectWorldTickCounter(AnimInstance))
@@ -926,6 +943,7 @@ void FAnimTrace::OutputAnimNodeValue(const FAnimationBaseContext& InContext, uin
 
 	UE_TRACE_LOG(Animation, AnimNodeValueString, AnimationChannel)
 		<< AnimNodeValueString.Cycle(FPlatformTime::Cycles64())
+		<< AnimNodeValueString.RecordingTime(FObjectTrace::GetWorldElapsedTime(AnimInstance->GetWorld()))
 		<< AnimNodeValueString.AnimInstanceId(FObjectTrace::GetObjectId(AnimInstance))
 		<< AnimNodeValueString.NodeId(NodeIndex)
 		<< AnimNodeValueString.FrameCounter(FObjectTrace::GetObjectWorldTickCounter(AnimInstance))
@@ -954,6 +972,7 @@ void FAnimTrace::OutputAnimNodeValue(const FAnimationBaseContext& InContext, uin
 
 	UE_TRACE_LOG(Animation, AnimNodeValueObject, AnimationChannel)
 		<< AnimNodeValueObject.Cycle(FPlatformTime::Cycles64())
+		<< AnimNodeValueObject.RecordingTime(FObjectTrace::GetWorldElapsedTime(AnimInstance->GetWorld()))
 		<< AnimNodeValueObject.AnimInstanceId(FObjectTrace::GetObjectId(AnimInstance))
 		<< AnimNodeValueObject.NodeId(NodeIndex)
 		<< AnimNodeValueObject.Value(FObjectTrace::GetObjectId(InValue))
@@ -982,6 +1001,7 @@ void FAnimTrace::OutputAnimNodeValue(const FAnimationBaseContext& InContext, uin
 
 	UE_TRACE_LOG(Animation, AnimNodeValueClass, AnimationChannel)
 		<< AnimNodeValueClass.Cycle(FPlatformTime::Cycles64())
+		<< AnimNodeValueClass.RecordingTime(FObjectTrace::GetWorldElapsedTime(AnimInstance->GetWorld()))
 		<< AnimNodeValueClass.AnimInstanceId(FObjectTrace::GetObjectId(AnimInstance))
 		<< AnimNodeValueClass.NodeId(NodeIndex)
 		<< AnimNodeValueClass.Value(FObjectTrace::GetObjectId(InValue))
