RewindDebugger: Added ability to click on AnimNodes in the rewind debugger and get shown them in the Animation Blueprint. Made inertialization record the node and AnimInstance that made the inertialization request and made these clickable in rewind debugger. Made it so that the inertialization track will not show unless you have inertialization nodes used in your graph. Some small refactoring of the rewind debugger display of AnimNodeValues to reduce the code duplication a bit.

diff --git a/Engine/Plugins/Animation/GameplayInsights/Source/GameplayInsights/Private/AnimationAnalyzer.cpp b/Engine/Plugins/Animation/GameplayInsights/Source/GameplayInsights/Private/AnimationAnalyzer.cpp
--- a/Engine/Plugins/Animation/GameplayInsights/Source/GameplayInsights/Private/AnimationAnalyzer.cpp
+++ b/Engine/Plugins/Animation/GameplayInsights/Source/GameplayInsights/Private/AnimationAnalyzer.cpp
@@ -37,6 +37,7 @@ void FAnimationAnalyzer::OnAnalysisBegin(const FOnAnalysisContext& Context)
 	Builder.RouteEvent(RouteId_AnimNodeValueString, "Animation", "AnimNodeValueString");
 	Builder.RouteEvent(RouteId_AnimNodeValueObject, "Animation", "AnimNodeValueObject");
 	Builder.RouteEvent(RouteId_AnimNodeValueClass, "Animation", "AnimNodeValueClass");
+	Builder.RouteEvent(RouteId_AnimNodeValueAnimNode, "Animation", "AnimNodeValueAnimNode");
 	Builder.RouteEvent(RouteId_AnimSequencePlayer, "Animation", "AnimSequencePlayer");
 	Builder.RouteEvent(RouteId_BlendSpacePlayer, "Animation", "BlendSpacePlayer");
 	Builder.RouteEvent(RouteId_StateMachineState, "Animation", "StateMachineState");
@@ -235,6 +236,7 @@ bool FAnimationAnalyzer::OnEvent(uint16 RouteId, EStyle Style, const FOnEventCon
 	case RouteId_AnimNodeValueString:
 	case RouteId_AnimNodeValueObject:
 	case RouteId_AnimNodeValueClass:
+	case RouteId_AnimNodeValueAnimNode:
 	{
 		uint64 Cycle = EventData.GetValue<uint64>("Cycle");
 		double RecordingTime = EventData.GetValue<double>("RecordingTime");
@@ -299,6 +301,13 @@ bool FAnimationAnalyzer::OnEvent(uint16 RouteId, EStyle Style, const FOnEventCon
 			AnimationProvider.AppendAnimNodeValueClass(AnimInstanceId, Context.EventTime.AsSeconds(Cycle), RecordingTime, FrameCounter, NodeId, *Key, Value);
 			break;
 		}
+		case RouteId_AnimNodeValueAnimNode:
+		{
+			int32 Value = EventData.GetValue<int32>("Value");
+			uint64 ValueAnimInstanceId = EventData.GetValue<uint64>("ValueAnimInstanceId");
+			AnimationProvider.AppendAnimNodeValueAnimNode(AnimInstanceId, Context.EventTime.AsSeconds(Cycle), RecordingTime, FrameCounter, NodeId, *Key, Value, ValueAnimInstanceId);
+			break;
+		}
 		}
 		break;
 	}
diff --git a/Engine/Plugins/Animation/GameplayInsights/Source/GameplayInsights/Private/AnimationAnalyzer.h b/Engine/Plugins/Animation/GameplayInsights/Source/GameplayInsights/Private/AnimationAnalyzer.h
--- a/Engine/Plugins/Animation/GameplayInsights/Source/GameplayInsights/Private/AnimationAnalyzer.h
+++ b/Engine/Plugins/Animation/GameplayInsights/Source/GameplayInsights/Private/AnimationAnalyzer.h
@@ -38,6 +38,7 @@ private:
 		RouteId_AnimNodeValueString,
 		RouteId_AnimNodeValueObject,
 		RouteId_AnimNodeValueClass,
+		RouteId_AnimNodeValueAnimNode,
 		RouteId_AnimSequencePlayer,
 		RouteId_BlendSpacePlayer,
 		RouteId_StateMachineState,
diff --git a/Engine/Plugins/Animation/GameplayInsights/Source/GameplayInsights/Private/AnimationProvider.cpp b/Engine/Plugins/Animation/GameplayInsights/Source/GameplayInsights/Private/AnimationProvider.cpp
--- a/Engine/Plugins/Animation/GameplayInsights/Source/GameplayInsights/Private/AnimationProvider.cpp
+++ b/Engine/Plugins/Animation/GameplayInsights/Source/GameplayInsights/Private/AnimationProvider.cpp
@@ -437,6 +437,22 @@ const FSkeletalMeshInfo* FAnimationProvider::FindSkeletalMeshInfo(uint64 InObjec
 	return nullptr;
 }
 
+const FAnimNodeInfo* FAnimationProvider::FindAnimNodeInfo(int32 InAnimNodeId, uint64 InAnimInstanceId) const
+{
+	Session.ReadAccessCheck();
+
+	const uint32* IndexPtr = AnimNodeIdToIndexMap.Find({ InAnimNodeId, InAnimInstanceId });
+	if (IndexPtr != nullptr)
+	{
+		if (*IndexPtr < uint32(AnimNodeInfos.Num()))
+		{
+			return &AnimNodeInfos[*IndexPtr];
+		}
+	}
+
+	return nullptr;
+}
+
 const TCHAR* FAnimationProvider::GetName(uint32 InId) const
 {
 	const TCHAR* const* FoundName = NameMap.Find(InId);
@@ -490,6 +506,12 @@ FText FAnimationProvider::FormatNodeValue(const FAnimNodeValueMessage& InMessage
 		Text = FText::FromString(ClassInfo.PathName);
 		break;
 	}
+	case EAnimNodeValueType::AnimNode:
+	{
+		const FAnimNodeInfo* AnimNodeInfo = FindAnimNodeInfo(InMessage.Value.AnimNode.Value, InMessage.Value.AnimNode.AnimInstanceId);
+		Text = FText::FromString(AnimNodeInfo ? AnimNodeInfo->Name : TEXT(""));
+		break;
+	}
 	}
 
 	return Text;
@@ -887,6 +909,20 @@ void FAnimationProvider::AppendAnimNodeStart(uint64 InAnimInstanceId, double InS
 
 	Timeline->AppendEvent(InStartTime, Message);
 
+	if (AnimNodeIdToIndexMap.Find({ InNodeId, InAnimInstanceId }) == nullptr)
+	{
+		bHasAnyData = true;
+
+		FAnimNodeInfo NewAnimNodeInfo;
+		NewAnimNodeInfo.Id = Message.NodeId;
+		NewAnimNodeInfo.AnimInstanceId = Message.AnimInstanceId;
+		NewAnimNodeInfo.Name = Message.NodeName;
+		NewAnimNodeInfo.TypeName = Message.NodeTypeName;
+
+		uint32 NewAnimNodeInfoIndex = AnimNodeInfos.Add(NewAnimNodeInfo);
+		AnimNodeIdToIndexMap.Add({ InNodeId, InAnimInstanceId }, NewAnimNodeInfoIndex);
+	}
+
 	Session.UpdateDurationSeconds(InStartTime);
 }
 
@@ -1030,6 +1066,16 @@ void FAnimationProvider::AppendAnimNodeValueClass(uint64 InAnimInstanceId, doubl
 	AppendAnimNodeValue(InAnimInstanceId, InTime, InRecordingTime, InFrameCounter, InNodeId, InKey, Message);
 }
 
+void FAnimationProvider::AppendAnimNodeValueAnimNode(uint64 InAnimInstanceId, double InTime, double InRecordingTime, uint16 InFrameCounter, int32 InNodeId, const TCHAR* InKey, int32 InValue, uint64 InValueAnimInstanceId)
+{
+	FAnimNodeValueMessage Message;
+	Message.Value.AnimNode.Value = InValue;
+	Message.Value.AnimNode.AnimInstanceId = InValueAnimInstanceId;
+	Message.Value.Type = EAnimNodeValueType::AnimNode;
+
+	AppendAnimNodeValue(InAnimInstanceId, InTime, InRecordingTime, InFrameCounter, InNodeId, InKey, Message);
+}
+
 void FAnimationProvider::AppendAnimNodeValue(uint64 InAnimInstanceId, double InTime, double InRecordingTime, uint16 InFrameCounter, int32 InNodeId, const TCHAR* InKey, FAnimNodeValueMessage& InMessage)
 {
 	Session.WriteAccessCheck();
diff --git a/Engine/Plugins/Animation/GameplayInsights/Source/GameplayInsights/Private/AnimationProvider.h b/Engine/Plugins/Animation/GameplayInsights/Source/GameplayInsights/Private/AnimationProvider.h
--- a/Engine/Plugins/Animation/GameplayInsights/Source/GameplayInsights/Private/AnimationProvider.h
+++ b/Engine/Plugins/Animation/GameplayInsights/Source/GameplayInsights/Private/AnimationProvider.h
@@ -47,6 +47,7 @@ public:
 	virtual bool ReadAnimSyncTimeline(uint64 InObjectId, TFunctionRef<void(const AnimSyncTimeline&)> Callback) const override;
 	virtual bool ReadPoseWatchTimeline(uint64 InObjectId, TFunctionRef<void(const PoseWatchTimeline&)> Callback) const override;
 	virtual const FSkeletalMeshInfo* FindSkeletalMeshInfo(uint64 InObjectId) const override;
+	virtual const FAnimNodeInfo* FindAnimNodeInfo(int32 InNodeId, uint64 InAnimInstanceId) const override;
 	virtual const TCHAR* GetName(uint32 InId) const override;
 	virtual FText FormatNodeKeyValue(const FAnimNodeValueMessage& InMessage) const override;
 	virtual FText FormatNodeValue(const FAnimNodeValueMessage& InMessage) const override;
@@ -88,6 +89,7 @@ public:
 	void AppendAnimNodeValue(uint64 InAnimInstanceId, double InTime, double InRecordingTime, uint16 InFrameCounter, int32 InNodeId, const TCHAR* InKey, const TCHAR* InValue);
 	void AppendAnimNodeValueObject(uint64 InAnimInstanceId, double InTime, double InRecordingTime, uint16 InFrameCounter, int32 InNodeId, const TCHAR* InKey, uint64 InValue);
 	void AppendAnimNodeValueClass(uint64 InAnimInstanceId, double InTime, double InRecordingTime, uint16 InFrameCounter, int32 InNodeId, const TCHAR* InKey, uint64 InValue);
+	void AppendAnimNodeValueAnimNode(uint64 InAnimInstanceId, double InTime, double InRecordingTime, uint16 InFrameCounter, int32 InNodeId, const TCHAR* InKey, int32 InValue, uint64 InValueAnimInstanceId);
 
 	/** Add anim graph attribute */
 	void AppendAnimGraphAttribute(uint64 InSourceAnimInstanceId, uint64 InTargetAnimInstanceId, double InTime, int32 InSourceNodeId, int32 InTargetNodeId, uint32 InAttributeNameId);
@@ -148,6 +150,12 @@ private:
 	/** Map of skeletal mesh->index into SkeletalMeshInfos index */
 	TMap<uint64, uint32> SkeletalMeshIdToIndexMap;
 
+	/** All the anim node info we have seen, grow only for stable indices */
+	TArray<FAnimNodeInfo> AnimNodeInfos;
+
+	/** Map of (node, anim instance)->index into AnimNodeInfos index */
+	TMap<TTuple<int32, uint64>, uint32> AnimNodeIdToIndexMap;
+
 	/** Map of name IDs to name string */
 	TMap<uint32, const TCHAR*> NameMap;
 	
diff --git a/Engine/Plugins/Animation/GameplayInsights/Source/GameplayInsights/Private/InertializationsTrack.cpp b/Engine/Plugins/Animation/GameplayInsights/Source/GameplayInsights/Private/InertializationsTrack.cpp
--- a/Engine/Plugins/Animation/GameplayInsights/Source/GameplayInsights/Private/InertializationsTrack.cpp
+++ b/Engine/Plugins/Animation/GameplayInsights/Source/GameplayInsights/Private/InertializationsTrack.cpp
@@ -277,12 +277,26 @@ bool FInertializationsTrackCreator::HasDebugInfoInternal(uint64 ObjectId) const
 {
 	const TraceServices::IAnalysisSession* AnalysisSession = IRewindDebugger::Instance()->GetAnalysisSession();
 	TraceServices::FAnalysisSessionReadScope SessionReadScope(*AnalysisSession);
+	
 	bool bHasData = false;
+
 	if (const FAnimationProvider* AnimationProvider = AnalysisSession->ReadProvider<FAnimationProvider>(FAnimationProvider::ProviderName))
 	{
-		AnimationProvider->ReadAnimGraphTimeline(ObjectId, [&bHasData](const FAnimationProvider::AnimGraphTimeline& InGraphTimeline)
+		AnimationProvider->ReadAnimNodesTimeline(ObjectId, [&AnimationProvider, &bHasData](const FAnimationProvider::AnimNodesTimeline& InTimeline)
 		{
-			bHasData = true;
+			InTimeline.EnumerateEvents(InTimeline.GetStartTime(), InTimeline.GetEndTime(), [&AnimationProvider, &bHasData](double InStartTime, double InEndTime, uint32 InDepth, const FAnimNodeMessage& InMessage)
+			{
+				for (const TCHAR* InertializationNodeType : { TEXT("AnimNode_DeadBlending"), TEXT("AnimNode_Inertialization") })
+				{
+					if (FCString::Strcmp(InertializationNodeType, InMessage.NodeTypeName) == 0)
+					{
+						bHasData = true;
+						return TraceServices::EEventEnumerate::Stop;
+					}
+				}
+
+				return TraceServices::EEventEnumerate::Continue;
+			});
 		});
 	}
 	
diff --git a/Engine/Plugins/Animation/GameplayInsights/Source/GameplayInsights/Private/InertializationsTrack.h b/Engine/Plugins/Animation/GameplayInsights/Source/GameplayInsights/Private/InertializationsTrack.h
--- a/Engine/Plugins/Animation/GameplayInsights/Source/GameplayInsights/Private/InertializationsTrack.h
+++ b/Engine/Plugins/Animation/GameplayInsights/Source/GameplayInsights/Private/InertializationsTrack.h
@@ -17,8 +17,6 @@ public:
 	
 	FInertializationTrack(uint64 InObjectId, int32 NodeId, const FText& Name);
 
-	int32 GetNodeId() const { return NodeId; };
-
 private:
 	virtual bool UpdateInternal() override;
 	virtual TSharedPtr<SWidget> GetTimelineViewInternal() override;
@@ -32,11 +30,11 @@ private:
 
 	TSharedPtr<SCurveTimelineView::FTimelineCurveData> GetCurveData() const;
 	
-	int32 NodeId;
+	int32 NodeId = INDEX_NONE;
 
 	FSlateIcon Icon;
 	FText CurveName;
-	uint64 ObjectId;
+	uint64 ObjectId = INDEX_NONE;
 
 	mutable TSharedPtr<SCurveTimelineView::FTimelineCurveData> CurveData;
 	mutable int CurvesUpdateRequested = 0;
@@ -59,7 +57,7 @@ private:
 	virtual FSlateIcon GetIconInternal() override { return Icon; }
 
 	FSlateIcon Icon;
-	uint64 ObjectId;
+	uint64 ObjectId = INDEX_NONE;
 	TArray<TSharedPtr<FInertializationTrack>> Children;
 };
 	
diff --git a/Engine/Plugins/Animation/GameplayInsights/Source/GameplayInsights/Private/SAnimGraphSchematicView.cpp b/Engine/Plugins/Animation/GameplayInsights/Source/GameplayInsights/Private/SAnimGraphSchematicView.cpp
--- a/Engine/Plugins/Animation/GameplayInsights/Source/GameplayInsights/Private/SAnimGraphSchematicView.cpp
+++ b/Engine/Plugins/Animation/GameplayInsights/Source/GameplayInsights/Private/SAnimGraphSchematicView.cpp
@@ -18,12 +18,16 @@
 #include "Widgets/Input/SCheckBox.h"
 #include "Widgets/Input/SComboButton.h"
 #include "Widgets/Images/SImage.h"
+#include "SVariantValueView.h"
 
 #if WITH_EDITOR
 #include "Animation/AnimInstance.h"
+#include "Animation/AnimBlueprint.h"
+#include "Animation/AnimBlueprintGeneratedClass.h"
 #include "Styling/SlateIconFinder.h"
 #include "Subsystems/AssetEditorSubsystem.h"
 #include "Editor.h"
+#include "IAnimationBlueprintEditor.h"
 #endif
 
 #define LOCTEXT_NAMESPACE "SAnimGraphSchematicView"
@@ -66,7 +70,7 @@ public:
 		TSharedPtr<FAnimNodeValueMessage> PinnedValue = Value.Pin();
 		if(PinnedValue.IsValid())
 		{
-			return StaticMakeValueWidget(AnalysisSession, PinnedValue.ToSharedRef());
+			return SVariantValueView::MakeVariantValueWidget(AnalysisSession, PinnedValue.ToSharedRef()->Value, FText::GetEmpty());
 		}
 		else
 		{
@@ -74,172 +78,6 @@ public:
 		}
 	}
 
-	static TSharedRef<SWidget> StaticMakeValueWidget(const TraceServices::IAnalysisSession& InAnalysisSession, const TSharedRef<FAnimNodeValueMessage>& InValue) 
-	{ 
-		switch(InValue->Value.Type)
-		{
-		case EAnimNodeValueType::Bool:
-			return 
-				SNew(SCheckBox)
-				.IsEnabled(false)
-				.IsChecked(InValue->Value.Bool.bValue ? ECheckBoxState::Checked : ECheckBoxState::Unchecked);
-
-		case EAnimNodeValueType::Int32:
-			return
-				SNew(SBox)
-				.WidthOverride(125.0f)
-				[
-					SNew(SEditableTextBox)
-					.IsEnabled(false)
-					.Font(FCoreStyle::Get().GetFontStyle("SmallFont"))
-					.Text(FText::AsNumber(InValue->Value.Int32.Value))
-				];
-
-		case EAnimNodeValueType::Float:
-			return 
-				SNew(SBox)
-				.WidthOverride(125.0f)
-				[
-					SNew(SEditableTextBox)
-					.IsEnabled(false)
-					.Font(FCoreStyle::Get().GetFontStyle("SmallFont"))
-					.Text(FText::AsNumber(InValue->Value.Float.Value))
-				];
-
-		case EAnimNodeValueType::Vector2D:
-			return SNew(SHorizontalBox)
-				+SHorizontalBox::Slot()
-				.AutoWidth()
-				[
-					SNew(SBox)
-					.WidthOverride(125.0f)
-					[
-						SNew(SEditableTextBox)
-						.IsEnabled(false)
-						.Font(FCoreStyle::Get().GetFontStyle("SmallFont"))
-						.Text(FText::AsNumber(InValue->Value.Vector2D.Value.X))
-					]
-				]
-				+SHorizontalBox::Slot()
-				.AutoWidth()
-				[
-					SNew(SBox)
-					.WidthOverride(125.0f)
-					[
-						SNew(SEditableTextBox)
-						.IsEnabled(false)
-						.Font(FCoreStyle::Get().GetFontStyle("SmallFont"))
-						.Text(FText::AsNumber(InValue->Value.Vector2D.Value.Y))
-					]
-				];
-
-		case EAnimNodeValueType::Vector:
-			return SNew(SHorizontalBox)
-				+SHorizontalBox::Slot()
-				.AutoWidth()
-				[
-					SNew(SBox)
-					.WidthOverride(125.0f)
-					[
-						SNew(SEditableTextBox)
-						.IsEnabled(false)
-						.Font(FCoreStyle::Get().GetFontStyle("SmallFont"))
-						.Text(FText::AsNumber(InValue->Value.Vector.Value.X))
-					]
-				]
-				+SHorizontalBox::Slot()
-				.AutoWidth()
-				[
-					SNew(SBox)
-					.WidthOverride(125.0f)
-					[
-						SNew(SEditableTextBox)
-						.IsEnabled(false)
-						.Font(FCoreStyle::Get().GetFontStyle("SmallFont"))
-						.Text(FText::AsNumber(InValue->Value.Vector.Value.Y))
-					]
-				]
-				+SHorizontalBox::Slot()
-				.AutoWidth()
-				[
-					SNew(SBox)
-					.WidthOverride(125.0f)
-					[
-						SNew(SEditableTextBox)
-						.IsEnabled(false)
-						.Font(FCoreStyle::Get().GetFontStyle("SmallFont"))
-						.Text(FText::AsNumber(InValue->Value.Vector.Value.Z))
-					]
-				];
-
-		case EAnimNodeValueType::String:
-			return 
-				SNew(STextBlock)
-				.Font(FCoreStyle::Get().GetFontStyle("SmallFont"))
-				.Text(FText::FromString(InValue->Value.String.Value));
-
-		case EAnimNodeValueType::Object:
-		{
-			const FGameplayProvider* GameplayProvider = InAnalysisSession.ReadProvider<FGameplayProvider>(FGameplayProvider::ProviderName);
-			if(GameplayProvider)
-			{
-				TraceServices::FAnalysisSessionReadScope SessionReadScope(InAnalysisSession);
-
-				const FObjectInfo& ObjectInfo = GameplayProvider->GetObjectInfo(InValue->Value.Object.Value);
-#if WITH_EDITOR
-				return 
-					SNew(SHyperlink)
-					.Text(FText::FromString(ObjectInfo.Name))
-					.TextStyle(&FCoreStyle::Get().GetWidgetStyle<FTextBlockStyle>("SmallText"))
-					.ToolTipText(FText::Format(LOCTEXT("AssetHyperlinkTooltipFormat", "Open asset '{0}'"), FText::FromString(ObjectInfo.PathName)))
-					.OnNavigate_Lambda([ObjectInfo]()
-					{
-						GEditor->GetEditorSubsystem<UAssetEditorSubsystem>()->OpenEditorForAsset(ObjectInfo.PathName);
-					});
-#else
-				return 
-					SNew(STextBlock)
-					.Font(FCoreStyle::Get().GetFontStyle("SmallFont"))
-					.Text(FText::FromString(ObjectInfo.Name))
-					.ToolTipText(FText::FromString(ObjectInfo.PathName));
-#endif
-			}
-		}
-		break;
-
-		case EAnimNodeValueType::Class:
-		{
-			const FGameplayProvider* GameplayProvider = InAnalysisSession.ReadProvider<FGameplayProvider>(FGameplayProvider::ProviderName);
-			if(GameplayProvider)
-			{
-				TraceServices::FAnalysisSessionReadScope SessionReadScope(InAnalysisSession);
-
-				const FClassInfo& ClassInfo = GameplayProvider->GetClassInfo(InValue->Value.Class.Value);
-#if WITH_EDITOR
-				return 
-					SNew(SHyperlink)
-					.Text(FText::FromString(ClassInfo.Name))
-					.TextStyle(&FCoreStyle::Get().GetWidgetStyle<FTextBlockStyle>("SmallText"))
-					.ToolTipText(FText::Format(LOCTEXT("ClassHyperlinkTooltipFormat", "Open class '{0}'"), FText::FromString(ClassInfo.PathName)))
-					.OnNavigate_Lambda([ClassInfo]()
-					{
-						GEditor->GetEditorSubsystem<UAssetEditorSubsystem>()->OpenEditorForAsset(ClassInfo.PathName);
-					});
-#else
-				return 
-					SNew(STextBlock)
-					.Font(FCoreStyle::Get().GetFontStyle("SmallFont"))
-					.Text(FText::FromString(ClassInfo.Name))
-					.ToolTipText(FText::FromString(ClassInfo.PathName));
-#endif			
-			}
-		}
-		break;
-		}
-
-		return SNullWidget::NullWidget; 
-	}
-
 	const TraceServices::IAnalysisSession& AnalysisSession;
 
 	FText Name;
@@ -322,9 +160,10 @@ public:
 class FAnimGraphSchematicNode : public TSharedFromThis<FAnimGraphSchematicNode>
 {
 public:
-	FAnimGraphSchematicNode(int32 InNodeId, const FText& InType, const TraceServices::IAnalysisSession& InAnalysisSession)
+	FAnimGraphSchematicNode(int32 InNodeId, uint64 InAnimInstanceId, const FText& InType, const TraceServices::IAnalysisSession& InAnalysisSession)
 		: AnalysisSession(InAnalysisSession)
 		, NodeId(InNodeId)
+		, AnimInstanceId(InAnimInstanceId)
 		, Type(InType)
 		, FilterState(EAnimGraphSchematicFilterState::Hidden)
 		, bLinearized(false)
@@ -347,7 +186,8 @@ public:
 
 	const TraceServices::IAnalysisSession& AnalysisSession;
 
-	int32 NodeId;
+	int32 NodeId = INDEX_NONE;
+	uint64 AnimInstanceId = INDEX_NONE;
 
 	FText Type;
 
@@ -393,29 +233,123 @@ class SAnimGraphSchematicNode : public SMultiColumnTableRow<TSharedRef<FAnimGrap
 
 		if (InColumnName == AnimGraphSchematicColumns::Type)
 		{
-			return 
-				SNew(SBorder)
-				.BorderImage(bIsRoot ? FGameplayInsightsStyle::Get().GetBrush("SchematicViewRootLeft") : FCoreStyle::Get().GetBrush("NoBorder"))
-				[
-					SNew(SHorizontalBox)
-					+SHorizontalBox::Slot()
-					.AutoWidth()
-					.VAlign(VAlign_Center)
-					.Padding(6, 0, 0, 0)
-					[
-						SNew(SExpanderArrow, SharedThis(this))
-						.IndentAmount(12)
-					]
-					+ SHorizontalBox::Slot()
-					.FillWidth(1.0f)
-					.VAlign(VAlign_Center)
-					[
-						SNew(STextBlock)
-						.Font(FCoreStyle::Get().GetFontStyle(bIsRoot ? "ExpandableArea.TitleFont" : "SmallFont"))
-						.Text(Node->Type)
-						.HighlightText(FilterText)
-					]
-				];
+#if WITH_EDITOR
+			IRewindDebugger* RewindDebugger = IRewindDebugger::Instance();
+
+			if (const TraceServices::IAnalysisSession* AnalysisSession = RewindDebugger->GetAnalysisSession())
+			{
+				if (const FGameplayProvider* GameplayProvider = AnalysisSession->ReadProvider<FGameplayProvider>(FGameplayProvider::ProviderName))
+				{
+					TraceServices::FAnalysisSessionReadScope SessionReadScope(*AnalysisSession);
+
+					const FObjectInfo* ObjectInfo = GameplayProvider->FindObjectInfo(Node->AnimInstanceId);
+					const FClassInfo* AnimInstanceClassInfo = ObjectInfo ? GameplayProvider->FindClassInfo(ObjectInfo->ClassId) : nullptr;
+					const int32 AnimNodeId = Node->NodeId;
+
+					if (ObjectInfo && AnimInstanceClassInfo)
+					{
+
+						return 
+							SNew(SBorder)
+							.BorderImage(bIsRoot ? FGameplayInsightsStyle::Get().GetBrush("SchematicViewRootLeft") : FCoreStyle::Get().GetBrush("NoBorder"))
+							[
+								SNew(SHorizontalBox)
+								+SHorizontalBox::Slot()
+								.AutoWidth()
+								.VAlign(VAlign_Center)
+								.Padding(6, 0, 0, 0)
+								[
+									SNew(SExpanderArrow, SharedThis(this))
+									.IndentAmount(12)
+								]
+								+SHorizontalBox::Slot()
+								.AutoWidth()
+								.VAlign(VAlign_Center)
+								[
+									SNew(SHyperlink)
+									.Text(Node->Type)
+									.TextStyle(&FCoreStyle::Get().GetWidgetStyle<FTextBlockStyle>("SmallText"))
+									.ToolTipText(FText::Format(LOCTEXT("AssetHyperlinkTooltipFormat", "Open node '{0}'"), Node->Type))
+									.HighlightText(FilterText)
+									.OnNavigate_Lambda([AnimNodeId, AnimInstanceClassInfo]()
+									{
+										TSoftObjectPtr<UAnimBlueprintGeneratedClass> InstanceClass;
+										InstanceClass = FSoftObjectPath(AnimInstanceClassInfo->PathName);
+
+										if (InstanceClass.LoadSynchronous())
+										{
+											if (UAnimBlueprint* AnimBlueprint = Cast<UAnimBlueprint>(InstanceClass.Get()->ClassGeneratedBy))
+											{
+												GEditor->GetEditorSubsystem<UAssetEditorSubsystem>()->OpenEditorForAsset(AnimBlueprint);
+
+												if (IAnimationBlueprintEditor* AnimBlueprintEditor = static_cast<IAnimationBlueprintEditor*>(GEditor->GetEditorSubsystem<UAssetEditorSubsystem>()->FindEditorForAsset(AnimBlueprint, true)))
+												{
+													int32 AnimNodeIndex = InstanceClass.Get()->GetAnimNodeProperties().Num() - AnimNodeId - 1;
+													TWeakObjectPtr<const UEdGraphNode>* GraphNode = InstanceClass.Get()->AnimBlueprintDebugData.NodePropertyIndexToNodeMap.Find(AnimNodeIndex);
+													if (GraphNode != nullptr && GraphNode->Get())
+													{
+														AnimBlueprintEditor->JumpToHyperlink(GraphNode->Get());
+													}
+												}
+											}
+										}
+									})
+								]
+							];
+#else
+						return 
+							SNew(SBorder)
+							.BorderImage(bIsRoot ? FGameplayInsightsStyle::Get().GetBrush("SchematicViewRootLeft") : FCoreStyle::Get().GetBrush("NoBorder"))
+							[
+								SNew(SHorizontalBox)
+								+SHorizontalBox::Slot()
+								.AutoWidth()
+								.VAlign(VAlign_Center)
+								.Padding(6, 0, 0, 0)
+								[
+									SNew(SExpanderArrow, SharedThis(this))
+									.IndentAmount(12)
+								]
+								+SHorizontalBox::Slot()
+								.AutoWidth()
+								.VAlign(VAlign_Center)
+								[
+									SNew(SHyperlink)
+									.Text(Node->Type)
+									.TextStyle(&FCoreStyle::Get().GetWidgetStyle<FTextBlockStyle>("SmallText"))
+									.HighlightText(FilterText)
+								]
+							];
+#endif
+					}
+					else
+					{
+						return 
+							SNew(SBorder)
+							.BorderImage(bIsRoot ? FGameplayInsightsStyle::Get().GetBrush("SchematicViewRootLeft") : FCoreStyle::Get().GetBrush("NoBorder"))
+							[
+								SNew(SHorizontalBox)
+								+SHorizontalBox::Slot()
+								.AutoWidth()
+								.VAlign(VAlign_Center)
+								.Padding(6, 0, 0, 0)
+								[
+									SNew(SExpanderArrow, SharedThis(this))
+									.IndentAmount(12)
+								]
+								+SHorizontalBox::Slot()
+								.AutoWidth()
+								.VAlign(VAlign_Center)
+								[
+									SNew(SHyperlink)
+									.Text(Node->Type)
+									.TextStyle(&FCoreStyle::Get().GetWidgetStyle<FTextBlockStyle>("SmallText"))
+									.HighlightText(FilterText)
+								]
+							];
+					}
+				}
+			}
 		}
 		else
 		{
@@ -423,7 +357,7 @@ class SAnimGraphSchematicNode : public SMultiColumnTableRow<TSharedRef<FAnimGrap
 			TSharedRef<FAnimNodeValueMessage>* ValuePtr = Node->KeysAndValues.Find(InColumnName);
 			if(ValuePtr)
 			{
-				ValueWidget = FAnimGraphSchematicPropertyNode::StaticMakeValueWidget(Node->AnalysisSession, *ValuePtr);
+				ValueWidget = SVariantValueView::MakeVariantValueWidget(Node->AnalysisSession, (*ValuePtr)->Value, FText::GetEmpty());
 			}
 
 			return
@@ -622,7 +556,7 @@ void SAnimGraphSchematicView::RefreshNodes()
 									TSharedRef<FAnimGraphSchematicNode>* ExistingNodePtr = NodeMap.Find(InMessage.NodeId);
 									if(ExistingNodePtr == nullptr)
 									{
-										Node = NodeMap.Add(InMessage.NodeId, MakeShared<FAnimGraphSchematicNode>(InMessage.NodeId, FText::FromString(InMessage.NodeName), *AnalysisSession));
+										Node = NodeMap.Add(InMessage.NodeId, MakeShared<FAnimGraphSchematicNode>(InMessage.NodeId, AnimInstanceId, FText::FromString(InMessage.NodeName), *AnalysisSession));
 
 										// Add dummy values for weight and root motion weight
 										TSharedRef<FAnimNodeValueMessage> WeightMessage = MakeShared<FAnimNodeValueMessage>();
diff --git a/Engine/Plugins/Animation/GameplayInsights/Source/GameplayInsights/Private/SInertializationDetailsView.cpp b/Engine/Plugins/Animation/GameplayInsights/Source/GameplayInsights/Private/SInertializationDetailsView.cpp
--- a/Engine/Plugins/Animation/GameplayInsights/Source/GameplayInsights/Private/SInertializationDetailsView.cpp
+++ b/Engine/Plugins/Animation/GameplayInsights/Source/GameplayInsights/Private/SInertializationDetailsView.cpp
@@ -11,13 +11,8 @@
 
 struct FInertializationDetailsNodeItem
 {
-	const TCHAR* Type = nullptr;
-	float ElapsedTime = 0.0f;
-	float Duration = 0.0f;
-	float MaxDuration = 0.0f;
-	float InertializationWeight = 0.0f;
-	bool bActive = false;
-	const TCHAR* Request = nullptr;
+	TArray<const TCHAR*, TInlineAllocator<16>> PropertyNames;
+	TArray<FVariantValue, TInlineAllocator<16>> PropertyValues;
 };
 
 void SInertializationDetailsView::GetVariantsAtFrame(const TraceServices::FFrame& InFrame, TArray<TSharedRef<FVariantTreeNode>>& OutVariants) const
@@ -25,7 +20,7 @@ void SInertializationDetailsView::GetVariantsAtFrame(const TraceServices::FFrame
 	const FGameplayProvider* GameplayProvider = AnalysisSession->ReadProvider<FGameplayProvider>(FGameplayProvider::ProviderName);
 	const FAnimationProvider* AnimationProvider = AnalysisSession->ReadProvider<FAnimationProvider>(FAnimationProvider::ProviderName);
 
-	TSortedMap<int32, FInertializationDetailsNodeItem, TInlineAllocator<8>> NodeMap;
+	TSortedMap<TTuple<int32, uint64>, FInertializationDetailsNodeItem, TInlineAllocator<8>> NodeMap;
 
 	if(GameplayProvider && AnimationProvider)
 	{
@@ -48,8 +43,7 @@ void SInertializationDetailsView::GetVariantsAtFrame(const TraceServices::FFrame
 					{
 						if (FCString::Strcmp(InertializationNodeType, InMessage.NodeTypeName) == 0)
 						{
-							FInertializationDetailsNodeItem& NodeItemRef = NodeMap.FindOrAdd(InMessage.NodeId);
-							NodeItemRef.Type = InMessage.NodeName;
+							NodeMap.FindOrAdd({ InMessage.NodeId, InMessage.AnimInstanceId });
 							break;
 						}
 					}
@@ -67,32 +61,10 @@ void SInertializationDetailsView::GetVariantsAtFrame(const TraceServices::FFrame
 			{
 				if (InStartTime >= InFrame.StartTime && InEndTime <= InFrame.EndTime)
 				{
-					if (FInertializationDetailsNodeItem* NodeItem = NodeMap.Find(InMessage.NodeId))
+					if (FInertializationDetailsNodeItem* NodeItem = NodeMap.Find({ InMessage.NodeId, InMessage.AnimInstanceId }))
 					{
-						if (FCString::Strcmp(InMessage.Key, TEXT("State")) == 0)
-						{
-							NodeItem->bActive = FCString::Strcmp(InMessage.Value.String.Value, TEXT("EInertializationState::Active")) == 0;
-						}
-						else if (FCString::Strcmp(InMessage.Key, TEXT("Elapsed Time")) == 0)
-						{
-							NodeItem->ElapsedTime = InMessage.Value.Float.Value;
-						}
-						else if (FCString::Strcmp(InMessage.Key, TEXT("Duration")) == 0)
-						{
-							NodeItem->Duration = InMessage.Value.Float.Value;
-						}
-						else if (FCString::Strcmp(InMessage.Key, TEXT("Max Duration")) == 0)
-						{
-							NodeItem->MaxDuration = InMessage.Value.Float.Value;
-						}
-						else if (FCString::Strcmp(InMessage.Key, TEXT("Inertialization Weight")) == 0)
-						{
-							NodeItem->InertializationWeight = InMessage.Value.Float.Value;
-						}
-						else if (FCString::Strcmp(InMessage.Key, TEXT("Request")) == 0)
-						{
-							NodeItem->Request = InMessage.Value.String.Value;
-						}
+						NodeItem->PropertyNames.Add(InMessage.Key);
+						NodeItem->PropertyValues.Add(InMessage.Value);
 					}
 				}
 				return TraceServices::EEventEnumerate::Continue;
@@ -101,16 +73,36 @@ void SInertializationDetailsView::GetVariantsAtFrame(const TraceServices::FFrame
 
 		// Make UI Variant Items
 
-		for (TPair<int32, FInertializationDetailsNodeItem>& NodePair : NodeMap)
+		for (TPair<TTuple<int32, uint64>, FInertializationDetailsNodeItem>& NodePair : NodeMap)
 		{
-			if (NodePair.Value.bActive)
+			bool bInertializationActive = false;
+
+			const int32 PropertyNum = NodePair.Value.PropertyNames.Num();
+
+			// Check if inertialization node is Active
+
+			for (int32 PropertyIdx = 0; PropertyIdx < PropertyNum; PropertyIdx++)
+			{
+				if (FCString::Strcmp(NodePair.Value.PropertyNames[PropertyIdx], TEXT("State")) == 0 &&
+					FCString::Strcmp(NodePair.Value.PropertyValues[PropertyIdx].String.Value, TEXT("EInertializationState::Active")) == 0)
+				{
+					bInertializationActive = true;
+				}
+			}
+
+			// If node is active then add to details view
+
+			if (bInertializationActive)
 			{
-				TSharedRef<FVariantTreeNode> Header = OutVariants.Add_GetRef(FVariantTreeNode::MakeHeader(NodePair.Value.Type ? FText::FromString(NodePair.Value.Type) : LOCTEXT("UnknownNode", "Unknown Node"), NodePair.Key));
-				Header->AddChild(FVariantTreeNode::MakeFloat(LOCTEXT("ElapsedTime", "Elapsed Time"), NodePair.Value.ElapsedTime));
-				Header->AddChild(FVariantTreeNode::MakeFloat(LOCTEXT("Duration", "Duration"), NodePair.Value.Duration));
-				Header->AddChild(FVariantTreeNode::MakeFloat(LOCTEXT("MaxDuration", "Max Duration"), NodePair.Value.MaxDuration));
-				Header->AddChild(FVariantTreeNode::MakeFloat(LOCTEXT("InertializationWeight", "Inertialization Weight"), NodePair.Value.InertializationWeight));
-				Header->AddChild(FVariantTreeNode::MakeString(LOCTEXT("RequestNode", "Request Node"), NodePair.Value.Request));
+				TSharedRef<FVariantTreeNode> Header = OutVariants.Add_GetRef(
+					FVariantTreeNode::MakeAnimNode(LOCTEXT("InertializationNode", "Inertialization Node"), NodePair.Key.Get<0>(), NodePair.Key.Get<1>()));
+
+				for (int32 PropertyIdx = 0; PropertyIdx < PropertyNum; PropertyIdx++)
+				{
+					Header->AddChild(MakeShared<FVariantTreeNode>(
+						FText::FromString(NodePair.Value.PropertyNames[PropertyIdx]), 
+						NodePair.Value.PropertyValues[PropertyIdx], 0, INDEX_NONE));
+				}
 			}
 		}
 	}
diff --git a/Engine/Plugins/Animation/GameplayInsights/Source/GameplayInsights/Private/SVariantValueView.cpp b/Engine/Plugins/Animation/GameplayInsights/Source/GameplayInsights/Private/SVariantValueView.cpp
--- a/Engine/Plugins/Animation/GameplayInsights/Source/GameplayInsights/Private/SVariantValueView.cpp
+++ b/Engine/Plugins/Animation/GameplayInsights/Source/GameplayInsights/Private/SVariantValueView.cpp
@@ -4,6 +4,7 @@
 #include "Widgets/Input/SSearchBox.h"
 #include "TraceServices/Model/AnalysisSession.h"
 #include "GameplayProvider.h"
+#include "AnimationProvider.h"
 #include "TraceServices/Model/Frames.h"
 #include "Widgets/Layout/SBorder.h"
 #include "Widgets/SBoxPanel.h"
@@ -22,8 +23,11 @@
 #if WITH_EDITOR
 #include "AnimPreviewInstance.h"
 #include "Animation/DebugSkelMeshComponent.h"
+#include "Animation/AnimBlueprint.h"
+#include "Animation/AnimBlueprintGeneratedClass.h"
 #include "Editor.h"
 #include "IAnimationEditor.h"
+#include "IAnimationBlueprintEditor.h"
 #include "IPersonaToolkit.h"
 #include "SourceCodeNavigation.h"
 #include "Subsystems/AssetEditorSubsystem.h"
@@ -37,7 +41,7 @@ namespace VariantColumns
 	static const FName Value("Value");
 };
 
-static TSharedRef<SWidget> MakeVariantValueWidget(const TraceServices::IAnalysisSession& InAnalysisSession, const FVariantValue& InValue, const TAttribute<FText>& InHighlightText) 
+TSharedRef<SWidget> SVariantValueView::MakeVariantValueWidget(const TraceServices::IAnalysisSession& InAnalysisSession, const FVariantValue& InValue, const TAttribute<FText>& InHighlightText)
 { 
 	switch(InValue.Type)
 	{
@@ -248,6 +252,62 @@ static TSharedRef<SWidget> MakeVariantValueWidget(const TraceServices::IAnalysis
 #endif			
 		}
 	}
+
+	case EAnimNodeValueType::AnimNode:
+	{
+		const FGameplayProvider* GameplayProvider = InAnalysisSession.ReadProvider<FGameplayProvider>(FGameplayProvider::ProviderName);
+		const FAnimationProvider* AnimationProvider = InAnalysisSession.ReadProvider<FAnimationProvider>(FAnimationProvider::ProviderName);
+
+		if (GameplayProvider && AnimationProvider)
+		{
+			TraceServices::FAnalysisSessionReadScope SessionReadScope(InAnalysisSession);
+
+			const FAnimNodeInfo* AnimNodeInfo = AnimationProvider->FindAnimNodeInfo(InValue.AnimNode.Value, InValue.AnimNode.AnimInstanceId);
+			const FObjectInfo* ObjectInfo = GameplayProvider->FindObjectInfo(InValue.AnimNode.AnimInstanceId);
+			const FClassInfo* AnimInstanceClassInfo = ObjectInfo ? GameplayProvider->FindClassInfo(ObjectInfo->ClassId) : nullptr;
+
+			if (AnimNodeInfo && ObjectInfo && AnimInstanceClassInfo)
+			{
+#if WITH_EDITOR
+				return
+					SNew(SHyperlink)
+					.Text(FText::FromString(AnimNodeInfo->Name))
+					.TextStyle(&FCoreStyle::Get().GetWidgetStyle<FTextBlockStyle>("SmallText"))
+					.ToolTipText(FText::Format(LOCTEXT("AssetHyperlinkTooltipFormat", "Open node '{0}'"), FText::FromString(AnimNodeInfo->Name)))
+					.OnNavigate_Lambda([ObjectInfo, AnimNodeInfo, AnimInstanceClassInfo, InValue]()
+					{
+						TSoftObjectPtr<UAnimBlueprintGeneratedClass> InstanceClass;
+						InstanceClass = FSoftObjectPath(AnimInstanceClassInfo->PathName);
+
+						if (InstanceClass.LoadSynchronous())
+						{
+							if (UAnimBlueprint* AnimBlueprint = Cast<UAnimBlueprint>(InstanceClass.Get()->ClassGeneratedBy))
+							{
+								GEditor->GetEditorSubsystem<UAssetEditorSubsystem>()->OpenEditorForAsset(AnimBlueprint);
+
+								if (IAnimationBlueprintEditor* AnimBlueprintEditor = static_cast<IAnimationBlueprintEditor*>(GEditor->GetEditorSubsystem<UAssetEditorSubsystem>()->FindEditorForAsset(AnimBlueprint, true)))
+								{
+									int32 AnimNodeIndex = InstanceClass.Get()->GetAnimNodeProperties().Num() - AnimNodeInfo->Id - 1;
+									TWeakObjectPtr<const UEdGraphNode>* GraphNode = InstanceClass.Get()->AnimBlueprintDebugData.NodePropertyIndexToNodeMap.Find(AnimNodeIndex);
+									if (GraphNode != nullptr && GraphNode->Get())
+									{
+										AnimBlueprintEditor->JumpToHyperlink(GraphNode->Get());
+									}
+								}
+							}
+						}
+					});
+#else
+				return
+					SNew(STextBlock)
+					.Font(FCoreStyle::Get().GetFontStyle("SmallFont"))
+					.Text(FText::FromString(AnimNodeInfo->Name))
+					.ToolTipText(FText::FromString(AnimNodeInfo->Name))
+					.HighlightText(InHighlightText);
+#endif
+			}
+		}
+	}
 	break;
 	}
 
@@ -328,7 +388,7 @@ public:
 					.AutoWidth()
 					.VAlign(VAlign_Center)
 					[
-						MakeVariantValueWidget(*AnalysisSession, Node->GetValue(), HighlightText)
+						SVariantValueView::MakeVariantValueWidget(*AnalysisSession, Node->GetValue(), HighlightText)
 					]
 				];
 		}
diff --git a/Engine/Plugins/Animation/GameplayInsights/Source/GameplayInsights/Private/SVariantValueView.h b/Engine/Plugins/Animation/GameplayInsights/Source/GameplayInsights/Private/SVariantValueView.h
--- a/Engine/Plugins/Animation/GameplayInsights/Source/GameplayInsights/Private/SVariantValueView.h
+++ b/Engine/Plugins/Animation/GameplayInsights/Source/GameplayInsights/Private/SVariantValueView.h
@@ -9,6 +9,7 @@
 #include "Containers/StringFwd.h"
 
 namespace TraceServices { class IAnalysisSession; }
+struct FVariantValue;
 struct FVariantTreeNode;
 class FTextFilterExpressionEvaluator;
 enum class EVariantTreeNodeFilterState;
@@ -40,6 +41,9 @@ public:
 	/** @return Trace session that the view references */
 	const TraceServices::IAnalysisSession&  GetAnalysisSession() const { return *AnalysisSession; }
 	
+	/** @return New Widget for the given Variant Value */
+	static TSharedRef<SWidget> MakeVariantValueWidget(const TraceServices::IAnalysisSession& InAnalysisSession, const FVariantValue& InValue, const TAttribute<FText>& InHighlightText);
+
 private:
 	// SWidget interface
 	virtual FReply OnKeyDown(const FGeometry& MyGeometry, const FKeyEvent& InKeyEvent) override;
diff --git a/Engine/Plugins/Animation/GameplayInsights/Source/GameplayInsights/Private/VariantTreeNode.cpp b/Engine/Plugins/Animation/GameplayInsights/Source/GameplayInsights/Private/VariantTreeNode.cpp
--- a/Engine/Plugins/Animation/GameplayInsights/Source/GameplayInsights/Private/VariantTreeNode.cpp
+++ b/Engine/Plugins/Animation/GameplayInsights/Source/GameplayInsights/Private/VariantTreeNode.cpp
@@ -3,6 +3,7 @@
 #include "VariantTreeNode.h"
 #include "TraceServices/Model/AnalysisSession.h"
 #include "GameplayProvider.h"
+#include "AnimationProvider.h"
 
 #define LOCTEXT_NAMESPACE "FVariantTreeNode"
 
@@ -46,6 +47,19 @@ FString FVariantTreeNode::GetValueAsString(const TraceServices::IAnalysisSession
 				return ClassInfo.PathName;
 			}
 
+			break;
+		}
+	case EAnimNodeValueType::AnimNode:
+		{
+			const FAnimationProvider* AnimationProvider = InAnalysisSession.ReadProvider<FAnimationProvider>(FAnimationProvider::ProviderName);
+			if (AnimationProvider)
+			{
+				TraceServices::FAnalysisSessionReadScope SessionReadScope(InAnalysisSession);
+
+				const FAnimNodeInfo* AnimNodeInfo = AnimationProvider->FindAnimNodeInfo(Value.AnimNode.Value, Value.AnimNode.AnimInstanceId);
+				return AnimNodeInfo ? AnimNodeInfo->Name : TEXT("");
+			}
+
 			break;
 		}
 	}
diff --git a/Engine/Plugins/Animation/GameplayInsights/Source/GameplayInsights/Private/VariantTreeNode.h b/Engine/Plugins/Animation/GameplayInsights/Source/GameplayInsights/Private/VariantTreeNode.h
--- a/Engine/Plugins/Animation/GameplayInsights/Source/GameplayInsights/Private/VariantTreeNode.h
+++ b/Engine/Plugins/Animation/GameplayInsights/Source/GameplayInsights/Private/VariantTreeNode.h
@@ -146,6 +146,16 @@ struct FVariantTreeNode : TSharedFromThis<FVariantTreeNode>
 		return MakeShared<FVariantTreeNode>(InName, Value, InId);
 	}
 
+	static TSharedRef<FVariantTreeNode> MakeAnimNode(const FText& InName, uint64 InValue, uint64 InAnimInstanceId, uint64 InId = INDEX_NONE)
+	{
+		FVariantValue Value;
+		Value.Type = EAnimNodeValueType::AnimNode;
+		Value.AnimNode.Value = InValue;
+		Value.AnimNode.AnimInstanceId = InAnimInstanceId;
+
+		return MakeShared<FVariantTreeNode>(InName, Value, InId);
+	}
+
 private:
 	FText Name;
 	FVariantValue Value;
diff --git a/Engine/Plugins/Animation/GameplayInsights/Source/GameplayInsights/Public/IAnimationProvider.h b/Engine/Plugins/Animation/GameplayInsights/Source/GameplayInsights/Public/IAnimationProvider.h
--- a/Engine/Plugins/Animation/GameplayInsights/Source/GameplayInsights/Public/IAnimationProvider.h
+++ b/Engine/Plugins/Animation/GameplayInsights/Source/GameplayInsights/Public/IAnimationProvider.h
@@ -15,6 +15,14 @@ struct FSkeletalMeshInfo
 	uint32 BoneCount = 0;
 };
 
+struct FAnimNodeInfo
+{
+	int32 Id = 0;
+	uint64 AnimInstanceId = 0;
+	const TCHAR* Name = nullptr;
+	const TCHAR* TypeName = nullptr;
+};
+
 struct FSkeletalMeshNamedCurve
 {
 	uint32 Id = 0;
@@ -120,6 +128,7 @@ enum class EAnimNodeValueType : uint8
 	String,
 	Object,
 	Class,
+	AnimNode,
 };
 
 struct FVariantValue
@@ -180,6 +189,11 @@ struct FVariantValue
 		{
 			uint64 Value;
 		} Class;
+		struct
+		{
+			int32 Value;
+			uint64 AnimInstanceId;
+		} AnimNode;
 	};
 
 	EAnimNodeValueType Type = EAnimNodeValueType::Bool;
@@ -316,6 +330,7 @@ public:
 	virtual bool ReadAnimSyncTimeline(uint64 InObjectId, TFunctionRef<void(const AnimSyncTimeline&)> Callback) const = 0;
 	virtual bool ReadPoseWatchTimeline(uint64 InObjectId, TFunctionRef<void(const PoseWatchTimeline&)> Callback) const = 0;
 	virtual const FSkeletalMeshInfo* FindSkeletalMeshInfo(uint64 InObjectId) const = 0;
+	virtual const FAnimNodeInfo* FindAnimNodeInfo(int32 InNodeId, uint64 InAnimInstanceId) const = 0;
 	virtual const TCHAR* GetName(uint32 InId) const = 0;
 	virtual FText FormatNodeKeyValue(const FAnimNodeValueMessage& InMessage) const = 0;
 	virtual FText FormatNodeValue(const FAnimNodeValueMessage& InMessage) const = 0;
diff --git a/Engine/Plugins/Experimental/Animation/PoseSearch/Source/Runtime/Private/PoseSearchLibrary.cpp b/Engine/Plugins/Experimental/Animation/PoseSearch/Source/Runtime/Private/PoseSearchLibrary.cpp
--- a/Engine/Plugins/Experimental/Animation/PoseSearch/Source/Runtime/Private/PoseSearchLibrary.cpp
+++ b/Engine/Plugins/Experimental/Animation/PoseSearch/Source/Runtime/Private/PoseSearchLibrary.cpp
@@ -70,6 +70,8 @@ static void RequestInertialBlend(const FAnimationUpdateContext& Context, float B
 			Request.Duration = BlendTime;
 #if ANIM_TRACE_ENABLED
 			Request.Description = LOCTEXT("InertializationRequestDescription", "Motion Matching");
+			Request.NodeId = Context.GetCurrentNodeId();
+			Request.AnimInstance = Context.AnimInstanceProxy->GetAnimInstanceObject();
 #endif
 
 			InertializationRequester->RequestInertialization(Request);
diff --git a/Engine/Source/Runtime/AnimGraphRuntime/Private/AnimNodes/AnimNode_BlendListBase.cpp b/Engine/Source/Runtime/AnimGraphRuntime/Private/AnimNodes/AnimNode_BlendListBase.cpp
--- a/Engine/Source/Runtime/AnimGraphRuntime/Private/AnimNodes/AnimNode_BlendListBase.cpp
+++ b/Engine/Source/Runtime/AnimGraphRuntime/Private/AnimNodes/AnimNode_BlendListBase.cpp
@@ -132,7 +132,8 @@ void FAnimNode_BlendListBase::Update_AnyThread(const FAnimationUpdateContext& Co
 					Request.BlendMode = GetBlendType();
 					Request.CustomBlendCurve = GetCustomBlendCurve();
 #if ANIM_TRACE_ENABLED
-					Request.Description = NSLOCTEXT("AnimNode_BlendListBase", "InertializationRequestDescription", "Blend List");
+					Request.NodeId = Context.GetCurrentNodeId();
+					Request.AnimInstance = Context.AnimInstanceProxy->GetAnimInstanceObject();
 #endif
 
 					InertializationRequester->RequestInertialization(Request);
diff --git a/Engine/Source/Runtime/AnimGraphRuntime/Private/AnimNodes/AnimNode_Mirror.cpp b/Engine/Source/Runtime/AnimGraphRuntime/Private/AnimNodes/AnimNode_Mirror.cpp
--- a/Engine/Source/Runtime/AnimGraphRuntime/Private/AnimNodes/AnimNode_Mirror.cpp
+++ b/Engine/Source/Runtime/AnimGraphRuntime/Private/AnimNodes/AnimNode_Mirror.cpp
@@ -146,7 +146,8 @@ void FAnimNode_MirrorBase::Update_AnyThread(const FAnimationUpdateContext& Conte
 				FInertializationRequest Request;
 				Request.Duration = GetBlendTimeOnMirrorStateChange();
 #if ANIM_TRACE_ENABLED
-				Request.Description = LOCTEXT("InertializationRequestDescription", "Mirror");
+				Request.NodeId = Context.GetCurrentNodeId();
+				Request.AnimInstance = Context.AnimInstanceProxy->GetAnimInstanceObject();
 #endif
 
 				InertializationRequester->RequestInertialization(Request);
diff --git a/Engine/Source/Runtime/AnimGraphRuntime/Private/AnimNodes/AnimNode_Slot.cpp b/Engine/Source/Runtime/AnimGraphRuntime/Private/AnimNodes/AnimNode_Slot.cpp
--- a/Engine/Source/Runtime/AnimGraphRuntime/Private/AnimNodes/AnimNode_Slot.cpp
+++ b/Engine/Source/Runtime/AnimGraphRuntime/Private/AnimNodes/AnimNode_Slot.cpp
@@ -52,7 +52,8 @@ void FAnimNode_Slot::Update_AnyThread(const FAnimationUpdateContext& Context)
 			Request.Duration = InertializationRequest.Get<0>();
 			Request.BlendProfile = InertializationRequest.Get<1>();
 #if ANIM_TRACE_ENABLED
-			Request.Description = NSLOCTEXT("AnimNode_Slot", "InertializationRequestDescription", "Slot");
+			Request.NodeId = Context.GetCurrentNodeId();
+			Request.AnimInstance = Context.AnimInstanceProxy->GetAnimInstanceObject();
 #endif
 
 			InertializationRequester->RequestInertialization(Request);
diff --git a/Engine/Source/Runtime/AnimGraphRuntime/Private/BlendSpacePlayerLibrary.cpp b/Engine/Source/Runtime/AnimGraphRuntime/Private/BlendSpacePlayerLibrary.cpp
--- a/Engine/Source/Runtime/AnimGraphRuntime/Private/BlendSpacePlayerLibrary.cpp
+++ b/Engine/Source/Runtime/AnimGraphRuntime/Private/BlendSpacePlayerLibrary.cpp
@@ -4,6 +4,7 @@
 #include "Animation/AnimNode_Inertialization.h"
 #include "AnimNodes/AnimNode_BlendSpacePlayer.h"
 #include "Animation/AnimTrace.h"
+#include "Animation/AnimInstanceProxy.h"
 
 #include UE_INLINE_GENERATED_CPP_BY_NAME(BlendSpacePlayerLibrary)
 
@@ -56,7 +57,8 @@ FBlendSpacePlayerReference UBlendSpacePlayerLibrary::SetBlendSpaceWithInertialBl
 						FInertializationRequest Request;
 						Request.Duration = BlendTime;
 #if ANIM_TRACE_ENABLED
-						Request.Description = NSLOCTEXT("BlendSpacePlayerLibrary", "InertializationRequestDescription", "Blend Space Player");
+						Request.NodeId = AnimationUpdateContext->GetCurrentNodeId();
+						Request.AnimInstance = AnimationUpdateContext->AnimInstanceProxy->GetAnimInstanceObject();
 #endif
 
 						InertializationRequester->RequestInertialization(Request);
diff --git a/Engine/Source/Runtime/AnimGraphRuntime/Private/SequenceEvaluatorLibrary.cpp b/Engine/Source/Runtime/AnimGraphRuntime/Private/SequenceEvaluatorLibrary.cpp
--- a/Engine/Source/Runtime/AnimGraphRuntime/Private/SequenceEvaluatorLibrary.cpp
+++ b/Engine/Source/Runtime/AnimGraphRuntime/Private/SequenceEvaluatorLibrary.cpp
@@ -6,6 +6,7 @@
 #include "AnimNodes/AnimNode_SequenceEvaluator.h"
 #include "AnimationRuntime.h"
 #include "Animation/AnimTrace.h"
+#include "Animation/AnimInstanceProxy.h"
 
 #include UE_INLINE_GENERATED_CPP_BY_NAME(SequenceEvaluatorLibrary)
 
@@ -94,7 +95,8 @@ FSequenceEvaluatorReference USequenceEvaluatorLibrary::SetSequenceWithInertialBl
 						FInertializationRequest Request;
 						Request.Duration = BlendTime;
 #if ANIM_TRACE_ENABLED
-						Request.Description = NSLOCTEXT("SequenceEvaluatorLibrary", "InertializationRequestDescription", "Sequence Evaluator");
+						Request.NodeId = AnimationUpdateContext->GetCurrentNodeId();
+						Request.AnimInstance = AnimationUpdateContext->AnimInstanceProxy->GetAnimInstanceObject();
 #endif
 
 						InertializationRequester->RequestInertialization(Request);
diff --git a/Engine/Source/Runtime/AnimGraphRuntime/Private/SequencePlayerLibrary.cpp b/Engine/Source/Runtime/AnimGraphRuntime/Private/SequencePlayerLibrary.cpp
--- a/Engine/Source/Runtime/AnimGraphRuntime/Private/SequencePlayerLibrary.cpp
+++ b/Engine/Source/Runtime/AnimGraphRuntime/Private/SequencePlayerLibrary.cpp
@@ -5,6 +5,7 @@
 #include "Animation/AnimNode_Inertialization.h"
 #include "Animation/AnimNode_SequencePlayer.h"
 #include "Animation/AnimTrace.h"
+#include "Animation/AnimInstanceProxy.h"
 
 #include UE_INLINE_GENERATED_CPP_BY_NAME(SequencePlayerLibrary)
 
@@ -95,7 +96,8 @@ FSequencePlayerReference USequencePlayerLibrary::SetSequenceWithInertialBlending
 						FInertializationRequest Request;
 						Request.Duration = BlendTime;
 #if ANIM_TRACE_ENABLED
-						Request.Description = NSLOCTEXT("SequencePlayerLibrary", "InertializationRequestDescription", "Sequence Player");
+						Request.NodeId = AnimationUpdateContext->GetCurrentNodeId();
+						Request.AnimInstance = AnimationUpdateContext->AnimInstanceProxy->GetAnimInstanceObject();
 #endif
 
 						InertializationRequester->RequestInertialization(Request);
diff --git a/Engine/Source/Runtime/Engine/Classes/Animation/AnimNode_DeadBlending.h b/Engine/Source/Runtime/Engine/Classes/Animation/AnimNode_DeadBlending.h
--- a/Engine/Source/Runtime/Engine/Classes/Animation/AnimNode_DeadBlending.h
+++ b/Engine/Source/Runtime/Engine/Classes/Animation/AnimNode_DeadBlending.h
@@ -212,6 +212,13 @@ private:
 	// Description for the current inertialization request - used for debugging
 	FText InertializationRequestDescription;
 
+	// Node Id for the current inertialization request - used for debugging
+	int32 InertializationRequestNodeId = INDEX_NONE;
+
+	// Anim Instance for the current inertialization request - used for debugging
+	UPROPERTY(Transient)
+	TObjectPtr<UObject> InertializationRequestAnimInstance = nullptr;
+
 	// Current inertialization durations for each bone, indexed by skeleton bone index (used for per-bone blending).
 	TCustomBoneIndexArray<float, FSkeletonPoseBoneIndex> InertializationDurationPerBone;
 
diff --git a/Engine/Source/Runtime/Engine/Classes/Animation/AnimNode_Inertialization.h b/Engine/Source/Runtime/Engine/Classes/Animation/AnimNode_Inertialization.h
--- a/Engine/Source/Runtime/Engine/Classes/Animation/AnimNode_Inertialization.h
+++ b/Engine/Source/Runtime/Engine/Classes/Animation/AnimNode_Inertialization.h
@@ -129,6 +129,8 @@ struct FInertializationRequest
 		BlendMode = EAlphaBlendOption::Linear;
 		CustomBlendCurve = nullptr;
 		Description = FText::GetEmpty();
+		NodeId = INDEX_NONE;
+		AnimInstance = nullptr;
 	}
 
 	friend bool operator==(const FInertializationRequest& A, const FInertializationRequest& B)
@@ -139,7 +141,9 @@ struct FInertializationRequest
 			(A.bUseBlendMode == B.bUseBlendMode) &&
 			(A.BlendMode == B.BlendMode) &&
 			(A.CustomBlendCurve == B.CustomBlendCurve) &&
-			(A.Description.EqualTo(B.Description));
+			(A.Description.EqualTo(B.Description) &&
+			(A.NodeId == B.NodeId) &&
+			(A.AnimInstance == B.AnimInstance));
 	}
 
 	friend bool operator!=(const FInertializationRequest& A, const FInertializationRequest& B)
@@ -170,6 +174,14 @@ struct FInertializationRequest
 	// Description of the request - used for debugging.
 	UPROPERTY(Transient)
 	FText Description;
+
+	// Node id from which this request was made.
+	UPROPERTY(Transient)
+	int32 NodeId = INDEX_NONE;
+
+	// Anim instance from which this request was made.
+	UPROPERTY(Transient)
+	TObjectPtr<UObject> AnimInstance = nullptr;
 };
 
 
@@ -433,6 +445,13 @@ private:
 	// Description for the current inertialization request - used for debugging
 	FText InertializationRequestDescription;
 
+	// Node Id for the current inertialization request - used for debugging
+	int32 InertializationRequestNodeId = INDEX_NONE;
+
+	// Anim Instance for the current inertialization request - used for debugging
+	UPROPERTY(Transient)
+	TObjectPtr<UObject> InertializationRequestAnimInstance = nullptr;
+
 	// Inertialization durations indexed by skeleton bone index (used for per-bone blending)
 	TCustomBoneIndexArray<float, FSkeletonPoseBoneIndex> InertializationDurationPerBone;
 
diff --git a/Engine/Source/Runtime/Engine/Private/Animation/AnimNode_DeadBlending.cpp b/Engine/Source/Runtime/Engine/Private/Animation/AnimNode_DeadBlending.cpp
--- a/Engine/Source/Runtime/Engine/Private/Animation/AnimNode_DeadBlending.cpp
+++ b/Engine/Source/Runtime/Engine/Private/Animation/AnimNode_DeadBlending.cpp
@@ -718,6 +718,8 @@ void FAnimNode_DeadBlending::Evaluate_AnyThread(FPoseContext& Output)
 
 #if ANIM_TRACE_ENABLED
 		InertializationRequestDescription = RequestQueue[ShortestRequestIdx].Description;
+		InertializationRequestNodeId = RequestQueue[ShortestRequestIdx].NodeId;
+		InertializationRequestAnimInstance = RequestQueue[ShortestRequestIdx].AnimInstance;
 #endif
 
 		// Override with defaults
@@ -837,7 +839,8 @@ void FAnimNode_DeadBlending::Evaluate_AnyThread(FPoseContext& Output)
 	TRACE_ANIM_NODE_VALUE_WITH_ID(Output, GetNodeIndex(), TEXT("Max Duration"), InertializationMaxDuration);
 	TRACE_ANIM_NODE_VALUE_WITH_ID(Output, GetNodeIndex(), TEXT("Normalized Time"), InertializationDuration > UE_KINDA_SMALL_NUMBER ? (InertializationTime / InertializationDuration) : 0.0f);
 	TRACE_ANIM_NODE_VALUE_WITH_ID(Output, GetNodeIndex(), TEXT("Inertialization Weight"), InertializationWeight);
-	TRACE_ANIM_NODE_VALUE_WITH_ID(Output, GetNodeIndex(), TEXT("Request"), *InertializationRequestDescription.ToString());
+	TRACE_ANIM_NODE_VALUE_WITH_ID(Output, GetNodeIndex(), TEXT("Request Description"), *InertializationRequestDescription.ToString());
+	TRACE_ANIM_NODE_VALUE_WITH_ID_ANIM_NODE(Output, GetNodeIndex(), TEXT("Request Node"), InertializationRequestNodeId, InertializationRequestAnimInstance);
 }
 
 bool FAnimNode_DeadBlending::NeedsDynamicReset() const
diff --git a/Engine/Source/Runtime/Engine/Private/Animation/AnimNode_Inertialization.cpp b/Engine/Source/Runtime/Engine/Private/Animation/AnimNode_Inertialization.cpp
--- a/Engine/Source/Runtime/Engine/Private/Animation/AnimNode_Inertialization.cpp
+++ b/Engine/Source/Runtime/Engine/Private/Animation/AnimNode_Inertialization.cpp
@@ -286,6 +286,8 @@ void FAnimNode_Inertialization::Evaluate_AnyThread(FPoseContext& Output)
 		InertializationDuration = FMath::Max(RequestQueue[0].Duration - AppliedDeficit, 0.0f);
 #if ANIM_TRACE_ENABLED
 		InertializationRequestDescription = RequestQueue[0].Description;
+		InertializationRequestNodeId = RequestQueue[0].NodeId;
+		InertializationRequestAnimInstance = RequestQueue[0].AnimInstance;
 #endif
 		FillSkeletonBoneDurationsArray(InertializationDurationPerBone, InertializationDuration, RequestQueue[0].BlendProfile);
 
@@ -304,6 +306,8 @@ void FAnimNode_Inertialization::Evaluate_AnyThread(FPoseContext& Output)
 					InertializationDuration = RequestDuration;
 #if ANIM_TRACE_ENABLED
 					InertializationRequestDescription = Request.Description;
+					InertializationRequestNodeId = Request.NodeId;
+					InertializationRequestAnimInstance = Request.AnimInstance;
 #endif
 				}
 
@@ -456,7 +460,8 @@ void FAnimNode_Inertialization::Evaluate_AnyThread(FPoseContext& Output)
 	TRACE_ANIM_NODE_VALUE_WITH_ID(Output, GetNodeIndex(), TEXT("Max Duration"), InertializationMaxDuration);
 	TRACE_ANIM_NODE_VALUE_WITH_ID(Output, GetNodeIndex(), TEXT("Normalized Time"), InertializationDuration > UE_KINDA_SMALL_NUMBER ? (InertializationElapsedTime / InertializationDuration) : 0.0f);
 	TRACE_ANIM_NODE_VALUE_WITH_ID(Output, GetNodeIndex(), TEXT("Inertialization Weight"), InertializationWeight);
-	TRACE_ANIM_NODE_VALUE_WITH_ID(Output, GetNodeIndex(), TEXT("Request"), *InertializationRequestDescription.ToString());
+	TRACE_ANIM_NODE_VALUE_WITH_ID(Output, GetNodeIndex(), TEXT("Request Description"), *InertializationRequestDescription.ToString());
+	TRACE_ANIM_NODE_VALUE_WITH_ID_ANIM_NODE(Output, GetNodeIndex(), TEXT("Request Node"), InertializationRequestNodeId, InertializationRequestAnimInstance);
 }
 
 
diff --git a/Engine/Source/Runtime/Engine/Private/Animation/AnimNode_LinkedAnimGraph.cpp b/Engine/Source/Runtime/Engine/Private/Animation/AnimNode_LinkedAnimGraph.cpp
--- a/Engine/Source/Runtime/Engine/Private/Animation/AnimNode_LinkedAnimGraph.cpp
+++ b/Engine/Source/Runtime/Engine/Private/Animation/AnimNode_LinkedAnimGraph.cpp
@@ -123,7 +123,9 @@ void FAnimNode_LinkedAnimGraph::Update_AnyThread(const FAnimationUpdateContext&
 				Request.Duration = PendingBlendOutDuration;
 				Request.BlendProfile = PendingBlendOutProfile;
 #if ANIM_TRACE_ENABLED
-				Request.Description = NSLOCTEXT("AnimNode_LinkedAnimGraph", "InertializationRequestDescriptionOut", "Linked Anim Graph Out");
+				Request.Description = NSLOCTEXT("AnimNode_LinkedAnimGraph", "InertializationRequestDescriptionOut", "Out");
+				Request.NodeId = InContext.GetCurrentNodeId();
+				Request.AnimInstance = InContext.AnimInstanceProxy->GetAnimInstanceObject();
 #endif
 
 				InertializationRequester->RequestInertialization(Request);
@@ -135,7 +137,9 @@ void FAnimNode_LinkedAnimGraph::Update_AnyThread(const FAnimationUpdateContext&
 				Request.Duration = PendingBlendInDuration;
 				Request.BlendProfile = PendingBlendInProfile;
 #if ANIM_TRACE_ENABLED
-				Request.Description = NSLOCTEXT("AnimNode_LinkedAnimGraph", "InertializationRequestDescriptionIn", "Linked Anim Graph In");
+				Request.Description = NSLOCTEXT("AnimNode_LinkedAnimGraph", "InertializationRequestDescriptionIn", "In");
+				Request.NodeId = InContext.GetCurrentNodeId();
+				Request.AnimInstance = InContext.AnimInstanceProxy->GetAnimInstanceObject();
 #endif
 
 				InertializationRequester->RequestInertialization(Request);
diff --git a/Engine/Source/Runtime/Engine/Private/Animation/AnimNode_StateMachine.cpp b/Engine/Source/Runtime/Engine/Private/Animation/AnimNode_StateMachine.cpp
--- a/Engine/Source/Runtime/Engine/Private/Animation/AnimNode_StateMachine.cpp
+++ b/Engine/Source/Runtime/Engine/Private/Animation/AnimNode_StateMachine.cpp
@@ -1254,10 +1254,12 @@ void FAnimNode_StateMachine::TransitionToState(const FAnimationUpdateContext& Co
 					Request.CustomBlendCurve = TransitionInfo.CustomCurve;
 #if ANIM_TRACE_ENABLED
 					Request.Description = FText::Format(LOCTEXT("InertializationRequestDescription", 
-						"State Machine \"{0}\" Transition \"{1}\" to \"{2}\""), 
+						"\"{0}\" Transition from \"{1}\" to \"{2}\""), 
 						FText::FromName(GetMachineDescription()->MachineName),
 						FText::FromName(GetStateInfo(TransitionInfo.PreviousState).StateName),
 						FText::FromName(GetStateInfo(TransitionInfo.NextState).StateName));
+					Request.NodeId = Context.GetCurrentNodeId();
+					Request.AnimInstance = Context.AnimInstanceProxy->GetAnimInstanceObject();
 #endif
 
 					InertializationRequester->RequestInertialization(Request);
diff --git a/Engine/Source/Runtime/Engine/Private/Animation/AnimTrace.cpp b/Engine/Source/Runtime/Engine/Private/Animation/AnimTrace.cpp
--- a/Engine/Source/Runtime/Engine/Private/Animation/AnimTrace.cpp
+++ b/Engine/Source/Runtime/Engine/Private/Animation/AnimTrace.cpp
@@ -179,6 +179,17 @@ UE_TRACE_EVENT_BEGIN(Animation, AnimNodeValueClass)
 	UE_TRACE_EVENT_FIELD(UE::Trace::WideString, Key)
 UE_TRACE_EVENT_END()
 
+UE_TRACE_EVENT_BEGIN(Animation, AnimNodeValueAnimNode)
+	UE_TRACE_EVENT_FIELD(uint64, Cycle)
+	UE_TRACE_EVENT_FIELD(double, RecordingTime)
+	UE_TRACE_EVENT_FIELD(uint64, AnimInstanceId)
+	UE_TRACE_EVENT_FIELD(int32, Value)
+	UE_TRACE_EVENT_FIELD(uint64, ValueAnimInstanceId)
+	UE_TRACE_EVENT_FIELD(int32, NodeId)
+	UE_TRACE_EVENT_FIELD(uint16, FrameCounter)
+	UE_TRACE_EVENT_FIELD(UE::Trace::WideString, Key)
+UE_TRACE_EVENT_END()
+
 UE_TRACE_EVENT_BEGIN(Animation, AnimSequencePlayer)
 	UE_TRACE_EVENT_FIELD(uint64, Cycle)
 	UE_TRACE_EVENT_FIELD(uint64, AnimInstanceId)
@@ -1009,6 +1020,35 @@ void FAnimTrace::OutputAnimNodeValue(const FAnimationBaseContext& InContext, uin
 		<< AnimNodeValueClass.Key(InKey);
 }
 
+void FAnimTrace::OutputAnimNodeValueAnimNode(const FAnimationBaseContext& InContext, uint32 NodeIndex, const TCHAR* InKey, int32 InValue, const UObject* InValueAnimInstanceId)
+{
+	bool bChannelEnabled = UE_TRACE_CHANNELEXPR_IS_ENABLED(AnimationChannel);
+	if (!bChannelEnabled)
+	{
+		return;
+	}
+
+	check(InContext.AnimInstanceProxy);
+
+	if (CANNOT_TRACE_OBJECT(InContext.AnimInstanceProxy->GetSkelMeshComponent()))
+	{
+		return;
+	}
+
+	UObject* AnimInstance = InContext.AnimInstanceProxy->GetAnimInstanceObject();
+	TRACE_OBJECT(AnimInstance);
+
+	UE_TRACE_LOG(Animation, AnimNodeValueAnimNode, AnimationChannel)
+		<< AnimNodeValueAnimNode.Cycle(FPlatformTime::Cycles64())
+		<< AnimNodeValueAnimNode.RecordingTime(FObjectTrace::GetWorldElapsedTime(AnimInstance->GetWorld()))
+		<< AnimNodeValueAnimNode.AnimInstanceId(FObjectTrace::GetObjectId(AnimInstance))
+		<< AnimNodeValueAnimNode.NodeId(NodeIndex)
+		<< AnimNodeValueAnimNode.Value(InValue)
+		<< AnimNodeValueAnimNode.ValueAnimInstanceId(FObjectTrace::GetObjectId(InValueAnimInstanceId))
+		<< AnimNodeValueAnimNode.FrameCounter(FObjectTrace::GetObjectWorldTickCounter(AnimInstance))
+		<< AnimNodeValueAnimNode.Key(InKey);
+}
+
 void FAnimTrace::OutputAnimSequencePlayer(const FAnimationBaseContext& InContext, const FAnimNode_SequencePlayerBase& InNode)
 {
 	bool bChannelEnabled = UE_TRACE_CHANNELEXPR_IS_ENABLED(AnimationChannel);
diff --git a/Engine/Source/Runtime/Engine/Public/Animation/AnimTrace.h b/Engine/Source/Runtime/Engine/Public/Animation/AnimTrace.h
--- a/Engine/Source/Runtime/Engine/Public/Animation/AnimTrace.h
+++ b/Engine/Source/Runtime/Engine/Public/Animation/AnimTrace.h
@@ -166,6 +166,7 @@ struct FAnimTrace
 	ENGINE_API FORCENOINLINE static void OutputAnimNodeValue(const FAnimationBaseContext& InContext, uint32 NodeIndex, const TCHAR* InKey, const TCHAR* InValue);
 	ENGINE_API FORCENOINLINE static void OutputAnimNodeValue(const FAnimationBaseContext& InContext, uint32 NodeIndex, const TCHAR* InKey, const UClass* InValue);
 	ENGINE_API FORCENOINLINE static void OutputAnimNodeValue(const FAnimationBaseContext& InContext, uint32 NodeIndex, const TCHAR* InKey, const UObject* InValue);
+	ENGINE_API FORCENOINLINE static void OutputAnimNodeValueAnimNode(const FAnimationBaseContext& InContext, uint32 NodeIndex, const TCHAR* InKey, int32 InValue, const UObject* InValueAnimInstanceId);
 
 	/** Helper function to output debug info for sequence player nodes */
 	ENGINE_API static void OutputAnimSequencePlayer(const FAnimationBaseContext& InContext, const FAnimNode_SequencePlayerBase& InNode);
@@ -228,6 +229,9 @@ struct FAnimTrace
 #define TRACE_ANIM_NODE_VALUE_WITH_ID(Context, NodeId, Key, Value) \
 	FAnimTrace::OutputAnimNodeValue(Context, NodeId, Key, Value);
 
+#define TRACE_ANIM_NODE_VALUE_WITH_ID_ANIM_NODE(Context, NodeId, Key, Value, ValueAnimInstanceId) \
+	FAnimTrace::OutputAnimNodeValueAnimNode(Context, NodeId, Key, Value, ValueAnimInstanceId);
+
 #define TRACE_ANIM_SEQUENCE_PLAYER(Context, Node) \
 	FAnimTrace::OutputAnimSequencePlayer(Context, Node);
 
@@ -262,6 +266,7 @@ struct FAnimTrace
 #define TRACE_SCOPED_ANIM_NODE_SUSPEND
 #define TRACE_ANIM_NODE_VALUE(Context, Key, Value)
 #define TRACE_ANIM_NODE_VALUE_WITH_ID(Context, NodeId, Key, Value)
+#define TRACE_ANIM_NODE_VALUE_WITH_ID_ANIM_NODE(Context, NodeId, Key, Value, ValueAnimInstanceId)
 #define TRACE_ANIM_SEQUENCE_PLAYER(Context, Node)
 #define TRACE_ANIM_STATE_MACHINE_STATE(Context, StateMachineIndex, StateIndex, StateWeight, ElapsedTime)
 #define TRACE_ANIM_NOTIFY(AnimInstance, NotifyEvent, EventType)
