gameplay insights - rewind debugger tracks priority sorting fix

diff --git a/Engine/Plugins/Animation/GameplayInsights/Source/RewindDebugger/Private/RewindDebuggerObjectTrack.cpp b/Engine/Plugins/Animation/GameplayInsights/Source/RewindDebugger/Private/RewindDebuggerObjectTrack.cpp
--- a/Engine/Plugins/Animation/GameplayInsights/Source/RewindDebugger/Private/RewindDebuggerObjectTrack.cpp
+++ b/Engine/Plugins/Animation/GameplayInsights/Source/RewindDebugger/Private/RewindDebuggerObjectTrack.cpp
@@ -283,8 +283,15 @@ bool FRewindDebuggerObjectTrack::UpdateInternal()
 	{
 		Children.Sort([](const TSharedPtr<FRewindDebuggerTrack>& A, const TSharedPtr<FRewindDebuggerTrack>& B)
 		{
-			return A->GetSortOrderPriority() > B->GetSortOrderPriority() ||
-					A->GetDisplayName().ToString() < B->GetDisplayName().ToString();
+			const int SortOrderPriorityA = A->GetSortOrderPriority();
+			const int SortOrderPriorityB = B->GetSortOrderPriority();
+			
+			if (SortOrderPriorityA != SortOrderPriorityB)
+			{
+				return SortOrderPriorityA > SortOrderPriorityB;
+			}
+			
+			return A->GetDisplayName().ToString() < B->GetDisplayName().ToString();
 		});
 	}
 
