Fixed conflicting LOCTEXT key

diff --git a/Engine/Plugins/Animation/GameplayInsights/Source/GameplayInsights/Private/SVariantValueView.cpp b/Engine/Plugins/Animation/GameplayInsights/Source/GameplayInsights/Private/SVariantValueView.cpp
--- a/Engine/Plugins/Animation/GameplayInsights/Source/GameplayInsights/Private/SVariantValueView.cpp
+++ b/Engine/Plugins/Animation/GameplayInsights/Source/GameplayInsights/Private/SVariantValueView.cpp
@@ -273,7 +273,7 @@ TSharedRef<SWidget> SVariantValueView::MakeVariantValueWidget(const TraceService
 					SNew(SHyperlink)
 					.Text(FText::FromString(AnimNodeInfo->Name))
 					.TextStyle(&FCoreStyle::Get().GetWidgetStyle<FTextBlockStyle>("SmallText"))
-					.ToolTipText(FText::Format(LOCTEXT("AssetHyperlinkTooltipFormat", "Open node '{0}'"), FText::FromString(AnimNodeInfo->Name)))
+					.ToolTipText(FText::Format(LOCTEXT("NodeHyperlinkTooltipFormat", "Open node '{0}'"), FText::FromString(AnimNodeInfo->Name)))
 					.OnNavigate_Lambda([ObjectInfo, AnimNodeInfo, AnimInstanceClassInfo, InValue]()
 					{
 						TSoftObjectPtr<UAnimBlueprintGeneratedClass> InstanceClass;
