Rewind Debugger: Force eject and repossess when using the actor selection tool.
#rnx

diff --git a/Engine/Plugins/Animation/GameplayInsights/Source/RewindDebugger/Private/SRewindDebugger.cpp b/Engine/Plugins/Animation/GameplayInsights/Source/RewindDebugger/Private/SRewindDebugger.cpp
--- a/Engine/Plugins/Animation/GameplayInsights/Source/RewindDebugger/Private/SRewindDebugger.cpp
+++ b/Engine/Plugins/Animation/GameplayInsights/Source/RewindDebugger/Private/SRewindDebugger.cpp
@@ -405,12 +405,7 @@ void SRewindDebugger::Construct(const FArguments& InArgs, TSharedRef<FUICommandL
 						SNew(SButton)
 							.ButtonStyle(FAppStyle::Get(), "SimpleButton")
 							.OnClicked(this, &SRewindDebugger::OnSelectActorClicked)
-							.ToolTipText(LOCTEXT("SelectActorTooltip", "Select Target Actor in Scene (Eject player control first)"))
-							.IsEnabled_Lambda([]()
-								{
-									return !FPlayWorldCommandCallbacks::IsInPIE();
-								}
-							)
+							.ToolTipText(LOCTEXT("SelectActorTooltip", "Select Target Actor in Scene"))
 						[
 							SNew(SImage)
 							.Image(FRewindDebuggerStyle::Get().GetBrush("RewindDebugger.SelectActor"))
@@ -553,12 +548,25 @@ FReply SRewindDebugger::OnSelectActorClicked()
 {
 	FActorPickerModeModule& ActorPickerMode = FModuleManager::Get().GetModuleChecked<FActorPickerModeModule>("ActorPickerMode");
 	
-	// todo: force eject (from within BeginActorPickingMode?)
+	const bool bShouldForceEject = GEditor->PlayWorld && !GEditor->bIsSimulatingInEditor;
+	if (bShouldForceEject)
+	{
+		// Eject PIE
+		GEditor->RequestToggleBetweenPIEandSIE();
+	}
 
 	ActorPickerMode.BeginActorPickingMode(
 		FOnGetAllowedClasses(), 
 		FOnShouldFilterActor(),
-		FOnActorSelected::CreateRaw(this, &SRewindDebugger::SetDebugTargetActor));
+		FOnActorSelected::CreateLambda([this, bShouldForceEject](AActor* InActor)
+		{
+			SetDebugTargetActor(InActor);
+			if (bShouldForceEject && GEditor->bIsSimulatingInEditor)
+			{
+				// If we force ejected PIE, revert this after actor selection.
+				GEditor->RequestToggleBetweenPIEandSIE();
+			}
+		}));
 
 
 
