remove code which disables tracing to all channels when animation insights activates

diff --git a/Engine/Plugins/Animation/GameplayInsights/Source/GameplayInsights/Private/GameplayInsightsModule.cpp b/Engine/Plugins/Animation/GameplayInsights/Source/GameplayInsights/Private/GameplayInsightsModule.cpp
--- a/Engine/Plugins/Animation/GameplayInsights/Source/GameplayInsights/Private/GameplayInsightsModule.cpp
+++ b/Engine/Plugins/Animation/GameplayInsights/Source/GameplayInsights/Private/GameplayInsightsModule.cpp
@@ -432,23 +432,6 @@ void FGameplayInsightsModule::StartTrace()
 	
 	if (!bTraceStarted)
 	{
-		if (!FTraceAuxiliary::IsConnected())
-		{
-			// cpu tracing is enabled by default, but not connected.
-			// when not connected, disable all channels initially to avoid a whole bunch of extra cpu, memory, and disk overhead from processing all the extra default trace channels
-			////////////////////////////////////////////////////////////////////////////////
-			UE::Trace::EnumerateChannels([](const ANSICHAR* ChannelName, bool bEnabled, void*)
-				{
-					if (bEnabled)
-					{
-						FString ChannelNameFString(ChannelName);
-						UE::Trace::ToggleChannel(ChannelNameFString.GetCharArray().GetData(), false);
-					}
-				}
-				, nullptr);
-		}
-
-
 		bTraceStarted = FTraceAuxiliary::Start(
 			FTraceAuxiliary::EConnectionType::Network,
 			TEXT("127.0.0.1"),
