Enable visual logger rewind debugger integration

diff --git a/Engine/Plugins/Animation/GameplayInsights/Source/GameplayInsights/Private/SEventTimelineView.h b/Engine/Plugins/Animation/GameplayInsights/Source/GameplayInsights/Private/SEventTimelineView.h
deleted file mode 100644
diff --git a/Engine/Plugins/Animation/GameplayInsights/Source/GameplayInsights/Public/SEventTimelineView.h b/Engine/Plugins/Animation/GameplayInsights/Source/GameplayInsights/Public/SEventTimelineView.h
new file mode 100644
--- /dev/null
+++ b/Engine/Plugins/Animation/GameplayInsights/Source/GameplayInsights/Public/SEventTimelineView.h
@@ -0,0 +1,77 @@
+// Copyright Epic Games, Inc. All Rights Reserved.
+
+#pragma once
+
+#include "CoreMinimal.h"
+#include "Widgets/SCompoundWidget.h"
+#include "Widgets/DeclarativeSyntaxSupport.h"
+
+class FPaintArgs;
+class FSlateWindowElementList;
+class FTimeSliderController;
+
+class SEventTimelineView : public SCompoundWidget
+{
+public:
+	struct FTimelineEventData
+	{
+	public:
+		struct EventPoint
+		{
+			double Time;
+			FText Type;
+			FText Description;
+			FLinearColor Color;
+		};
+		
+		struct EventWindow
+		{
+			double TimeStart;
+			double TimeEnd;
+			FText Type;
+			FText Description;
+			FLinearColor Color;
+		};
+
+		TArray<EventPoint> Points;
+		TArray<EventWindow> Windows;
+	};
+	
+	SLATE_BEGIN_ARGS(SEventTimelineView)
+		: _ViewRange(TRange<double>(0,10))
+		, _DesiredSize(FVector2D(100.f,20.f))
+	{}
+    	/** View time range */
+    	SLATE_ATTRIBUTE(TRange<double>, ViewRange);
+
+		/** Data for events to render */
+		SLATE_ATTRIBUTE(TSharedPtr<FTimelineEventData>, EventData);
+
+		/** Desired widget size */
+		SLATE_ATTRIBUTE(FVector2D, DesiredSize);
+	
+	SLATE_END_ARGS()
+
+
+	/**
+	 * Construct the widget
+	 * 
+	 * @param InArgs   A declaration from which to construct the widget
+	 */
+	void Construct( const FArguments& InArgs );
+
+protected:
+	// SWidget interface
+	virtual FVector2D ComputeDesiredSize(float) const override;
+	virtual int32 OnPaint( const FPaintArgs& Args, const FGeometry& AllottedGeometry, const FSlateRect& MyCullingRect, FSlateWindowElementList& OutDrawElements, int32 LayerId, const FWidgetStyle& InWidgetStyle, bool bParentEnabled ) const override;
+	
+	int32 PaintEvents(const FGeometry& AllottedGeometry, const FSlateRect& MyCullingRect, FSlateWindowElementList& OutDrawElements, int32 LayerId, const FWidgetStyle& InWidgetStyle, bool bParentEnabled ) const;
+	
+	TAttribute<TRange<double>> ViewRange;
+	TAttribute<TSharedPtr<FTimelineEventData>> EventData;
+	TAttribute<FVector2D> DesiredSize;
+	
+	const FSlateBrush* EventBrush;
+	const FSlateBrush* EventBorderBrush;
+};
+
diff --git a/Engine/Plugins/Animation/GameplayInsights/Source/RewindDebugger/Private/RewindDebugger.cpp b/Engine/Plugins/Animation/GameplayInsights/Source/RewindDebugger/Private/RewindDebugger.cpp
--- a/Engine/Plugins/Animation/GameplayInsights/Source/RewindDebugger/Private/RewindDebugger.cpp
+++ b/Engine/Plugins/Animation/GameplayInsights/Source/RewindDebugger/Private/RewindDebugger.cpp
@@ -22,6 +22,7 @@
 #include "EngineUtils.h"
 #include "ToolMenus.h"
 #include "RewindDebuggerSettings.h"
+#include "RewindDebuggerCommands.h"
 #include "LevelEditor.h"
 #include "RewindDebuggerModule.h"
 #include "Engine/PoseWatch.h"
@@ -29,6 +30,8 @@
 #include "UObject/UObjectIterator.h"
 #include "Engine/World.h"
 
+#define LOCTEXT_NAMESPACE "RewindDebugger"
+
 static void IterateExtensions(TFunction<void(IRewindDebuggerExtension* Extension)> IteratorFunction)
 {
 	// update extensions
@@ -1247,6 +1250,72 @@ void FRewindDebugger::RegisterComponentContextMenu()
 	}));
 }
 
+void FRewindDebugger::RegisterToolBar()
+{
+	UToolMenu* Menu = UToolMenus::Get()->RegisterMenu("RewindDebugger.ToolBar", NAME_None, EMultiBoxType::ToolBar);
+	
+	FToolMenuSection& Section = Menu->FindOrAddSection("VCRControls");
+	const FRewindDebuggerCommands& Commands = FRewindDebuggerCommands::Get();
+	Section.AddEntry(FToolMenuEntry::InitToolBarButton(
+			Commands.FirstFrame,
+			LOCTEXT("Blank",""),
+			TAttribute<FText>(),
+			FSlateIcon("RewindDebuggerStyle", "RewindDebugger.FirstFrame.small")));
+	
+	Section.AddEntry(FToolMenuEntry::InitToolBarButton(
+			Commands.PreviousFrame,
+			LOCTEXT("Blank",""),
+			TAttribute<FText>(),
+			FSlateIcon("RewindDebuggerStyle", "RewindDebugger.PreviousFrame.small")));
+			
+	Section.AddEntry(FToolMenuEntry::InitToolBarButton(
+				Commands.ReversePlay,
+				LOCTEXT("Blank",""),
+				TAttribute<FText>(),
+				FSlateIcon("RewindDebuggerStyle", "RewindDebugger.ReversePlay.small")));
+	
+	Section.AddEntry(FToolMenuEntry::InitToolBarButton(
+				Commands.Pause,
+				LOCTEXT("Blank",""),
+				FText::Format(LOCTEXT("PauseButtonTooltip", "{0} ({1})"), Commands.Pause->GetDescription(), Commands.PauseOrPlay->GetInputText()),
+				FSlateIcon("RewindDebuggerStyle", "RewindDebugger.Pause.small")));
+	
+	Section.AddEntry(FToolMenuEntry::InitToolBarButton(
+				Commands.Play,
+				LOCTEXT("Blank",""),
+				FText::Format(LOCTEXT("PlayButtonTooltip", "{0} ({1})"), Commands.Play->GetDescription(), Commands.PauseOrPlay->GetInputText()),
+				FSlateIcon("RewindDebuggerStyle", "RewindDebugger.Play.small")));
+
+	Section.AddEntry(FToolMenuEntry::InitToolBarButton(
+				Commands.NextFrame,
+				LOCTEXT("Blank",""),
+				TAttribute<FText>(),
+				FSlateIcon("RewindDebuggerStyle", "RewindDebugger.NextFrame.small")));
+	
+	Section.AddEntry(FToolMenuEntry::InitToolBarButton(
+				Commands.LastFrame,
+				LOCTEXT("Blank",""),
+				TAttribute<FText>(),
+				FSlateIcon("RewindDebuggerStyle", "RewindDebugger.LastFrame.small")));
+
+	Section.AddEntry(FToolMenuEntry::InitToolBarButton(
+				Commands.StartRecording,
+				LOCTEXT("Blank",""),
+				TAttribute<FText>(),
+				FSlateIcon("RewindDebuggerStyle", "RewindDebugger.StartRecording.small")));
+				
+
+	Section.AddEntry(FToolMenuEntry::InitToolBarButton(
+				Commands.StopRecording,
+				LOCTEXT("Blank",""),
+				TAttribute<FText>(),
+				FSlateIcon("RewindDebuggerStyle", "RewindDebugger.StopRecording.small")));
+	
+	Menu->SetStyleSet(&FAppStyle::Get());
+	Menu->StyleName = "PaletteToolBar";
+}
+
+
 void FRewindDebugger::ComponentDoubleClicked(TSharedPtr<RewindDebugger::FRewindDebuggerTrack> SelectedObject)
 {
 	if (!SelectedObject.IsValid())
@@ -1330,3 +1399,5 @@ TArray<TSharedPtr<FDebugObjectInfo>>& FRewindDebugger::GetDebugComponents()
 	RefreshDebugComponents(DebugTracks, DebugComponents);
 	return DebugComponents;
 }
+
+#undef LOCTEXT_NAMESPACE
diff --git a/Engine/Plugins/Animation/GameplayInsights/Source/RewindDebugger/Private/RewindDebugger.h b/Engine/Plugins/Animation/GameplayInsights/Source/RewindDebugger/Private/RewindDebugger.h
--- a/Engine/Plugins/Animation/GameplayInsights/Source/RewindDebugger/Private/RewindDebugger.h
+++ b/Engine/Plugins/Animation/GameplayInsights/Source/RewindDebugger/Private/RewindDebugger.h
@@ -112,6 +112,7 @@ public:
 	
 	void UpdateDetailsPanel(TSharedRef<SDockTab> DetailsTab);
 	static void RegisterComponentContextMenu();
+	static void RegisterToolBar();
 	
 	DECLARE_DELEGATE_OneParam( FOnTrackCursor, bool)
 	void OnTrackCursor(const FOnTrackCursor& TrackCursorCallback);
diff --git a/Engine/Plugins/Animation/GameplayInsights/Source/RewindDebugger/Private/RewindDebuggerModule.cpp b/Engine/Plugins/Animation/GameplayInsights/Source/RewindDebugger/Private/RewindDebuggerModule.cpp
--- a/Engine/Plugins/Animation/GameplayInsights/Source/RewindDebugger/Private/RewindDebuggerModule.cpp
+++ b/Engine/Plugins/Animation/GameplayInsights/Source/RewindDebugger/Private/RewindDebuggerModule.cpp
@@ -231,6 +231,7 @@ void FRewindDebuggerModule::StartupModule()
 	FPropertyTraceMenu::Register();
 	FAnimInstanceMenu::Register();
 	FRewindDebugger::RegisterComponentContextMenu();
+	FRewindDebugger::RegisterToolBar();
 }
 
 void FRewindDebuggerModule::ShutdownModule()
diff --git a/Engine/Plugins/Animation/GameplayInsights/Source/RewindDebugger/Private/SRewindDebugger.cpp b/Engine/Plugins/Animation/GameplayInsights/Source/RewindDebugger/Private/SRewindDebugger.cpp
--- a/Engine/Plugins/Animation/GameplayInsights/Source/RewindDebugger/Private/SRewindDebugger.cpp
+++ b/Engine/Plugins/Animation/GameplayInsights/Source/RewindDebugger/Private/SRewindDebugger.cpp
@@ -240,24 +240,8 @@ void SRewindDebugger::Construct(const FArguments& InArgs, TSharedRef<FUICommandL
 		RefreshDebugComponents();
 	});
 	
-	FSlimHorizontalToolBarBuilder ToolBarBuilder(CommandList, FMultiBoxCustomization::None, nullptr, true);
-	
-	ToolBarBuilder.SetStyle(&FAppStyle::Get(), "PaletteToolBar");
-	ToolBarBuilder.BeginSection("Debugger");
-	{
-		ToolBarBuilder.AddToolBarButton(Commands.FirstFrame);
-		ToolBarBuilder.AddToolBarButton(Commands.PreviousFrame);
-		ToolBarBuilder.AddToolBarButton(Commands.ReversePlay);
-		ToolBarBuilder.AddToolBarButton(Commands.Pause, NAME_None, {}, FText::Format(LOCTEXT("PauseButtonTooltip", "{0} ({1})"), Commands.Pause->GetDescription(), Commands.PauseOrPlay->GetInputText()));
-		ToolBarBuilder.AddToolBarButton(Commands.Play, NAME_None, {}, FText::Format(LOCTEXT("PlayButtonTooltip", "{0} ({1}) or"), Commands.Play->GetDescription(), Commands.PauseOrPlay->GetInputText(), Commands.Play->GetInputText()));
-		ToolBarBuilder.AddToolBarButton(Commands.NextFrame);
-		ToolBarBuilder.AddToolBarButton(Commands.LastFrame);
-		ToolBarBuilder.AddToolBarButton(Commands.StartRecording);
-		ToolBarBuilder.AddToolBarButton(Commands.StopRecording);
-	}
-	ToolBarBuilder.EndSection();
-
-	TSharedPtr<SScrollBar> ScrollBar = SNew(SScrollBar); 
+	TSharedPtr<SScrollBar> ScrollBar = SNew(SScrollBar);
+	FToolMenuContext ToolMenuContext(CommandList);
 
 	ComponentTreeView =	SNew(SRewindDebuggerComponentTree)
 		.ExternalScrollBar(ScrollBar)
@@ -311,7 +295,8 @@ void SRewindDebugger::Construct(const FArguments& InArgs, TSharedRef<FUICommandL
 				}
 			});
 
-	UToolMenu* Menu = UToolMenus::Get()->FindMenu("RewindDebugger.MainMenu");
+	
+	TSharedRef<SWidget> ToolBar = UToolMenus::Get()->GenerateWidget("RewindDebugger.ToolBar", ToolMenuContext);
 
 	ChildSlot
 	[
@@ -319,7 +304,7 @@ void SRewindDebugger::Construct(const FArguments& InArgs, TSharedRef<FUICommandL
 		 + SVerticalBox::Slot().AutoHeight()
 		[
 			SNew(SVerticalBox)
-			+SVerticalBox::Slot().AutoHeight()
+			+SVerticalBox::Slot().MaxHeight(48)
 			[
 				SNew(SHorizontalBox)
 				+SHorizontalBox::Slot().AutoWidth()
@@ -335,7 +320,7 @@ void SRewindDebugger::Construct(const FArguments& InArgs, TSharedRef<FUICommandL
 				]
 				+SHorizontalBox::Slot().FillWidth(1.0)
 				[
-					ToolBarBuilder.MakeWidget()
+					ToolBar
 				]
 			]
 		]
diff --git a/Engine/Plugins/Animation/GameplayInsights/Source/RewindDebuggerVLog/Private/RewindDebuggerVLog.cpp b/Engine/Plugins/Animation/GameplayInsights/Source/RewindDebuggerVLog/Private/RewindDebuggerVLog.cpp
--- a/Engine/Plugins/Animation/GameplayInsights/Source/RewindDebuggerVLog/Private/RewindDebuggerVLog.cpp
+++ b/Engine/Plugins/Animation/GameplayInsights/Source/RewindDebuggerVLog/Private/RewindDebuggerVLog.cpp
@@ -1,11 +1,15 @@
 // Copyright Epic Games, Inc. All Rights Reserved.
 
 #include "RewindDebuggerVLog.h"
+#include "Editor/EditorEngine.h"
 #include "Engine/World.h"
+#include "IRewindDebugger.h"
 #include "IVisualLoggerProvider.h"
 #include "Insights/IUnrealInsightsModule.h"
-#include "IRewindDebugger.h"
+#include "LogVisualizerSettings.h"
 #include "Modules/ModuleManager.h"
+#include "RewindDebuggerVLogSettings.h"
+#include "ToolMenus.h"
 #include "TraceServices/Model/Frames.h"
 #include "VisualLogger/VisualLogger.h"
 
@@ -16,6 +20,103 @@ FRewindDebuggerVLog::FRewindDebuggerVLog()
 
 }
 
+void FRewindDebuggerVLog::Initialize()
+{
+	UToolMenu* Menu = UToolMenus::Get()->FindMenu("RewindDebugger.ToolBar");
+
+	FToolMenuSection& NewSection = Menu->AddSection("Visual Logger", LOCTEXT("Visual Logger","Visual Logger"));
+
+	NewSection.AddSeparator("VisualLogger");
+
+	NewSection.AddEntry(FToolMenuEntry::InitComboButton(
+		"VLog Categories",
+		FUIAction(),
+		FNewToolMenuDelegate::CreateRaw(this, &FRewindDebuggerVLog::MakeCategoriesMenu),
+		LOCTEXT("VLog Categories", "VLog Categories")
+		));
+	
+	NewSection.AddEntry(FToolMenuEntry::InitComboButton(
+		"VLog Level",
+		FUIAction(),
+		FNewToolMenuDelegate::CreateRaw(this, &FRewindDebuggerVLog::MakeLogLevelMenu),
+		LOCTEXT("VLog Level", "VLog Level")));
+}
+
+bool FRewindDebuggerVLog::IsCategoryActive(const FName& Category)
+{
+	URewindDebuggerVLogSettings& Settings = URewindDebuggerVLogSettings::Get();
+	return Settings.DisplayCategories.Contains(Category);
+}
+
+void FRewindDebuggerVLog::ToggleCategory(const FName& Category)
+{
+	URewindDebuggerVLogSettings& Settings = URewindDebuggerVLogSettings::Get();
+	if (Settings.DisplayCategories.Remove(Category) == 0)
+	{
+		Settings.DisplayCategories.Add(Category);
+	}
+}
+
+ELogVerbosity::Type FRewindDebuggerVLog::GetMinLogVerbosity() const
+{
+	return static_cast<ELogVerbosity::Type>(URewindDebuggerVLogSettings::Get().DisplayVerbosity);
+}
+
+void FRewindDebuggerVLog::SetMinLogVerbosity(ELogVerbosity::Type Value)
+{
+	URewindDebuggerVLogSettings::Get().DisplayVerbosity = Value;
+}
+
+void FRewindDebuggerVLog::MakeLogLevelMenu(UToolMenu* Menu)
+{
+	FToolMenuSection& Section = Menu->AddSection("Levels");
+
+	for(ELogVerbosity::Type LogVerbosityLevel = ELogVerbosity::All; LogVerbosityLevel > 0; LogVerbosityLevel = static_cast<ELogVerbosity::Type>(LogVerbosityLevel-1))
+	{
+		FString Name = ToString(LogVerbosityLevel);
+		FText Label(FText::FromString(Name));
+		
+		Section.AddMenuEntry(FName(Name),
+			Label,
+			FText(),
+			FSlateIcon(),
+			FUIAction(
+				FExecuteAction::CreateLambda([LogVerbosityLevel, this]() { SetMinLogVerbosity(LogVerbosityLevel); }),
+				FCanExecuteAction(),
+				FGetActionCheckState::CreateLambda( [LogVerbosityLevel, this]() { return GetMinLogVerbosity() == LogVerbosityLevel ? ECheckBoxState::Checked : ECheckBoxState::Unchecked; } )
+			),
+			EUserInterfaceActionType::Check
+			);
+	}
+}
+
+void FRewindDebuggerVLog::MakeCategoriesMenu(UToolMenu* Menu)
+{
+	FToolMenuSection& Section = Menu->AddSection("Categories");
+	IRewindDebugger* RewindDebugger = IRewindDebugger::Instance();
+	if (const TraceServices::IAnalysisSession* Session = RewindDebugger->GetAnalysisSession())
+	{
+		TraceServices::FAnalysisSessionReadScope SessionReadScope(*Session);
+		if (const IVisualLoggerProvider* VisualLoggerProvider = Session->ReadProvider<IVisualLoggerProvider>("VisualLoggerProvider"))
+		{
+			VisualLoggerProvider->EnumerateCategories([&Section, this](const FName& Category)
+			{
+				Section.AddMenuEntry(Category,
+					FText::FromName(Category),
+					FText(),
+					FSlateIcon(),
+					FUIAction(
+						FExecuteAction::CreateLambda([Category, this]() { ToggleCategory(Category); }),
+						FCanExecuteAction(),
+						FGetActionCheckState::CreateLambda( [Category, this]() { return IsCategoryActive(Category) ? ECheckBoxState::Checked : ECheckBoxState::Unchecked; } )
+					),
+					EUserInterfaceActionType::Check
+					);
+			});
+		}
+	}
+}
+
 void FRewindDebuggerVLog::AddLogEntries(const TArray<TSharedPtr<FDebugObjectInfo>>& Components, float StartTime, float EndTime, const IVisualLoggerProvider* VisualLoggerProvider)
 {
 	for(const TSharedPtr<FDebugObjectInfo>& ComponentInfo : Components)
@@ -60,14 +161,20 @@ void FRewindDebuggerVLog::Update(float DeltaTime, IRewindDebugger* RewindDebugge
 				{
 					if (!VLogActor.IsValid())
 					{
-						FActorSpawnParameters SpawnParameters;
-						SpawnParameters.ObjectFlags |= RF_Transient;
-						VLogActor = RewindDebugger->GetWorldToVisualize()->SpawnActor<AVLogRenderingActor>(SpawnParameters);
+						UEditorEngine* EditorEngine = Cast<UEditorEngine>(GEngine);
+						if (GIsEditor && EditorEngine && EditorEngine->PlayWorld)
+						{
+							FActorSpawnParameters SpawnParameters;
+							SpawnParameters.ObjectFlags |= RF_Transient;
+							VLogActor = EditorEngine->PlayWorld->SpawnActor<AVLogRenderingActor>(SpawnParameters);
+						}
 					}
 
-					VLogActor->Reset();
-
-					AddLogEntries(RewindDebugger->GetDebugComponents(), CurrentFrame.StartTime, CurrentFrame.EndTime, VisualLoggerProvider);
+					if (VLogActor.IsValid())
+					{
+						VLogActor->Reset();
+						AddLogEntries(RewindDebugger->GetDebugComponents(), CurrentFrame.StartTime, CurrentFrame.EndTime, VisualLoggerProvider);
+					}
 				}
 			}
 		}
diff --git a/Engine/Plugins/Animation/GameplayInsights/Source/RewindDebuggerVLog/Private/RewindDebuggerVLog.h b/Engine/Plugins/Animation/GameplayInsights/Source/RewindDebuggerVLog/Private/RewindDebuggerVLog.h
--- a/Engine/Plugins/Animation/GameplayInsights/Source/RewindDebuggerVLog/Private/RewindDebuggerVLog.h
+++ b/Engine/Plugins/Animation/GameplayInsights/Source/RewindDebuggerVLog/Private/RewindDebuggerVLog.h
@@ -6,6 +6,7 @@
 #include "IRewindDebuggerExtension.h"
 #include "UObject/WeakObjectPtr.h"
 #include "VLogRenderingActor.h"
+#include "ToolMenu.h"
 
 // Rewind debugger extension for Visual Logger support
 
@@ -15,6 +16,15 @@ public:
 	FRewindDebuggerVLog();
 	virtual ~FRewindDebuggerVLog() {};
 
+	void Initialize();
+	void MakeCategoriesMenu(UToolMenu* Menu);
+	void MakeLogLevelMenu(UToolMenu* Menu);
+	void ToggleCategory(const FName& Category);
+	bool IsCategoryActive(const FName& Category);
+
+	ELogVerbosity::Type GetMinLogVerbosity() const;
+	void SetMinLogVerbosity(ELogVerbosity::Type Value);
+	
 	virtual void Update(float DeltaTime, IRewindDebugger* RewindDebugger) override;
 
 	virtual void RecordingStarted(IRewindDebugger* RewindDebugger) override;
diff --git a/Engine/Plugins/Animation/GameplayInsights/Source/RewindDebuggerVLog/Private/RewindDebuggerVLogModule.cpp b/Engine/Plugins/Animation/GameplayInsights/Source/RewindDebuggerVLog/Private/RewindDebuggerVLogModule.cpp
--- a/Engine/Plugins/Animation/GameplayInsights/Source/RewindDebuggerVLog/Private/RewindDebuggerVLogModule.cpp
+++ b/Engine/Plugins/Animation/GameplayInsights/Source/RewindDebuggerVLog/Private/RewindDebuggerVLogModule.cpp
@@ -7,7 +7,7 @@
 #define LOCTEXT_NAMESPACE "RewindDebuggerVLogModule"
 
 #ifndef ENABLE_REWINDDEBUGGER_VLOG_INTEGRATION
-#define ENABLE_REWINDDEBUGGER_VLOG_INTEGRATION 0
+#define ENABLE_REWINDDEBUGGER_VLOG_INTEGRATION 1
 #endif
 
 void FRewindDebuggerVLogModule::StartupModule()
@@ -15,6 +15,7 @@ void FRewindDebuggerVLogModule::StartupModule()
 #if ENABLE_REWINDDEBUGGER_VLOG_INTEGRATION
 	IModularFeatures::Get().RegisterModularFeature(IRewindDebuggerExtension::ModularFeatureName, &RewindDebuggerVLogExtension);
 	IModularFeatures::Get().RegisterModularFeature(TraceServices::ModuleFeatureName, &VLogTraceModule);
+	RewindDebuggerVLogExtension.Initialize();
 #endif
 }
 
diff --git a/Engine/Plugins/Animation/GameplayInsights/Source/RewindDebuggerVLog/Private/RewindDebuggerVLogSettings.cpp b/Engine/Plugins/Animation/GameplayInsights/Source/RewindDebuggerVLog/Private/RewindDebuggerVLogSettings.cpp
new file mode 100644
--- /dev/null
+++ b/Engine/Plugins/Animation/GameplayInsights/Source/RewindDebuggerVLog/Private/RewindDebuggerVLogSettings.cpp
@@ -0,0 +1,42 @@
+﻿// Copyright Epic Games, Inc. All Rights Reserved.
+
+#include "RewindDebuggerVLogSettings.h"
+#include "Misc/CoreDelegates.h"
+
+#include UE_INLINE_GENERATED_CPP_BY_NAME(RewindDebuggerVLogSettings)
+
+URewindDebuggerVLogSettings::URewindDebuggerVLogSettings()
+{
+	FCoreDelegates::OnPreExit.AddLambda([]()
+	{
+		Get().SaveConfig();
+	});
+}
+
+#if WITH_EDITOR
+
+FText URewindDebuggerVLogSettings::GetSectionText() const
+{
+	return NSLOCTEXT("RewindDebugger", "RewindDebuggerVLogSettingsName", "Rewind Debugger");
+}
+
+FText URewindDebuggerVLogSettings::GetSectionDescription() const
+{
+	return NSLOCTEXT("RewindDebugger", "RewindDebuggerVLogSettingsDescription", "Configure options for the Rewind Debugger.");
+}
+
+#endif
+
+FName URewindDebuggerVLogSettings::GetCategoryName() const
+{
+	return FName(TEXT("Plugins"));
+}
+
+URewindDebuggerVLogSettings& URewindDebuggerVLogSettings::Get()
+{
+	URewindDebuggerVLogSettings* MutableCDO = GetMutableDefault<URewindDebuggerVLogSettings>();
+	check(MutableCDO != nullptr)
+	
+	return *MutableCDO;
+}
+
diff --git a/Engine/Plugins/Animation/GameplayInsights/Source/RewindDebuggerVLog/Private/RewindDebuggerVLogSettings.h b/Engine/Plugins/Animation/GameplayInsights/Source/RewindDebuggerVLog/Private/RewindDebuggerVLogSettings.h
new file mode 100644
--- /dev/null
+++ b/Engine/Plugins/Animation/GameplayInsights/Source/RewindDebuggerVLog/Private/RewindDebuggerVLogSettings.h
@@ -0,0 +1,37 @@
+﻿// Copyright Epic Games, Inc. All Rights Reserved.
+
+#pragma once
+
+#include "CoreMinimal.h"
+#include "Engine/DeveloperSettings.h"
+#include "Logging/LogVerbosity.h"
+#include "RewindDebuggerVLogSettings.generated.h"
+
+/**
+ * Settings for the Rewind Debugger Visual Logger integration.
+ */
+UCLASS(Config=Editor, meta=(DisplayName="Rewind Debugger - Visual Logging"))
+class URewindDebuggerVLogSettings : public UDeveloperSettings
+{
+	GENERATED_BODY()
+	
+public:
+	URewindDebuggerVLogSettings();
+
+#if WITH_EDITOR
+	virtual FText GetSectionText() const override;
+	virtual FText GetSectionDescription() const override;
+#endif
+
+	virtual FName GetCategoryName() const override;
+
+	// Display Visual Logger shapes above this verbosity level
+	UPROPERTY(EditAnywhere, Config, Category = VisualLogger)
+	uint8 DisplayVerbosity = ELogVerbosity::Display;
+
+	// Display Visual Logger shapes in these categories
+	UPROPERTY(EditAnywhere, Config, Category = VisualLogger)
+	TSet<FName> DisplayCategories;
+
+	static URewindDebuggerVLogSettings & Get();
+};
diff --git a/Engine/Plugins/Animation/GameplayInsights/Source/RewindDebuggerVLog/Private/VLogRenderingActor.cpp b/Engine/Plugins/Animation/GameplayInsights/Source/RewindDebuggerVLog/Private/VLogRenderingActor.cpp
--- a/Engine/Plugins/Animation/GameplayInsights/Source/RewindDebuggerVLog/Private/VLogRenderingActor.cpp
+++ b/Engine/Plugins/Animation/GameplayInsights/Source/RewindDebuggerVLog/Private/VLogRenderingActor.cpp
@@ -1,6 +1,7 @@
 // Copyright Epic Games, Inc. All Rights Reserved.
 
 #include "VLogRenderingActor.h"
+#include "RewindDebuggerVLogSettings.h"
 
 #include UE_INLINE_GENERATED_CPP_BY_NAME(VLogRenderingActor)
 
@@ -28,3 +29,9 @@ void AVLogRenderingActor::IterateDebugShapes(TFunction<void(const AVisualLoggerR
 	Callback(DebugShapes);
 }
 
+bool AVLogRenderingActor::MatchCategoryFilters(const FName& CategoryName, ELogVerbosity::Type Verbosity) const
+{
+	URewindDebuggerVLogSettings& Settings = URewindDebuggerVLogSettings::Get();
+	return Settings.DisplayCategories.Contains(CategoryName) && Verbosity <= Settings.DisplayVerbosity;
+}
+
diff --git a/Engine/Plugins/Animation/GameplayInsights/Source/RewindDebuggerVLog/Private/VLogRenderingActor.h b/Engine/Plugins/Animation/GameplayInsights/Source/RewindDebuggerVLog/Private/VLogRenderingActor.h
--- a/Engine/Plugins/Animation/GameplayInsights/Source/RewindDebuggerVLog/Private/VLogRenderingActor.h
+++ b/Engine/Plugins/Animation/GameplayInsights/Source/RewindDebuggerVLog/Private/VLogRenderingActor.h
@@ -11,6 +11,8 @@
 *	Transient actor used to draw visual logger data on level
 */
 
+
+
 UCLASS(config = Engine, NotBlueprintable, Transient, notplaceable, AdvancedClassDisplay)
 class AVLogRenderingActor : public AVisualLoggerRenderingActorBase 
 {
@@ -21,7 +23,8 @@ public:
 	void AddLogEntry(const FVisualLogEntry& LogEntry);
 
 	virtual void IterateDebugShapes(TFunction<void(const FTimelineDebugShapes&) > Callback) override;
-
+	virtual bool MatchCategoryFilters(const FName& CategoryName, ELogVerbosity::Type Verbosity) const override;
+	
 private:
 	FTimelineDebugShapes DebugShapes;
 };
diff --git a/Engine/Plugins/Animation/GameplayInsights/Source/RewindDebuggerVLog/Private/VisualLoggerProvider.cpp b/Engine/Plugins/Animation/GameplayInsights/Source/RewindDebuggerVLog/Private/VisualLoggerProvider.cpp
--- a/Engine/Plugins/Animation/GameplayInsights/Source/RewindDebuggerVLog/Private/VisualLoggerProvider.cpp
+++ b/Engine/Plugins/Animation/GameplayInsights/Source/RewindDebuggerVLog/Private/VisualLoggerProvider.cpp
@@ -28,6 +28,14 @@ bool FVisualLoggerProvider::ReadVisualLogEntryTimeline(uint64 InObjectId, TFunct
 	return false;
 }
 
+void FVisualLoggerProvider::EnumerateCategories(TFunctionRef<void(const FName&)> Callback) const
+{
+	for(const FName& Category : Categories)
+	{
+		Callback(Category);
+	}
+}
+
 void  FVisualLoggerProvider::AppendVisualLogEntry(uint64 InObjectId, double InTime, const FVisualLogEntry& Entry)
 {
 	Session.WriteAccessCheck();
@@ -46,6 +54,11 @@ void  FVisualLoggerProvider::AppendVisualLogEntry(uint64 InObjectId, double InTi
 	}
 
 	Timeline->AppendEvent(InTime, Entry);
+	
+	for (const FVisualLogShapeElement& Shape : Entry.ElementsToDraw)
+	{
+		Categories.AddUnique(Shape.Category);
+	}
 }
 
 #undef LOCTEXT_NAMESPACE
diff --git a/Engine/Plugins/Animation/GameplayInsights/Source/RewindDebuggerVLog/Private/VisualLoggerProvider.h b/Engine/Plugins/Animation/GameplayInsights/Source/RewindDebuggerVLog/Private/VisualLoggerProvider.h
--- a/Engine/Plugins/Animation/GameplayInsights/Source/RewindDebuggerVLog/Private/VisualLoggerProvider.h
+++ b/Engine/Plugins/Animation/GameplayInsights/Source/RewindDebuggerVLog/Private/VisualLoggerProvider.h
@@ -21,6 +21,8 @@ public:
 	/** IVisualLoggerProvider interface */
 	virtual bool ReadVisualLogEntryTimeline(uint64 InObjectId, TFunctionRef<void(const VisualLogEntryTimeline&)> Callback) const;
 
+	virtual void EnumerateCategories(TFunctionRef<void(const FName&)> Callback) const;
+
 	/** Add an object event message */
 	void AppendVisualLogEntry(uint64 InObjectId, double InTime, const FVisualLogEntry& Entry);
 
@@ -29,4 +31,5 @@ private:
 
 	TMap<uint64, uint32> ObjectIdToLogEntryTimelines;
 	TArray<TSharedRef<TraceServices::TPointTimeline<FVisualLogEntry>>> LogEntryTimelines;
+	TArray<FName> Categories;
 };
diff --git a/Engine/Plugins/Animation/GameplayInsights/Source/RewindDebuggerVLog/Public/IVisualLoggerProvider.h b/Engine/Plugins/Animation/GameplayInsights/Source/RewindDebuggerVLog/Public/IVisualLoggerProvider.h
--- a/Engine/Plugins/Animation/GameplayInsights/Source/RewindDebuggerVLog/Public/IVisualLoggerProvider.h
+++ b/Engine/Plugins/Animation/GameplayInsights/Source/RewindDebuggerVLog/Public/IVisualLoggerProvider.h
@@ -13,5 +13,6 @@ public:
 	typedef TraceServices::ITimeline<FVisualLogEntry> VisualLogEntryTimeline;
 
 	virtual bool ReadVisualLogEntryTimeline(uint64 InObjectId, TFunctionRef<void(const VisualLogEntryTimeline&)> Callback) const = 0;
+	virtual void EnumerateCategories(TFunctionRef<void(const FName&)> Callback) const = 0;
 };
 
diff --git a/Engine/Plugins/Animation/GameplayInsights/Source/RewindDebuggerVLog/RewindDebuggerVLog.Build.cs b/Engine/Plugins/Animation/GameplayInsights/Source/RewindDebuggerVLog/RewindDebuggerVLog.Build.cs
--- a/Engine/Plugins/Animation/GameplayInsights/Source/RewindDebuggerVLog/RewindDebuggerVLog.Build.cs
+++ b/Engine/Plugins/Animation/GameplayInsights/Source/RewindDebuggerVLog/RewindDebuggerVLog.Build.cs
@@ -12,11 +12,18 @@ namespace UnrealBuildTool.Rules
 				"Core",
 				"CoreUObject",
 				"Engine",
+				"UnrealEd",
 				"LogVisualizer",
 				"TraceAnalysis",
 				"TraceServices",
 				"TraceInsights",
-				"RewindDebuggerInterface"
+				"RewindDebuggerInterface",
+				"GameplayInsights",
+				"Slate",
+				"SlateCore",
+				"ToolWidgets",
+				"ToolMenus",
+				"DeveloperSettings",
 			});
 		}
 	}
diff --git a/Engine/Plugins/Animation/IKRig/Source/IKRig/Private/AnimNodes/AnimNode_IKRig.cpp b/Engine/Plugins/Animation/IKRig/Source/IKRig/Private/AnimNodes/AnimNode_IKRig.cpp
--- a/Engine/Plugins/Animation/IKRig/Source/IKRig/Private/AnimNodes/AnimNode_IKRig.cpp
+++ b/Engine/Plugins/Animation/IKRig/Source/IKRig/Private/AnimNodes/AnimNode_IKRig.cpp
@@ -13,6 +13,8 @@
 #include "Algo/ForEach.h"
 #include "SceneManagement.h"
 #include "Engine/SkeletalMesh.h"
+#include "Logging/LogVerbosity.h"
+#include "VisualLogger/VisualLogger.h"
 
 #include UE_INLINE_GENERATED_CPP_BY_NAME(AnimNode_IKRig)
 
@@ -199,13 +201,13 @@ void FAnimNode_IKRig::Update_AnyThread(const FAnimationUpdateContext& Context)
 	// alpha handlers
 	switch (AlphaInputType)
 	{
-	case EAnimAlphaInputType::Float : 
+	case EAnimAlphaInputType::Float:
 		ActualAlpha = AlphaScaleBias.ApplyTo(AlphaScaleBiasClamp.ApplyTo(Alpha, Context.GetDeltaTime()));
 		break;
-	case EAnimAlphaInputType::Bool :
+	case EAnimAlphaInputType::Bool:
 		ActualAlpha = AlphaBoolBlend.ApplyTo(bAlphaBoolEnabled, Context.GetDeltaTime());
 		break;
-	case EAnimAlphaInputType::Curve :
+	case EAnimAlphaInputType::Curve:
 		if (UAnimInstance* AnimInstance = Cast<UAnimInstance>(Context.AnimInstanceProxy->GetAnimInstanceObject()))
 		{
 			ActualAlpha = AlphaScaleBiasClamp.ApplyTo(AnimInstance->GetCurveValue(AlphaCurveName), Context.GetDeltaTime());
@@ -217,9 +219,30 @@ void FAnimNode_IKRig::Update_AnyThread(const FAnimationUpdateContext& Context)
 	ActualAlpha = FMath::Clamp<float>(ActualAlpha, 0.f, 1.f);
 
 	PropagateInputProperties(Context.AnimInstanceProxy->GetAnimInstanceObject());
-	
+
 	FAnimNode_Base::Update_AnyThread(Context);
 	Source.Update(Context);
+
+#if ENABLE_VISUAL_LOG
+#if WITH_EDITORONLY_DATA
+	if( FVisualLogger::IsRecording() )
+	{
+		static const FBox UnitBox(FVector(-1, -1, -1), FVector(1, 1, 1));
+
+		const TArray<FIKRigGoal>& ProcessorGoals = IKRigProcessor->GetGoalContainer().GetGoalArray();
+		for (const FIKRigGoal& Goal : ProcessorGoals)
+		{
+			FTransform GoalTransform(Goal.FinalBlendedRotation, Goal.FinalBlendedPosition, FVector(DebugScale, DebugScale, DebugScale));
+			GoalTransform = GoalTransform * Context.AnimInstanceProxy->GetComponentTransform();
+			UE_VLOG_OBOX(Context.AnimInstanceProxy->GetAnimInstanceObject(), "IKRig", Display, UnitBox, GoalTransform.ToMatrixWithScale(), FColor::Yellow, TEXT(""));
+			UE_VLOG_LOCATION(Context.AnimInstanceProxy->GetAnimInstanceObject(), "IKRig", Verbose, GoalTransform.GetTranslation(), 0.f, FColor::White, TEXT("%ls"), ToCStr(Goal.Name.ToString()));
+		}
+	}
+#endif
+#endif
+
+	TRACE_ANIM_NODE_VALUE(Context, TEXT("Name"), RigDefinitionAsset ? ToCStr(RigDefinitionAsset->GetName()) : TEXT(""));
+	TRACE_ANIM_NODE_VALUE(Context, TEXT("Asset"), RigDefinitionAsset);
 }
 
 void FAnimNode_IKRig::PreUpdate(const UAnimInstance* InAnimInstance)
diff --git a/Engine/Source/Developer/LogVisualizer/Private/VisualLoggerRenderingActor.cpp b/Engine/Source/Developer/LogVisualizer/Private/VisualLoggerRenderingActor.cpp
--- a/Engine/Source/Developer/LogVisualizer/Private/VisualLoggerRenderingActor.cpp
+++ b/Engine/Source/Developer/LogVisualizer/Private/VisualLoggerRenderingActor.cpp
@@ -4,6 +4,7 @@
 #include "AI/NavigationSystemBase.h"
 #include "VisualLogger/VisualLogger.h"
 #include "VisualLoggerDatabase.h"
+#include "LogVisualizerSettings.h"
 #include "LogVisualizerPublic.h"
 
 #include UE_INLINE_GENERATED_CPP_BY_NAME(VisualLoggerRenderingActor)
@@ -215,3 +216,8 @@ void AVisualLoggerRenderingActor::IterateDebugShapes(const TFunction<void(const
 	Callback(TestDebugShapes);
 #endif
 }
+
+bool AVisualLoggerRenderingActor::MatchCategoryFilters(const FName& CategoryName, ELogVerbosity::Type Verbosity) const
+{
+	return FVisualLoggerFilters::Get().MatchCategoryFilters(CategoryName.ToString(), Verbosity);
+}
diff --git a/Engine/Source/Developer/LogVisualizer/Private/VisualLoggerRenderingActor.h b/Engine/Source/Developer/LogVisualizer/Private/VisualLoggerRenderingActor.h
--- a/Engine/Source/Developer/LogVisualizer/Private/VisualLoggerRenderingActor.h
+++ b/Engine/Source/Developer/LogVisualizer/Private/VisualLoggerRenderingActor.h
@@ -29,6 +29,7 @@ public:
 	void OnItemSelectionChanged(const FVisualLoggerDBRow& BDRow, int32 ItemIndex);
 
 	virtual void IterateDebugShapes(TFunction<void(const FTimelineDebugShapes&) > Callback) override;
+	virtual bool MatchCategoryFilters(const FName& CategoryName, ELogVerbosity::Type Verbosity) const override;
 private:
 	void OnFiltersChanged();
 
diff --git a/Engine/Source/Developer/LogVisualizer/Private/VisualLoggerRenderingActorBase.cpp b/Engine/Source/Developer/LogVisualizer/Private/VisualLoggerRenderingActorBase.cpp
--- a/Engine/Source/Developer/LogVisualizer/Private/VisualLoggerRenderingActorBase.cpp
+++ b/Engine/Source/Developer/LogVisualizer/Private/VisualLoggerRenderingActorBase.cpp
@@ -193,12 +193,11 @@ void AVisualLoggerRenderingActorBase::GetDebugShapes(const FVisualLogEntry& InEn
 
 	for (int32 ElementIndex = 0; ElementIndex < ElementsCount; ++ElementIndex, ++ElementToDraw)
 	{
-		if (!FVisualLoggerFilters::Get().MatchCategoryFilters(ElementToDraw->Category.ToString(), ElementToDraw->Verbosity))
+		if (!MatchCategoryFilters(ElementToDraw->Category, ElementToDraw->Verbosity))
 		{
 			continue;
 		}
 
-
 		const FVector3f CorridorOffset = (FVector3f)FDebugDrawing::NavOffset * 1.25f;
 		const FColor Color = ElementToDraw->GetFColor();
 
diff --git a/Engine/Source/Developer/LogVisualizer/Public/VisualLoggerRenderingActorBase.h b/Engine/Source/Developer/LogVisualizer/Public/VisualLoggerRenderingActorBase.h
--- a/Engine/Source/Developer/LogVisualizer/Public/VisualLoggerRenderingActorBase.h
+++ b/Engine/Source/Developer/LogVisualizer/Public/VisualLoggerRenderingActorBase.h
@@ -50,7 +50,8 @@ public:
 	};
 
 	// Iterate over each active FTimelineDebugShapes, and call callback
-	virtual void IterateDebugShapes(TFunction<void(const FTimelineDebugShapes&) > Callback) {};
+	virtual void IterateDebugShapes(TFunction<void(const FTimelineDebugShapes&) > Callback) { };
+	virtual bool MatchCategoryFilters(const FName& CategoryName, ELogVerbosity::Type Verbosity) const { return true; };
 
 protected:
 #if ENABLE_VISUAL_LOG
