Visual logger track for Rewind Debugger

diff --git a/Engine/Plugins/Animation/GameplayInsights/Source/GameplayInsights/Private/SEventTimelineView.cpp b/Engine/Plugins/Animation/GameplayInsights/Source/GameplayInsights/Private/SEventTimelineView.cpp
--- a/Engine/Plugins/Animation/GameplayInsights/Source/GameplayInsights/Private/SEventTimelineView.cpp
+++ b/Engine/Plugins/Animation/GameplayInsights/Source/GameplayInsights/Private/SEventTimelineView.cpp
@@ -63,9 +63,15 @@ int32 SEventTimelineView::PaintEvents(const FGeometry& AllottedGeometry, const F
 				float Y = (AllottedGeometry.Size.Y - EventSize.Y)/2;
 				if (Point.Time == PrevPointTime)
 				{
-					OverlappingPointCount++;
-
 					static const int OverlapOffsetAmount = 2;
+					static const int MaxOverlappingEvents = 2;
+				
+					if (OverlappingPointCount > MaxOverlappingEvents)
+					{
+						continue;
+					}
+					
+					OverlappingPointCount++;
 					Y += OverlapOffsetAmount * OverlappingPointCount;
 				}
 				else
diff --git a/Engine/Plugins/Animation/GameplayInsights/Source/GameplayInsights/Public/SEventTimelineView.h b/Engine/Plugins/Animation/GameplayInsights/Source/GameplayInsights/Public/SEventTimelineView.h
--- a/Engine/Plugins/Animation/GameplayInsights/Source/GameplayInsights/Public/SEventTimelineView.h
+++ b/Engine/Plugins/Animation/GameplayInsights/Source/GameplayInsights/Public/SEventTimelineView.h
@@ -10,7 +10,7 @@ class FPaintArgs;
 class FSlateWindowElementList;
 class FTimeSliderController;
 
-class SEventTimelineView : public SCompoundWidget
+class GAMEPLAYINSIGHTS_API SEventTimelineView : public SCompoundWidget
 {
 public:
 	struct FTimelineEventData
diff --git a/Engine/Plugins/Animation/GameplayInsights/Source/RewindDebugger/Private/RewindDebugger.cpp b/Engine/Plugins/Animation/GameplayInsights/Source/RewindDebugger/Private/RewindDebugger.cpp
--- a/Engine/Plugins/Animation/GameplayInsights/Source/RewindDebugger/Private/RewindDebugger.cpp
+++ b/Engine/Plugins/Animation/GameplayInsights/Source/RewindDebugger/Private/RewindDebugger.cpp
@@ -2,33 +2,36 @@
 
 #include "RewindDebugger.h"
 
+#include "Animation/AnimBlueprint.h"
+#include "Animation/AnimBlueprintGeneratedClass.h"
+#include "Animation/AnimTrace.h"
 #include "Components/SkeletalMeshComponent.h"
 #include "Editor.h"
+#include "Engine/PoseWatch.h"
+#include "Engine/World.h"
+#include "EngineUtils.h"
+#include "GameFramework/Controller.h"
+#include "GameFramework/Pawn.h"
 #include "IAnimationProvider.h"
 #include "IGameplayProvider.h"
+#include "IRewindDebuggerDoubleClickHandler.h"
+#include "IRewindDebuggerExtension.h"
 #include "Insights/IUnrealInsightsModule.h"
+#include "LevelEditor.h"
 #include "Modules/ModuleManager.h"
-#include "Animation/AnimTrace.h"
 #include "ObjectTrace.h"
-#include "TraceServices/Model/Frames.h"
+#include "ProfilingDebugging/TraceAuxiliary.h"
+#include "RewindDebuggerCommands.h"
+#include "RewindDebuggerModule.h"
+#include "RewindDebuggerObjectTrack.h"
+#include "RewindDebuggerSettings.h"
 #include "SLevelViewport.h"
+#include "ToolMenus.h"
+#include "TraceServices/Model/Frames.h"
+#include "UObject/UObjectIterator.h"
 #include "Widgets/Docking/SDockTab.h"
 #include "Widgets/Layout/SSpacer.h"
-#include "IRewindDebuggerExtension.h"
-#include "IRewindDebuggerDoubleClickHandler.h"
-#include "RewindDebuggerObjectTrack.h"
-#include "Animation/AnimBlueprint.h"
-#include "Animation/AnimBlueprintGeneratedClass.h"
-#include "EngineUtils.h"
-#include "ToolMenus.h"
-#include "RewindDebuggerSettings.h"
 #include "RewindDebuggerCommands.h"
-#include "LevelEditor.h"
-#include "RewindDebuggerModule.h"
-#include "Engine/PoseWatch.h"
-#include "ProfilingDebugging/TraceAuxiliary.h"
-#include "UObject/UObjectIterator.h"
-#include "Engine/World.h"
 
 #define LOCTEXT_NAMESPACE "RewindDebugger"
 
@@ -831,6 +834,13 @@ void FRewindDebugger::Tick(float DeltaTime)
 				for (TActorIterator<AActor> Iterator(*World); Iterator; ++Iterator)
 				{
 					TRACE_OBJECT_LIFETIME_BEGIN(*Iterator);
+					if (APawn* Pawn = Cast<APawn>(*Iterator))
+					{
+						if (AController* Controller = Pawn->GetController())
+						{
+							TRACE_PAWN_POSSESS(static_cast<UObject*>(Controller), static_cast<UObject*>(Pawn));
+						}
+					}
 				}
 			}
 		}
diff --git a/Engine/Plugins/Animation/GameplayInsights/Source/RewindDebuggerVLog/Private/RewindDebuggerVLog.cpp b/Engine/Plugins/Animation/GameplayInsights/Source/RewindDebuggerVLog/Private/RewindDebuggerVLog.cpp
--- a/Engine/Plugins/Animation/GameplayInsights/Source/RewindDebuggerVLog/Private/RewindDebuggerVLog.cpp
+++ b/Engine/Plugins/Animation/GameplayInsights/Source/RewindDebuggerVLog/Private/RewindDebuggerVLog.cpp
@@ -7,11 +7,14 @@
 #include "IVisualLoggerProvider.h"
 #include "Insights/IUnrealInsightsModule.h"
 #include "LogVisualizerSettings.h"
+#include "Editor/EditorEngine.h"
 #include "Modules/ModuleManager.h"
 #include "RewindDebuggerVLogSettings.h"
 #include "ToolMenus.h"
 #include "TraceServices/Model/Frames.h"
 #include "VisualLogger/VisualLogger.h"
+#include "RewindDebuggerVLogSettings.h"
+#include "ToolMenus.h"
 
 #define LOCTEXT_NAMESPACE "RewindDebuggerVLog"
 
diff --git a/Engine/Plugins/Animation/GameplayInsights/Source/RewindDebuggerVLog/Private/RewindDebuggerVLogModule.cpp b/Engine/Plugins/Animation/GameplayInsights/Source/RewindDebuggerVLog/Private/RewindDebuggerVLogModule.cpp
--- a/Engine/Plugins/Animation/GameplayInsights/Source/RewindDebuggerVLog/Private/RewindDebuggerVLogModule.cpp
+++ b/Engine/Plugins/Animation/GameplayInsights/Source/RewindDebuggerVLog/Private/RewindDebuggerVLogModule.cpp
@@ -1,6 +1,8 @@
 // Copyright Epic Games, Inc. All Rights Reserved.
 
 #include "RewindDebuggerVLogModule.h"
+
+#include "VisualLogTrack.h"
 #include "Modules/ModuleManager.h"
 #include "Features/IModularFeatures.h"
 
@@ -10,11 +12,18 @@
 #define ENABLE_REWINDDEBUGGER_VLOG_INTEGRATION 1
 #endif
 
+
+#if ENABLE_REWINDDEBUGGER_VLOG_INTEGRATION
+RewindDebugger::FVisualLogTrackCreator GVisualLogTrackCreator;
+#endif
+
 void FRewindDebuggerVLogModule::StartupModule()
 {
 #if ENABLE_REWINDDEBUGGER_VLOG_INTEGRATION
 	IModularFeatures::Get().RegisterModularFeature(IRewindDebuggerExtension::ModularFeatureName, &RewindDebuggerVLogExtension);
 	IModularFeatures::Get().RegisterModularFeature(TraceServices::ModuleFeatureName, &VLogTraceModule);
+	IModularFeatures::Get().RegisterModularFeature(RewindDebugger::IRewindDebuggerTrackCreator::ModularFeatureName, &GVisualLogTrackCreator);
+
 	RewindDebuggerVLogExtension.Initialize();
 #endif
 }
@@ -24,6 +33,7 @@ void FRewindDebuggerVLogModule::ShutdownModule()
 #if ENABLE_REWINDDEBUGGER_VLOG_INTEGRATION
 	IModularFeatures::Get().UnregisterModularFeature(IRewindDebuggerExtension::ModularFeatureName, &RewindDebuggerVLogExtension);
 	IModularFeatures::Get().UnregisterModularFeature(TraceServices::ModuleFeatureName, &VLogTraceModule);
+	IModularFeatures::Get().UnregisterModularFeature(RewindDebugger::IRewindDebuggerTrackCreator::ModularFeatureName, &GVisualLogTrackCreator);
 #endif
 }
 
diff --git a/Engine/Plugins/Animation/GameplayInsights/Source/RewindDebuggerVLog/Private/VLogDetailsObject.cpp b/Engine/Plugins/Animation/GameplayInsights/Source/RewindDebuggerVLog/Private/VLogDetailsObject.cpp
new file mode 100644
--- /dev/null
+++ b/Engine/Plugins/Animation/GameplayInsights/Source/RewindDebuggerVLog/Private/VLogDetailsObject.cpp
@@ -0,0 +1,3 @@
+// Copyright Epic Games, Inc. All Rights Reserved.
+
+#include "VLogDetailsObject.h"
diff --git a/Engine/Plugins/Animation/GameplayInsights/Source/RewindDebuggerVLog/Private/VLogDetailsObject.h b/Engine/Plugins/Animation/GameplayInsights/Source/RewindDebuggerVLog/Private/VLogDetailsObject.h
new file mode 100644
--- /dev/null
+++ b/Engine/Plugins/Animation/GameplayInsights/Source/RewindDebuggerVLog/Private/VLogDetailsObject.h
@@ -0,0 +1,27 @@
+// Copyright Epic Games, Inc. All Rights Reserved.
+#pragma once
+
+#include "UObject/Object.h"
+#include "VLogDetailsObject.generated.h"
+
+USTRUCT()
+struct FVisualLogDetails
+{
+	GENERATED_BODY()
+	
+	UPROPERTY(VisibleAnywhere, Category=Details)
+	FName Category;
+	
+	UPROPERTY(VisibleAnywhere, Category=Details)
+	FString Description;
+};
+
+UCLASS()
+class UVLogDetailsObject : public UObject
+{
+	GENERATED_BODY()
+
+public:
+	UPROPERTY(VisibleAnywhere, Category=Details)
+	TArray<FVisualLogDetails> VisualLogDetails;
+};
\ No newline at end of file
diff --git a/Engine/Plugins/Animation/GameplayInsights/Source/RewindDebuggerVLog/Private/VisualLogTrack.cpp b/Engine/Plugins/Animation/GameplayInsights/Source/RewindDebuggerVLog/Private/VisualLogTrack.cpp
new file mode 100644
--- /dev/null
+++ b/Engine/Plugins/Animation/GameplayInsights/Source/RewindDebuggerVLog/Private/VisualLogTrack.cpp
@@ -0,0 +1,251 @@
+// Copyright Epic Games, Inc. All Rights Reserved.
+
+#include "VisualLogTrack.h"
+#include "IRewindDebugger.h"
+#include "Modules/ModuleManager.h"
+#include "PropertyEditorModule.h"
+#include "VisualLoggerProvider.h"
+#include "UObject/Package.h"
+
+#define LOCTEXT_NAMESPACE "VisualLogTrack"
+
+namespace RewindDebugger
+{
+
+FVisualLogCategoryTrack::FVisualLogCategoryTrack(uint64 InObjectId, const FName& InCategory) :
+	ObjectId(InObjectId),
+	Category(InCategory)
+{
+	TrackName = FText::FromName(Category);
+	EventData = MakeShared<SEventTimelineView::FTimelineEventData>();
+	Icon = FSlateIcon("EditorStyle", "Sequencer.Tracks.Event", "Sequencer.Tracks.Event");
+
+	DetailsObject = NewObject<UVLogDetailsObject>(GetTransientPackage(), UVLogDetailsObject::StaticClass(), NAME_None, RF_Transient|RF_Standalone);
+	
+	FPropertyEditorModule& PropertyEditorModule = FModuleManager::GetModuleChecked<FPropertyEditorModule>("PropertyEditor");
+	FDetailsViewArgs DetailsViewArgs;
+	DetailsViewArgs.NameAreaSettings = FDetailsViewArgs::HideNameArea;
+	DetailsView = PropertyEditorModule.CreateDetailView(DetailsViewArgs);
+	DetailsView->SetObject(DetailsObject);
+	
+}
+
+FVisualLogCategoryTrack::~FVisualLogCategoryTrack()
+{
+	DetailsObject->SetFlags(RF_Standalone);
+}
+
+TSharedPtr<SEventTimelineView::FTimelineEventData> FVisualLogCategoryTrack::GetEventData() const
+{
+	if (!EventData.IsValid())
+	{
+		EventData = MakeShared<SEventTimelineView::FTimelineEventData>();
+	}
+	
+	EventUpdateRequested++;
+	
+	return EventData;
+}
+	
+static FLinearColor MakeNotifyColor(uint32 InSeed, bool bInLine = false)
+{
+	FRandomStream Stream(InSeed);
+	const uint8 Hue = (uint8)(Stream.FRand() * 255.0f);
+	const uint8 SatVal = bInLine ? 196 : 128;
+	return FLinearColor::MakeFromHSV8(Hue, SatVal, SatVal);
+}
+
+TSharedPtr<SWidget> FVisualLogCategoryTrack::GetDetailsViewInternal() 
+{
+	return DetailsView;
+}
+
+bool FVisualLogCategoryTrack::UpdateInternal()
+{
+	IRewindDebugger* RewindDebugger = IRewindDebugger::Instance();
+	
+	TRange<double> TraceTimeRange = RewindDebugger->GetCurrentTraceRange();
+	double StartTime = TraceTimeRange.GetLowerBoundValue();
+	double EndTime = TraceTimeRange.GetUpperBoundValue();
+	
+	const TraceServices::IAnalysisSession* AnalysisSession = RewindDebugger->GetAnalysisSession();
+	const FVisualLoggerProvider* VisLogProvider = AnalysisSession->ReadProvider<FVisualLoggerProvider>(FVisualLoggerProvider::ProviderName);
+	
+	if(EventUpdateRequested > 10 && VisLogProvider)
+	{
+		EventUpdateRequested = 0;
+		
+		EventData->Points.SetNum(0,false);
+		EventData->Windows.SetNum(0);
+
+		TraceServices::FAnalysisSessionReadScope SessionReadScope(*AnalysisSession);
+		
+		VisLogProvider->ReadVisualLogEntryTimeline(ObjectId, [this, StartTime, EndTime, VisLogProvider, AnalysisSession](const FVisualLoggerProvider::VisualLogEntryTimeline& InTimeline)
+		{
+			InTimeline.EnumerateEvents(StartTime, EndTime, [this, StartTime, EndTime, VisLogProvider, AnalysisSession](double InStartTime, double InEndTime, uint32 InDepth, const FVisualLogEntry& InMessage)
+			{
+				for(const FVisualLogShapeElement& Element : InMessage.ElementsToDraw)
+				{
+					EventData->Points.Add({InMessage.TimeStamp,FText::FromName(Element.Category), FText::FromString(Element.Description),Element.GetFColor()});
+				}
+				return TraceServices::EEventEnumerate::Continue;
+			});
+		});
+	}
+
+
+	double CurrentScrubTime = IRewindDebugger::Instance()->CurrentTraceTime();
+	if (PreviousScrubTime != CurrentScrubTime)
+	{
+		PreviousScrubTime = CurrentScrubTime;
+		DetailsObject->VisualLogDetails.SetNum(0,false);
+		
+		const TraceServices::IFrameProvider& FramesProvider = TraceServices::ReadFrameProvider(*AnalysisSession);
+		TraceServices::FFrame MarkerFrame;
+		if(FramesProvider.GetFrameFromTime(ETraceFrameType::TraceFrameType_Game, CurrentScrubTime, MarkerFrame))
+		{
+			VisLogProvider->ReadVisualLogEntryTimeline(ObjectId, [this, &MarkerFrame, EndTime, VisLogProvider, AnalysisSession](const FVisualLoggerProvider::VisualLogEntryTimeline& InTimeline)
+			{
+				InTimeline.EnumerateEvents(MarkerFrame.StartTime, MarkerFrame.EndTime, [this, VisLogProvider, AnalysisSession](double InStartTime, double InEndTime, uint32 InDepth, const FVisualLogEntry& InMessage)
+				{
+					for(const FVisualLogShapeElement& Element : InMessage.ElementsToDraw)
+					{
+						DetailsObject->VisualLogDetails.Add({Element.Category, Element.Description});
+					}
+					
+					for(const FVisualLogLine& Line : InMessage.LogLines)
+					{
+						DetailsObject->VisualLogDetails.Add({Line.Category, Line.Line});
+					}
+					return TraceServices::EEventEnumerate::Continue;
+				});
+			});
+		}
+	}
+	
+	
+	bool bChanged = false;
+
+	return bChanged;
+	
+}
+
+TSharedPtr<SWidget> FVisualLogTrack::GetDetailsViewInternal() 
+{
+	return nullptr;
+}
+
+TSharedPtr<SWidget> FVisualLogCategoryTrack::GetTimelineViewInternal()
+{
+	return SNew(SEventTimelineView)
+		.ViewRange_Lambda([]() { return IRewindDebugger::Instance()->GetCurrentViewRange(); })
+		.EventData_Raw(this, &FVisualLogCategoryTrack::GetEventData);
+}
+
+static const FName VisualLogName("Visual Logging");
+
+FName FVisualLogTrackCreator::GetTargetTypeNameInternal() const
+{
+ 	static const FName AnimInstanceName("Object");
+	return AnimInstanceName;
+}
+	
+FName FVisualLogTrackCreator::GetNameInternal() const
+{
+	return VisualLogName;
+}
+
+TSharedPtr<RewindDebugger::FRewindDebuggerTrack> FVisualLogTrackCreator::CreateTrackInternal(uint64 ObjectId) const
+{
+	return MakeShared<RewindDebugger::FVisualLogTrack>(ObjectId);
+}
+
+		
+FVisualLogTrack::FVisualLogTrack(uint64 InObjectId) :
+	ObjectId(InObjectId)
+{
+	Icon = FSlateIcon("EditorStyle", "Sequencer.Tracks.Event", "Sequencer.Tracks.Event");
+}
+
+
+bool FVisualLogTrack::UpdateInternal()
+{
+	IRewindDebugger* RewindDebugger = IRewindDebugger::Instance();
+	
+	TRange<double> TraceTimeRange = RewindDebugger->GetCurrentTraceRange();
+	double StartTime = TraceTimeRange.GetLowerBoundValue();
+	double EndTime = TraceTimeRange.GetUpperBoundValue();
+	
+	const TraceServices::IAnalysisSession* AnalysisSession = RewindDebugger->GetAnalysisSession();
+	const FVisualLoggerProvider* VisLogProvider = AnalysisSession->ReadProvider<FVisualLoggerProvider>(FVisualLoggerProvider::ProviderName);
+	
+	bool bChanged = false;
+	
+	if(VisLogProvider)
+	{
+		TArray<FName> UniqueTrackIds;
+		
+		TraceServices::FAnalysisSessionReadScope SessionReadScope(*AnalysisSession);
+	
+		VisLogProvider->ReadVisualLogEntryTimeline(ObjectId, [this, StartTime, EndTime, VisLogProvider, AnalysisSession, &UniqueTrackIds](const FVisualLoggerProvider::VisualLogEntryTimeline& InTimeline)
+		{
+			InTimeline.EnumerateEvents(StartTime, EndTime, [this, StartTime, EndTime, VisLogProvider, AnalysisSession, &UniqueTrackIds](double InStartTime, double InEndTime, uint32 InDepth, const FVisualLogEntry& InMessage)
+			{
+				for(const FVisualLogShapeElement& Element : InMessage.ElementsToDraw)
+				{
+					UniqueTrackIds.AddUnique(Element.Category);
+				}
+				return TraceServices::EEventEnumerate::Continue;
+			});
+		});
+	
+		UniqueTrackIds.StableSort([](const FName& A, const FName& B) { return A.ToString() < B.ToString(); });
+		const int32 TrackCount = UniqueTrackIds.Num();
+	
+		if (Children.Num() != TrackCount)
+			bChanged = true;
+	
+		Children.SetNum(UniqueTrackIds.Num());
+		for(int i = 0; i < TrackCount; i++)
+		{
+			if (!Children[i].IsValid() || !(Children[i].Get()->GetName() == UniqueTrackIds[i]))
+			{
+				Children[i] = MakeShared<FVisualLogCategoryTrack>(ObjectId, UniqueTrackIds[i]);
+				bChanged = true;
+			}
+	
+			bChanged = bChanged || Children[i]->Update();
+		}
+	}
+	
+	return bChanged;
+}
+	
+void FVisualLogTrack::IterateSubTracksInternal(TFunction<void(TSharedPtr<FRewindDebuggerTrack> SubTrack)> IteratorFunction)
+{
+	for(TSharedPtr<FVisualLogCategoryTrack>& Track : Children)
+	{
+		IteratorFunction(Track);
+	}
+};
+
+bool FVisualLogTrackCreator::HasDebugInfoInternal(uint64 ObjectId) const
+{
+	const TraceServices::IAnalysisSession* AnalysisSession = IRewindDebugger::Instance()->GetAnalysisSession();
+	
+	TraceServices::FAnalysisSessionReadScope SessionReadScope(*AnalysisSession);
+	bool bHasData = false;
+	if (const FVisualLoggerProvider* VLogProvider = AnalysisSession->ReadProvider<FVisualLoggerProvider>(FVisualLoggerProvider::ProviderName))
+	{
+		VLogProvider->ReadVisualLogEntryTimeline(ObjectId,  [&bHasData](const FVisualLoggerProvider::VisualLogEntryTimeline& InTimeline)
+	 	{
+	 		bHasData = true;
+	 	});
+	}
+	return bHasData;
+}
+
+	
+}
+
+#undef LOCTEXT_NAMESPACE
\ No newline at end of file
diff --git a/Engine/Plugins/Animation/GameplayInsights/Source/RewindDebuggerVLog/Private/VisualLogTrack.h b/Engine/Plugins/Animation/GameplayInsights/Source/RewindDebuggerVLog/Private/VisualLogTrack.h
new file mode 100644
--- /dev/null
+++ b/Engine/Plugins/Animation/GameplayInsights/Source/RewindDebuggerVLog/Private/VisualLogTrack.h
@@ -0,0 +1,85 @@
+// Copyright Epic Games, Inc. All Rights Reserved.
+#pragma once
+
+#include "IRewindDebuggerTrackCreator.h"
+#include "RewindDebuggerTrack.h"
+#include "Textures/SlateIcon.h"
+#include "SEventTimelineView.h"
+#include "IDetailsView.h"
+#include "VLogDetailsObject.h"
+
+namespace RewindDebugger
+{
+
+class FVisualLogCategoryTrack : public FRewindDebuggerTrack
+{
+public:
+	FVisualLogCategoryTrack(uint64 InObjectId, const FName& Category);
+	virtual ~FVisualLogCategoryTrack();
+	const FName GetCategory() const { return Category; }
+private:
+	virtual bool UpdateInternal() override;
+	TSharedPtr<SWidget> GetDetailsViewInternal() override;
+	virtual TSharedPtr<SWidget> GetTimelineViewInternal() override;
+	virtual FSlateIcon GetIconInternal() override { return Icon; }
+	virtual FText GetDisplayNameInternal() const override { return TrackName; }
+	virtual uint64 GetObjectIdInternal() const override { return ObjectId; }
+	virtual FName GetNameInternal() const override { return Category; }
+
+	TSharedPtr<SEventTimelineView::FTimelineEventData> GetEventData() const;
+	
+	mutable TSharedPtr<SEventTimelineView::FTimelineEventData> EventData;
+	mutable int EventUpdateRequested = 0;
+	
+	uint64 ObjectId;
+	FName Category;
+	FSlateIcon Icon;
+	FText TrackName;
+	TSharedPtr<IDetailsView> DetailsView;
+	UVLogDetailsObject* DetailsObject;
+
+	double PreviousScrubTime = 0.0f;
+	
+	enum class EFilterType
+	{
+		VisualLog,
+		VisualLogState,
+		SyncMarker,
+	};
+};
+
+
+
+// Parent track 
+class FVisualLogTrack : public FRewindDebuggerTrack
+{
+public:
+	FVisualLogTrack(uint64 InObjectId);
+
+private:
+	virtual bool UpdateInternal() override;
+	virtual TSharedPtr<SWidget> GetDetailsViewInternal() override;
+
+	virtual FSlateIcon GetIconInternal() override { return Icon; }
+	virtual FName GetNameInternal() const override { return "Visual Logging"; }
+	virtual FText GetDisplayNameInternal() const override { return NSLOCTEXT("RewindDebugger", "VisualLoggerTrackName", "Visual Logging"); }
+	virtual uint64 GetObjectIdInternal() const override { return ObjectId; }
+	virtual void IterateSubTracksInternal(TFunction<void(TSharedPtr<FRewindDebuggerTrack> SubTrack)> IteratorFunction) override;
+
+	FSlateIcon Icon;
+	uint64 ObjectId;
+
+	TArray<TSharedPtr<FVisualLogCategoryTrack>> Children;
+};
+
+
+class FVisualLogTrackCreator : public IRewindDebuggerTrackCreator
+{
+private:
+	virtual FName GetTargetTypeNameInternal() const override;
+	virtual FName GetNameInternal() const override;
+	virtual TSharedPtr<FRewindDebuggerTrack> CreateTrackInternal(uint64 ObjectId) const override;
+	virtual bool HasDebugInfoInternal(uint64 ObjectId) const override;
+};
+
+}
\ No newline at end of file
diff --git a/Engine/Plugins/Animation/GameplayInsights/Source/RewindDebuggerVLog/Private/VisualLoggerAnalyzer.cpp b/Engine/Plugins/Animation/GameplayInsights/Source/RewindDebuggerVLog/Private/VisualLoggerAnalyzer.cpp
--- a/Engine/Plugins/Animation/GameplayInsights/Source/RewindDebuggerVLog/Private/VisualLoggerAnalyzer.cpp
+++ b/Engine/Plugins/Animation/GameplayInsights/Source/RewindDebuggerVLog/Private/VisualLoggerAnalyzer.cpp
@@ -33,10 +33,12 @@ bool FVisualLoggerAnalyzer::OnEvent(uint16 RouteId, EStyle Style, const FOnEvent
 		{
 			uint64 OwnerId = EventData.GetValue<uint64>("OwnerId");
 			uint64 Cycle = EventData.GetValue<uint64>("Cycle");
+			double RecordingTime = EventData.GetValue<double>("RecordingTime");
 			TArrayView<const uint8> SerializedData = EventData.GetArrayView<uint8>("LogEntry");
 			FMemoryReaderView Archive(SerializedData);
 			FVisualLogEntry Entry;
 			Archive << Entry;
+			Entry.TimeStamp = RecordingTime;
 			Provider.AppendVisualLogEntry(OwnerId, Context.EventTime.AsSeconds(Cycle), Entry);
 			break;
 		}
diff --git a/Engine/Source/Runtime/Engine/Private/VisualLogger/VisualLoggerTraceDevice.cpp b/Engine/Source/Runtime/Engine/Private/VisualLogger/VisualLoggerTraceDevice.cpp
--- a/Engine/Source/Runtime/Engine/Private/VisualLogger/VisualLoggerTraceDevice.cpp
+++ b/Engine/Source/Runtime/Engine/Private/VisualLogger/VisualLoggerTraceDevice.cpp
@@ -13,6 +13,7 @@ UE_TRACE_CHANNEL_DEFINE(VisualLoggerChannel);
 
 UE_TRACE_EVENT_BEGIN(VisualLogger, VisualLogEntry)
 	UE_TRACE_EVENT_FIELD(uint64, Cycle)
+	UE_TRACE_EVENT_FIELD(double, RecordingTime)
 	UE_TRACE_EVENT_FIELD(uint64, OwnerId)
 	UE_TRACE_EVENT_FIELD(uint8[], LogEntry)
 UE_TRACE_EVENT_END()
@@ -63,6 +64,7 @@ void FVisualLoggerTraceDevice::Serialize(const UObject* LogOwner, FName OwnerNam
 
 		UE_TRACE_LOG(VisualLogger, VisualLogEntry, VisualLoggerChannel)
 			<< VisualLogEntry.Cycle(FPlatformTime::Cycles64())
+			<< VisualLogEntry.RecordingTime(FObjectTrace::GetWorldElapsedTime(LogOwner->GetWorld()))
 			<< VisualLogEntry.OwnerId(FObjectTrace::GetObjectId(LogOwner))
 			<< VisualLogEntry.LogEntry(Archive.GetData(), Archive.Num());
 	}
