Fix static analysis warning in RewindDebuggerVLog

diff --git a/Engine/Plugins/Animation/GameplayInsights/Source/RewindDebuggerVLog/Private/VisualLogTrack.cpp b/Engine/Plugins/Animation/GameplayInsights/Source/RewindDebuggerVLog/Private/VisualLogTrack.cpp
--- a/Engine/Plugins/Animation/GameplayInsights/Source/RewindDebuggerVLog/Private/VisualLogTrack.cpp
+++ b/Engine/Plugins/Animation/GameplayInsights/Source/RewindDebuggerVLog/Private/VisualLogTrack.cpp
@@ -69,63 +69,62 @@ bool FVisualLogCategoryTrack::UpdateInternal()
 	double EndTime = TraceTimeRange.GetUpperBoundValue();
 	
 	const TraceServices::IAnalysisSession* AnalysisSession = RewindDebugger->GetAnalysisSession();
-	const FVisualLoggerProvider* VisLogProvider = AnalysisSession->ReadProvider<FVisualLoggerProvider>(FVisualLoggerProvider::ProviderName);
-	
-	if(EventUpdateRequested > 10 && VisLogProvider)
+	if (const FVisualLoggerProvider* VisLogProvider = AnalysisSession->ReadProvider<FVisualLoggerProvider>(FVisualLoggerProvider::ProviderName))
 	{
-		EventUpdateRequested = 0;
-		
-		EventData->Points.SetNum(0,false);
-		EventData->Windows.SetNum(0);
-
-		TraceServices::FAnalysisSessionReadScope SessionReadScope(*AnalysisSession);
-		
-		VisLogProvider->ReadVisualLogEntryTimeline(ObjectId, [this, StartTime, EndTime, VisLogProvider, AnalysisSession](const FVisualLoggerProvider::VisualLogEntryTimeline& InTimeline)
+		if(EventUpdateRequested > 10)
 		{
-			InTimeline.EnumerateEvents(StartTime, EndTime, [this, StartTime, EndTime, VisLogProvider, AnalysisSession](double InStartTime, double InEndTime, uint32 InDepth, const FVisualLogEntry& InMessage)
+			EventUpdateRequested = 0;
+			
+			EventData->Points.SetNum(0,false);
+			EventData->Windows.SetNum(0);
+
+			TraceServices::FAnalysisSessionReadScope SessionReadScope(*AnalysisSession);
+			
+			VisLogProvider->ReadVisualLogEntryTimeline(ObjectId, [this, StartTime, EndTime, VisLogProvider, AnalysisSession](const FVisualLoggerProvider::VisualLogEntryTimeline& InTimeline)
 			{
-				for(const FVisualLogShapeElement& Element : InMessage.ElementsToDraw)
+				InTimeline.EnumerateEvents(StartTime, EndTime, [this, StartTime, EndTime, VisLogProvider, AnalysisSession](double InStartTime, double InEndTime, uint32 InDepth, const FVisualLogEntry& InMessage)
 				{
-					EventData->Points.Add({InMessage.TimeStamp,FText::FromName(Element.Category), FText::FromString(Element.Description),Element.GetFColor()});
-				}
-				return TraceServices::EEventEnumerate::Continue;
+					for(const FVisualLogShapeElement& Element : InMessage.ElementsToDraw)
+					{
+						EventData->Points.Add({InMessage.TimeStamp,FText::FromName(Element.Category), FText::FromString(Element.Description),Element.GetFColor()});
+					}
+					return TraceServices::EEventEnumerate::Continue;
+				});
 			});
-		});
-	}
+		}
 
 
-	double CurrentScrubTime = IRewindDebugger::Instance()->CurrentTraceTime();
-	if (PreviousScrubTime != CurrentScrubTime)
-	{
-		PreviousScrubTime = CurrentScrubTime;
-		DetailsObject->VisualLogDetails.SetNum(0,false);
-		
-		const TraceServices::IFrameProvider& FramesProvider = TraceServices::ReadFrameProvider(*AnalysisSession);
-		TraceServices::FFrame MarkerFrame;
-		if(FramesProvider.GetFrameFromTime(ETraceFrameType::TraceFrameType_Game, CurrentScrubTime, MarkerFrame))
+		double CurrentScrubTime = IRewindDebugger::Instance()->CurrentTraceTime();
+		if (PreviousScrubTime != CurrentScrubTime)
 		{
-			VisLogProvider->ReadVisualLogEntryTimeline(ObjectId, [this, &MarkerFrame, EndTime, VisLogProvider, AnalysisSession](const FVisualLoggerProvider::VisualLogEntryTimeline& InTimeline)
+			PreviousScrubTime = CurrentScrubTime;
+			DetailsObject->VisualLogDetails.SetNum(0,false);
+			
+			const TraceServices::IFrameProvider& FramesProvider = TraceServices::ReadFrameProvider(*AnalysisSession);
+			TraceServices::FFrame MarkerFrame;
+			if(FramesProvider.GetFrameFromTime(ETraceFrameType::TraceFrameType_Game, CurrentScrubTime, MarkerFrame))
 			{
-				InTimeline.EnumerateEvents(MarkerFrame.StartTime, MarkerFrame.EndTime, [this, VisLogProvider, AnalysisSession](double InStartTime, double InEndTime, uint32 InDepth, const FVisualLogEntry& InMessage)
+				VisLogProvider->ReadVisualLogEntryTimeline(ObjectId, [this, &MarkerFrame, EndTime, VisLogProvider, AnalysisSession](const FVisualLoggerProvider::VisualLogEntryTimeline& InTimeline)
 				{
-					for(const FVisualLogShapeElement& Element : InMessage.ElementsToDraw)
-					{
-						DetailsObject->VisualLogDetails.Add({Element.Category, Element.Description});
-					}
-					
-					for(const FVisualLogLine& Line : InMessage.LogLines)
+					InTimeline.EnumerateEvents(MarkerFrame.StartTime, MarkerFrame.EndTime, [this, VisLogProvider, AnalysisSession](double InStartTime, double InEndTime, uint32 InDepth, const FVisualLogEntry& InMessage)
 					{
-						DetailsObject->VisualLogDetails.Add({Line.Category, Line.Line});
-					}
-					return TraceServices::EEventEnumerate::Continue;
+						for(const FVisualLogShapeElement& Element : InMessage.ElementsToDraw)
+						{
+							DetailsObject->VisualLogDetails.Add({Element.Category, Element.Description});
+						}
+						
+						for(const FVisualLogLine& Line : InMessage.LogLines)
+						{
+							DetailsObject->VisualLogDetails.Add({Line.Category, Line.Line});
+						}
+						return TraceServices::EEventEnumerate::Continue;
+					});
 				});
-			});
+			}
 		}
 	}
 	
-	
 	bool bChanged = false;
-
 	return bChanged;
 	
 }
