Visual logger Rewind Debugger integration bug fixes

diff --git a/Engine/Plugins/Animation/GameplayInsights/Source/RewindDebuggerVLog/Private/VLogTraceModule.cpp b/Engine/Plugins/Animation/GameplayInsights/Source/RewindDebuggerVLog/Private/VLogTraceModule.cpp
--- a/Engine/Plugins/Animation/GameplayInsights/Source/RewindDebuggerVLog/Private/VLogTraceModule.cpp
+++ b/Engine/Plugins/Animation/GameplayInsights/Source/RewindDebuggerVLog/Private/VLogTraceModule.cpp
@@ -4,7 +4,7 @@
 #include "VisualLoggerProvider.h"
 #include "VisualLoggerAnalyzer.h"
 
-FName FVLogTraceModule::ModuleName("GameplayTrace");
+FName FVLogTraceModule::ModuleName("VisualLogger");
 
 void FVLogTraceModule::GetModuleInfo(TraceServices::FModuleInfo& OutModuleInfo)
 {
diff --git a/Engine/Plugins/Animation/GameplayInsights/Source/RewindDebuggerVLog/Private/VisualLogTrack.cpp b/Engine/Plugins/Animation/GameplayInsights/Source/RewindDebuggerVLog/Private/VisualLogTrack.cpp
--- a/Engine/Plugins/Animation/GameplayInsights/Source/RewindDebuggerVLog/Private/VisualLogTrack.cpp
+++ b/Engine/Plugins/Animation/GameplayInsights/Source/RewindDebuggerVLog/Private/VisualLogTrack.cpp
@@ -12,6 +12,15 @@
 namespace RewindDebugger
 {
 
+UVLogDetailsObject* FVisualLogCategoryTrack::InitializeDetailsObject()
+{
+	UVLogDetailsObject* DetailsObject = NewObject<UVLogDetailsObject>();
+	DetailsObject->SetFlags(RF_Standalone);
+	DetailsObjectWeakPtr = MakeWeakObjectPtr(DetailsObject);
+	DetailsView->SetObject(DetailsObject);
+	return DetailsObject;
+}
+
 FVisualLogCategoryTrack::FVisualLogCategoryTrack(uint64 InObjectId, const FName& InCategory) :
 	ObjectId(InObjectId),
 	Category(InCategory)
@@ -20,19 +29,21 @@ FVisualLogCategoryTrack::FVisualLogCategoryTrack(uint64 InObjectId, const FName&
 	EventData = MakeShared<SEventTimelineView::FTimelineEventData>();
 	Icon = FSlateIcon("EditorStyle", "Sequencer.Tracks.Event", "Sequencer.Tracks.Event");
 
-	DetailsObject = NewObject<UVLogDetailsObject>(GetTransientPackage(), UVLogDetailsObject::StaticClass(), NAME_None, RF_Transient|RF_Standalone);
 	
 	FPropertyEditorModule& PropertyEditorModule = FModuleManager::GetModuleChecked<FPropertyEditorModule>("PropertyEditor");
 	FDetailsViewArgs DetailsViewArgs;
 	DetailsViewArgs.NameAreaSettings = FDetailsViewArgs::HideNameArea;
 	DetailsView = PropertyEditorModule.CreateDetailView(DetailsViewArgs);
-	DetailsView->SetObject(DetailsObject);
-	
+
+	InitializeDetailsObject();
 }
 
 FVisualLogCategoryTrack::~FVisualLogCategoryTrack()
 {
-	DetailsObject->SetFlags(RF_Standalone);
+	if (UVLogDetailsObject* DetailsObject = DetailsObjectWeakPtr.Get())
+	{
+		DetailsObject->ClearFlags(RF_Standalone);
+	}
 }
 
 TSharedPtr<SEventTimelineView::FTimelineEventData> FVisualLogCategoryTrack::GetEventData() const
@@ -69,6 +80,7 @@ bool FVisualLogCategoryTrack::UpdateInternal()
 	double EndTime = TraceTimeRange.GetUpperBoundValue();
 	
 	const TraceServices::IAnalysisSession* AnalysisSession = RewindDebugger->GetAnalysisSession();
+	
 	if (const FVisualLoggerProvider* VisLogProvider = AnalysisSession->ReadProvider<FVisualLoggerProvider>(FVisualLoggerProvider::ProviderName))
 	{
 		if(EventUpdateRequested > 10)
@@ -98,15 +110,25 @@ bool FVisualLogCategoryTrack::UpdateInternal()
 		if (PreviousScrubTime != CurrentScrubTime)
 		{
 			PreviousScrubTime = CurrentScrubTime;
+
+			UVLogDetailsObject* DetailsObject = DetailsObjectWeakPtr.Get();
+			if (DetailsObject == nullptr)
+			{
+				// this should not happen unless the object was garbage collected (which should not happen since it's marked as Standalone)
+				DetailsObject = InitializeDetailsObject();
+			}
+
 			DetailsObject->VisualLogDetails.SetNum(0,false);
 			
 			const TraceServices::IFrameProvider& FramesProvider = TraceServices::ReadFrameProvider(*AnalysisSession);
+			TraceServices::FAnalysisSessionReadScope SessionReadScope(*AnalysisSession);
+			
 			TraceServices::FFrame MarkerFrame;
 			if(FramesProvider.GetFrameFromTime(ETraceFrameType::TraceFrameType_Game, CurrentScrubTime, MarkerFrame))
 			{
-				VisLogProvider->ReadVisualLogEntryTimeline(ObjectId, [this, &MarkerFrame, EndTime, VisLogProvider, AnalysisSession](const FVisualLoggerProvider::VisualLogEntryTimeline& InTimeline)
+				VisLogProvider->ReadVisualLogEntryTimeline(ObjectId, [this, DetailsObject, &MarkerFrame, EndTime, VisLogProvider, AnalysisSession](const FVisualLoggerProvider::VisualLogEntryTimeline& InTimeline)
 				{
-					InTimeline.EnumerateEvents(MarkerFrame.StartTime, MarkerFrame.EndTime, [this, VisLogProvider, AnalysisSession](double InStartTime, double InEndTime, uint32 InDepth, const FVisualLogEntry& InMessage)
+					InTimeline.EnumerateEvents(MarkerFrame.StartTime, MarkerFrame.EndTime, [this, DetailsObject, VisLogProvider, AnalysisSession](double InStartTime, double InEndTime, uint32 InDepth, const FVisualLogEntry& InMessage)
 					{
 						for(const FVisualLogShapeElement& Element : InMessage.ElementsToDraw)
 						{
diff --git a/Engine/Plugins/Animation/GameplayInsights/Source/RewindDebuggerVLog/Private/VisualLogTrack.h b/Engine/Plugins/Animation/GameplayInsights/Source/RewindDebuggerVLog/Private/VisualLogTrack.h
--- a/Engine/Plugins/Animation/GameplayInsights/Source/RewindDebuggerVLog/Private/VisualLogTrack.h
+++ b/Engine/Plugins/Animation/GameplayInsights/Source/RewindDebuggerVLog/Private/VisualLogTrack.h
@@ -7,6 +7,7 @@
 #include "SEventTimelineView.h"
 #include "IDetailsView.h"
 #include "VLogDetailsObject.h"
+#include "UObject/WeakObjectPtrTemplates.h"
 
 namespace RewindDebugger
 {
@@ -26,6 +27,8 @@ private:
 	virtual uint64 GetObjectIdInternal() const override { return ObjectId; }
 	virtual FName GetNameInternal() const override { return Category; }
 
+	UVLogDetailsObject* InitializeDetailsObject();
+
 	TSharedPtr<SEventTimelineView::FTimelineEventData> GetEventData() const;
 	
 	mutable TSharedPtr<SEventTimelineView::FTimelineEventData> EventData;
@@ -36,7 +39,7 @@ private:
 	FSlateIcon Icon;
 	FText TrackName;
 	TSharedPtr<IDetailsView> DetailsView;
-	UVLogDetailsObject* DetailsObject;
+	TWeakObjectPtr<UVLogDetailsObject> DetailsObjectWeakPtr;
 
 	double PreviousScrubTime = 0.0f;
 	
diff --git a/Engine/Plugins/Animation/GameplayInsights/Source/RewindDebuggerVLog/Private/VisualLoggerProvider.cpp b/Engine/Plugins/Animation/GameplayInsights/Source/RewindDebuggerVLog/Private/VisualLoggerProvider.cpp
--- a/Engine/Plugins/Animation/GameplayInsights/Source/RewindDebuggerVLog/Private/VisualLoggerProvider.cpp
+++ b/Engine/Plugins/Animation/GameplayInsights/Source/RewindDebuggerVLog/Private/VisualLoggerProvider.cpp
@@ -30,6 +30,8 @@ bool FVisualLoggerProvider::ReadVisualLogEntryTimeline(uint64 InObjectId, TFunct
 
 void FVisualLoggerProvider::EnumerateCategories(TFunctionRef<void(const FName&)> Callback) const
 {
+	Session.ReadAccessCheck();
+	
 	for(const FName& Category : Categories)
 	{
 		Callback(Category);
