rewind debugger optimizations - fix bugs causing treeview to refresh repeatedly, make curves tracks not update unless they are expanded

diff --git a/Engine/Plugins/Animation/GameplayInsights/Source/GameplayInsights/Private/AnimCurvesTrack.cpp b/Engine/Plugins/Animation/GameplayInsights/Source/GameplayInsights/Private/AnimCurvesTrack.cpp
--- a/Engine/Plugins/Animation/GameplayInsights/Source/GameplayInsights/Private/AnimCurvesTrack.cpp
+++ b/Engine/Plugins/Animation/GameplayInsights/Source/GameplayInsights/Private/AnimCurvesTrack.cpp
@@ -14,6 +14,7 @@ namespace RewindDebugger
 
 FAnimCurvesTrack::FAnimCurvesTrack(uint64 InObjectId) : ObjectId(InObjectId)
 {
+	ChildPlaceholder = MakeShared<FRewindDebuggerPlaceholderTrack>( "Child", LOCTEXT("No Curves", "No Curves to Display"));
 	SetIsExpanded(false);
 	Icon = FSlateIcon("EditorStyle", "AnimGraph.Attribute.Curves.Icon", "AnimGraph.Attribute.Curves.Icon");
 }
@@ -28,6 +29,11 @@ TSharedPtr<SWidget> FAnimCurvesTrack::GetDetailsViewInternal()
 
 void FAnimCurvesTrack::IterateSubTracksInternal(TFunction<void(TSharedPtr<FRewindDebuggerTrack> SubTrack)> IteratorFunction)
 {
+	if (Children.Num() == 0)
+	{
+		IteratorFunction(ChildPlaceholder);
+	}
+	
 	for(TSharedPtr<FAnimCurveTrack>& Track : Children)
 	{
 		IteratorFunction(Track);
@@ -36,6 +42,11 @@ void FAnimCurvesTrack::IterateSubTracksInternal(TFunction<void(TSharedPtr<FRewin
 
 bool FAnimCurvesTrack::UpdateInternal()
 {
+	if (!GetIsExpanded())
+	{
+		return false;
+	}
+	
 	TArray<uint32> UniqueTrackIds;
 
 	IRewindDebugger* RewindDebugger = IRewindDebugger::Instance();
@@ -89,7 +100,10 @@ bool FAnimCurvesTrack::UpdateInternal()
 				bChanged = true;
 			}
 
-			bChanged = bChanged || Children[i]->Update();
+			if (Children[i]->Update())
+			{
+				bChanged = true;
+			}
 		}
 	}
 
diff --git a/Engine/Plugins/Animation/GameplayInsights/Source/GameplayInsights/Private/AnimCurvesTrack.h b/Engine/Plugins/Animation/GameplayInsights/Source/GameplayInsights/Private/AnimCurvesTrack.h
--- a/Engine/Plugins/Animation/GameplayInsights/Source/GameplayInsights/Private/AnimCurvesTrack.h
+++ b/Engine/Plugins/Animation/GameplayInsights/Source/GameplayInsights/Private/AnimCurvesTrack.h
@@ -1,9 +1,12 @@
 // Copyright Epic Games, Inc. All Rights Reserved.
 #pragma once
 
+#include <RewindDebuggerPlaceholderTrack.h>
+
 #include "IAnimationProvider.h"
 #include "IRewindDebuggerTrackCreator.h"
 #include "RewindDebuggerTrack.h"
+#include "RewindDebuggerPlaceholderTrack.h"
 #include "SCurveTimelineView.h"
 #include "Textures/SlateIcon.h"
 
@@ -60,6 +63,7 @@ private:
 	FSlateIcon Icon;
 	uint64 ObjectId;
 
+	TSharedPtr<FRewindDebuggerPlaceholderTrack> ChildPlaceholder;
 	TArray<TSharedPtr<FAnimCurveTrack>> Children;
 };
 
diff --git a/Engine/Plugins/Animation/GameplayInsights/Source/GameplayInsights/Private/BlendWeightsTrack.cpp b/Engine/Plugins/Animation/GameplayInsights/Source/GameplayInsights/Private/BlendWeightsTrack.cpp
--- a/Engine/Plugins/Animation/GameplayInsights/Source/GameplayInsights/Private/BlendWeightsTrack.cpp
+++ b/Engine/Plugins/Animation/GameplayInsights/Source/GameplayInsights/Private/BlendWeightsTrack.cpp
@@ -101,7 +101,10 @@ bool FBlendWeightsTrack::UpdateInternal()
 				bChanged = true;
 			}
 
-			bChanged = bChanged || Children[i]->Update();
+			if (Children[i]->Update())
+			{
+				bChanged = true;
+			}
 		}
 	}
 
diff --git a/Engine/Plugins/Animation/GameplayInsights/Source/GameplayInsights/Private/MontagesTrack.cpp b/Engine/Plugins/Animation/GameplayInsights/Source/GameplayInsights/Private/MontagesTrack.cpp
--- a/Engine/Plugins/Animation/GameplayInsights/Source/GameplayInsights/Private/MontagesTrack.cpp
+++ b/Engine/Plugins/Animation/GameplayInsights/Source/GameplayInsights/Private/MontagesTrack.cpp
@@ -87,7 +87,10 @@ bool FMontagesTrack::UpdateInternal()
 				bChanged = true;
 			}
 
-			bChanged = bChanged || Children[i]->Update();
+			if (Children[i]->Update())
+			{
+				bChanged = true;
+			}
 		}
 	}
 
diff --git a/Engine/Plugins/Animation/GameplayInsights/Source/GameplayInsights/Private/NotifiesTrack.cpp b/Engine/Plugins/Animation/GameplayInsights/Source/GameplayInsights/Private/NotifiesTrack.cpp
--- a/Engine/Plugins/Animation/GameplayInsights/Source/GameplayInsights/Private/NotifiesTrack.cpp
+++ b/Engine/Plugins/Animation/GameplayInsights/Source/GameplayInsights/Private/NotifiesTrack.cpp
@@ -238,7 +238,10 @@ bool FNotifiesTrack::UpdateInternal()
 				bChanged = true;
 			}
 
-			bChanged = bChanged || Children[i]->Update();
+			if (Children[i]->Update())
+			{
+				bChanged = true;
+			}
 		}
 	}
 	
diff --git a/Engine/Plugins/Animation/GameplayInsights/Source/GameplayInsights/Private/PoseWatchTrack.cpp b/Engine/Plugins/Animation/GameplayInsights/Source/GameplayInsights/Private/PoseWatchTrack.cpp
--- a/Engine/Plugins/Animation/GameplayInsights/Source/GameplayInsights/Private/PoseWatchTrack.cpp
+++ b/Engine/Plugins/Animation/GameplayInsights/Source/GameplayInsights/Private/PoseWatchTrack.cpp
@@ -204,8 +204,11 @@ bool FPoseWatchTrack::UpdateInternal()
 				Children[i] = MakeShared<FPoseWatchCurveTrack>(ObjectId, UniqueCurveIds[i], PoseWatchTrackId);
 				bChanged = true;
 			}
-
-			bChanged = bChanged || Children[i]->Update();
+			
+			if (Children[i]->Update())
+			{
+				bChanged = true;
+			}
 		}
 
 		if(const TCHAR* FoundName = AnimationProvider->GetName(NameId))
@@ -323,7 +326,10 @@ bool FPoseWatchesTrack::UpdateInternal()
 				bChanged = true;
 			}
 
-			bChanged = bChanged || Children[i]->Update();
+			if (Children[i]->Update())
+			{
+				bChanged = true;
+			}
 		}
 	}
 	
diff --git a/Engine/Plugins/Animation/GameplayInsights/Source/GameplayInsights/Public/RewindDebuggerPlaceholderTrack.h b/Engine/Plugins/Animation/GameplayInsights/Source/GameplayInsights/Public/RewindDebuggerPlaceholderTrack.h
new file mode 100644
--- /dev/null
+++ b/Engine/Plugins/Animation/GameplayInsights/Source/GameplayInsights/Public/RewindDebuggerPlaceholderTrack.h
@@ -0,0 +1,22 @@
+// Copyright Epic Games, Inc. All Rights Reserved.
+
+#pragma once
+
+#include "CoreMinimal.h"
+#include "RewindDebuggerTrack.h"
+
+class GAMEPLAYINSIGHTS_API FRewindDebuggerPlaceholderTrack : public RewindDebugger::FRewindDebuggerTrack
+{
+public:
+	FRewindDebuggerPlaceholderTrack(const FName& InObjectName, const FText& InDisplayName)
+		: ObjectName(InObjectName), DisplayName(InDisplayName)
+	{
+	}
+
+private:
+	virtual FName GetNameInternal() const override { return ObjectName; }
+	virtual FText GetDisplayNameInternal() const override { return DisplayName; }
+	
+	FName ObjectName;
+	FText DisplayName;
+};
diff --git a/Engine/Plugins/Animation/GameplayInsights/Source/RewindDebugger/Private/RewindDebugger.cpp b/Engine/Plugins/Animation/GameplayInsights/Source/RewindDebugger/Private/RewindDebugger.cpp
--- a/Engine/Plugins/Animation/GameplayInsights/Source/RewindDebugger/Private/RewindDebugger.cpp
+++ b/Engine/Plugins/Animation/GameplayInsights/Source/RewindDebugger/Private/RewindDebugger.cpp
@@ -24,6 +24,7 @@
 #include "RewindDebuggerCommands.h"
 #include "RewindDebuggerModule.h"
 #include "RewindDebuggerObjectTrack.h"
+#include "RewindDebuggerPlaceholderTrack.h"
 #include "RewindDebuggerSettings.h"
 #include "SLevelViewport.h"
 #include "ToolMenus.h"
@@ -294,22 +295,6 @@ void FRewindDebugger::GetTargetObjectIds(TArray<uint64>& OutTargetObjectIds) con
 #endif
 }
 
-class FRewindDebuggerDummyTrack : public RewindDebugger::FRewindDebuggerTrack
-{
-public:
-
-	FRewindDebuggerDummyTrack(const FName& InObjectName, const FText& InDisplayName)
-		: ObjectName(InObjectName), DisplayName(InDisplayName)
-	{
-	}
-
-private:
-	virtual FName GetNameInternal() const override { return ObjectName; }
-	virtual FText GetDisplayNameInternal() const override { return DisplayName; }
-	
-	FName ObjectName;
-	FText DisplayName;
-};
 
 void FRewindDebugger::RefreshDebugTracks()
 {
@@ -336,8 +321,8 @@ void FRewindDebugger::RefreshDebugTracks()
 
 		if (DebugTracks[1] == nullptr || DebugTracks[0] == nullptr || DebugTracks[0]->GetName().ToString() != DebugTargetActor.Get() )
 		{
-			DebugTracks[0] = MakeShared<FRewindDebuggerDummyTrack>(FName(DebugTargetActorName), FText::FromString(DebugTargetActorName)); 
-			DebugTracks[1] = MakeShared<FRewindDebuggerDummyTrack>(DebugMessageTrackName, NSLOCTEXT("RewindDebugger", "No Debug Data", " - Start a recording to debug"));
+			DebugTracks[0] = MakeShared<FRewindDebuggerPlaceholderTrack>(FName(DebugTargetActorName), FText::FromString(DebugTargetActorName)); 
+			DebugTracks[1] = MakeShared<FRewindDebuggerPlaceholderTrack>(DebugMessageTrackName, NSLOCTEXT("RewindDebugger", "No Debug Data", " - Start a recording to debug"));
 			ComponentListChangedDelegate.ExecuteIfBound();
 		}
 	}
@@ -379,7 +364,10 @@ void FRewindDebugger::RefreshDebugTracks()
     			// update all tracks
     			for (TSharedPtr<RewindDebugger::FRewindDebuggerTrack>& DebugTrack : DebugTracks )
     			{
-    				bChanged = bChanged || DebugTrack->Update();
+    				if (DebugTrack->Update())
+    				{
+    					bChanged = true;
+    				}
     			}
     
     			if (bChanged)
diff --git a/Engine/Plugins/Animation/GameplayInsights/Source/RewindDebuggerVLog/Private/VisualLogTrack.cpp b/Engine/Plugins/Animation/GameplayInsights/Source/RewindDebuggerVLog/Private/VisualLogTrack.cpp
--- a/Engine/Plugins/Animation/GameplayInsights/Source/RewindDebuggerVLog/Private/VisualLogTrack.cpp
+++ b/Engine/Plugins/Animation/GameplayInsights/Source/RewindDebuggerVLog/Private/VisualLogTrack.cpp
@@ -235,7 +235,10 @@ bool FVisualLogTrack::UpdateInternal()
 				bChanged = true;
 			}
 	
-			bChanged = bChanged || Children[i]->Update();
+			if (Children[i]->Update())
+			{
+				bChanged = true;
+			}
 		}
 	}
 	
diff --git a/Engine/Plugins/Experimental/Chooser/Source/ChooserEditor/Private/ChooserTrack.cpp b/Engine/Plugins/Experimental/Chooser/Source/ChooserEditor/Private/ChooserTrack.cpp
--- a/Engine/Plugins/Experimental/Chooser/Source/ChooserEditor/Private/ChooserTrack.cpp
+++ b/Engine/Plugins/Experimental/Chooser/Source/ChooserEditor/Private/ChooserTrack.cpp
@@ -206,7 +206,10 @@ bool FChoosersTrack::UpdateInternal()
 				bChanged = true;
 			}
 		
-			bChanged = bChanged || Children[i]->Update();
+			if (Children[i]->Update())
+			{
+				bChanged = true;
+			}
 		}
 	}
 	
