Improve handling of Child AnimationBluprints in Rewind Debugger

diff --git a/Engine/Plugins/Animation/GameplayInsights/Source/RewindDebugger/Private/RewindDebugger.cpp b/Engine/Plugins/Animation/GameplayInsights/Source/RewindDebugger/Private/RewindDebugger.cpp
--- a/Engine/Plugins/Animation/GameplayInsights/Source/RewindDebugger/Private/RewindDebugger.cpp
+++ b/Engine/Plugins/Animation/GameplayInsights/Source/RewindDebugger/Private/RewindDebugger.cpp
@@ -1001,6 +1001,16 @@ void FRewindDebugger::Tick(float DeltaTime)
 											{
 												if(UAnimBlueprint* AnimBlueprint = Cast<UAnimBlueprint>(InstanceClass->ClassGeneratedBy))
 												{
+													// for child Animation Blueprints, we actually want to debug the root blueprint (since the child doesn't contain any anim graphs)
+													if (UAnimBlueprint* RootAnimBP = UAnimBlueprint::FindRootAnimBlueprint(AnimBlueprint))
+													{
+														if (UAnimBlueprintGeneratedClass* RootInstanceClass = Cast<UAnimBlueprintGeneratedClass>(RootAnimBP->GeneratedClass))
+														{
+															AnimBlueprint = RootAnimBP;
+															InstanceClass = RootInstanceClass;
+														}
+													}
+
 													if(AnimBlueprint->IsObjectBeingDebugged(AnimInstance))
 													{
 														TRACE_CPUPROFILER_EVENT_SCOPE(FRewindDebugger::Tick_UpdateBlueprintDebug);
diff --git a/Engine/Source/Editor/Kismet/Private/BlueprintEditor.cpp b/Engine/Source/Editor/Kismet/Private/BlueprintEditor.cpp
--- a/Engine/Source/Editor/Kismet/Private/BlueprintEditor.cpp
+++ b/Engine/Source/Editor/Kismet/Private/BlueprintEditor.cpp
@@ -2741,6 +2741,14 @@ FReply FBlueprintEditor::OnEditParentClassClicked()
 			if (ParentBlueprintGeneratedClass)
 			{
 				GEditor->GetEditorSubsystem<UAssetEditorSubsystem>()->OpenEditorForAsset(ParentBlueprintGeneratedClass->ClassGeneratedBy);
+				
+				if (UObject* DebugObject = Blueprint->GetObjectBeingDebugged())
+				{
+					if (UBlueprint* ParentBlueprint = Cast<UBlueprint>(ParentBlueprintGeneratedClass->ClassGeneratedBy))
+					{
+						ParentBlueprint->SetObjectBeingDebugged(DebugObject);
+					}
+				}
 			}
 		}
 	}
