Merge skipped changelist 29324573

diff --git a/Engine/Plugins/Animation/GameplayInsights/Source/GameplayInsights/Private/AnimCurvesTrack.cpp b/Engine/Plugins/Animation/GameplayInsights/Source/GameplayInsights/Private/AnimCurvesTrack.cpp
--- a/Engine/Plugins/Animation/GameplayInsights/Source/GameplayInsights/Private/AnimCurvesTrack.cpp
+++ b/Engine/Plugins/Animation/GameplayInsights/Source/GameplayInsights/Private/AnimCurvesTrack.cpp
@@ -42,6 +42,7 @@ void FAnimCurvesTrack::IterateSubTracksInternal(TFunction<void(TSharedPtr<FRewin
 
 bool FAnimCurvesTrack::UpdateInternal()
 {
+	TRACE_CPUPROFILER_EVENT_SCOPE(FAnimCurvesTrack::UpdateInternal);
 	if (!GetIsExpanded())
 	{
 		return false;
@@ -154,6 +155,8 @@ bool FAnimCurveTrack::UpdateInternal()
 
 void FAnimCurveTrack::UpdateCurvePointsInternal()
 {
+	TRACE_CPUPROFILER_EVENT_SCOPE(FAnimCurveTrack::UpdateCurvePointsInternal);
+	
 	IRewindDebugger* RewindDebugger = IRewindDebugger::Instance();
 	const TraceServices::IAnalysisSession* AnalysisSession = RewindDebugger->GetAnalysisSession();
 	const FAnimationProvider* AnimationProvider = AnalysisSession->ReadProvider<FAnimationProvider>(FAnimationProvider::ProviderName);
diff --git a/Engine/Plugins/Animation/GameplayInsights/Source/GameplayInsights/Private/AnimationAnalyzer.cpp b/Engine/Plugins/Animation/GameplayInsights/Source/GameplayInsights/Private/AnimationAnalyzer.cpp
--- a/Engine/Plugins/Animation/GameplayInsights/Source/GameplayInsights/Private/AnimationAnalyzer.cpp
+++ b/Engine/Plugins/Animation/GameplayInsights/Source/GameplayInsights/Private/AnimationAnalyzer.cpp
@@ -52,6 +52,7 @@ void FAnimationAnalyzer::OnAnalysisBegin(const FOnAnalysisContext& Context)
 	Builder.RouteEvent(RouteId_Sync, "Animation", "Sync");
 	Builder.RouteEvent(RouteId_PoseWatch, "Animation", "PoseWatch");
 	Builder.RouteEvent(RouteId_PoseWatch2, "Animation", "PoseWatch2");
+	Builder.RouteEvent(RouteId_Inertialization, "Animation", "Inertialization");
 }
 
 bool FAnimationAnalyzer::OnEvent(uint16 RouteId, EStyle Style, const FOnEventContext& Context)
@@ -494,6 +495,18 @@ bool FAnimationAnalyzer::OnEvent(uint16 RouteId, EStyle Style, const FOnEventCon
 		AnimationProvider.AppendPoseWatch(ComponentId, AnimInstanceId, Context.EventTime.AsSeconds(Cycle), RecordingTime, PoseWatchId, NameId, FColor(Color), BoneTransformsFloatArray, CurveIds, CurveValues, RequiredBonesIntArray, WorldTransformFloatArray, bIsEnabled);
 		break;
 	}
+	case RouteId_Inertialization:
+	{
+		uint64 Cycle = EventData.GetValue<uint64>("Cycle");
+		double RecordingTime = EventData.GetValue<double>("RecordingTime");
+		uint64 AnimInstanceId = EventData.GetValue<uint64>("AnimInstanceId");
+		int32 NodeId = EventData.GetValue<int32>("NodeId");
+		float Weight = EventData.GetValue<float>("Weight");
+		EInertializationType Type = static_cast<EInertializationType>(EventData.GetValue<uint8>("Type"));	
+			
+		AnimationProvider.AppendInertialization(AnimInstanceId, Context.EventTime.AsSeconds(Cycle), RecordingTime, NodeId, Weight, Type);
+		break;
+	}
 	}
 
 	return true;
diff --git a/Engine/Plugins/Animation/GameplayInsights/Source/GameplayInsights/Private/AnimationAnalyzer.h b/Engine/Plugins/Animation/GameplayInsights/Source/GameplayInsights/Private/AnimationAnalyzer.h
--- a/Engine/Plugins/Animation/GameplayInsights/Source/GameplayInsights/Private/AnimationAnalyzer.h
+++ b/Engine/Plugins/Animation/GameplayInsights/Source/GameplayInsights/Private/AnimationAnalyzer.h
@@ -54,6 +54,7 @@ private:
 		RouteId_Sync,
 		RouteId_PoseWatch,
 		RouteId_PoseWatch2,
+		RouteId_Inertialization,
 	};
 
 	TraceServices::IAnalysisSession& Session;
diff --git a/Engine/Plugins/Animation/GameplayInsights/Source/GameplayInsights/Private/AnimationProvider.cpp b/Engine/Plugins/Animation/GameplayInsights/Source/GameplayInsights/Private/AnimationProvider.cpp
--- a/Engine/Plugins/Animation/GameplayInsights/Source/GameplayInsights/Private/AnimationProvider.cpp
+++ b/Engine/Plugins/Animation/GameplayInsights/Source/GameplayInsights/Private/AnimationProvider.cpp
@@ -204,6 +204,40 @@ bool FAnimationProvider::ReadTickRecordTimeline(uint64 InObjectId, TFunctionRef<
 	return false;
 }
 
+bool FAnimationProvider::ReadInertializationTimeline(uint64 InObjectId, TFunctionRef<void(const InertializationTimeline&)> Callback) const
+{
+	Session.ReadAccessCheck();
+
+	const uint32* IndexPtr = ObjectIdToInertializationTimelines.Find(InObjectId);
+	if (IndexPtr != nullptr)
+	{
+		if (*IndexPtr < uint32(InertializationTimelineStorage.Num()))
+		{
+			Callback(*InertializationTimelineStorage[*IndexPtr]->Timeline.Get());
+			return true;
+		}
+	}
+
+	return false;
+}
+
+void FAnimationProvider::EnumerateInertializationNodes(uint64 InObjectId, TFunctionRef<void(int32 NodeId, EInertializationType Type)> Callback) const
+{
+	Session.ReadAccessCheck();
+
+	const uint32* IndexPtr = ObjectIdToInertializationTimelines.Find(InObjectId);
+	if (IndexPtr != nullptr)
+	{
+		if (*IndexPtr < uint32(InertializationTimelineStorage.Num()))
+		{
+			for(auto& Entry : InertializationTimelineStorage[*IndexPtr]->NodeTypes)
+			{
+				Callback(Entry.Key, Entry.Value);
+			}
+		}
+	}
+}
+
 void FAnimationProvider::EnumerateAnimGraphTimelines(TFunctionRef<void(uint64 ObjectId, const AnimGraphTimeline&)> Callback) const
 {
 	Session.ReadAccessCheck();
@@ -1693,4 +1727,34 @@ void FAnimationProvider::AppendPoseWatch(uint64 InComponentId, uint64 InAnimInst
 	Session.UpdateDurationSeconds(InTime);
 }
 
+void FAnimationProvider::AppendInertialization(uint64 InAnimInstanceId, double InProfileTime, double InRecordingTime, int32 InNodeId, float InWeight, EInertializationType InType)
+{
+	Session.WriteAccessCheck();
+	
+	TSharedPtr<FInertializationTimelineStorage> TimelineStorage;
+	uint32* IndexPtr = ObjectIdToInertializationTimelines.Find(InAnimInstanceId);
+	if(IndexPtr != nullptr)
+	{
+		TimelineStorage = InertializationTimelineStorage[*IndexPtr];
+	}
+	else
+	{
+		TimelineStorage = MakeShared<FInertializationTimelineStorage>();
+		TimelineStorage->Timeline = MakeShared<TraceServices::TPointTimeline<FInertializationMessage>>(Session.GetLinearAllocator());
+		TimelineStorage->Timeline->SetEnumerateOutsideRange(true);
+		ObjectIdToInertializationTimelines.Add(InAnimInstanceId, InertializationTimelineStorage.Num());
+		InertializationTimelineStorage.Add(TimelineStorage.ToSharedRef());
+	}
+	
+	
+	FInertializationMessage Message;
+	Message.AnimInstanceId = InAnimInstanceId;
+	Message.RecordingTime = InRecordingTime;
+	Message.NodeId = InNodeId;
+	Message.Weight = InWeight;
+	Message.Type = InType;
+	TimelineStorage->Timeline->AppendEvent(InProfileTime, Message);
+	TimelineStorage->NodeTypes.Add(InNodeId, InType);
+}
+
 #undef LOCTEXT_NAMESPACE
diff --git a/Engine/Plugins/Animation/GameplayInsights/Source/GameplayInsights/Private/AnimationProvider.h b/Engine/Plugins/Animation/GameplayInsights/Source/GameplayInsights/Private/AnimationProvider.h
--- a/Engine/Plugins/Animation/GameplayInsights/Source/GameplayInsights/Private/AnimationProvider.h
+++ b/Engine/Plugins/Animation/GameplayInsights/Source/GameplayInsights/Private/AnimationProvider.h
@@ -32,6 +32,8 @@ public:
 	virtual void EnumeratePoseWatchCurves(const FPoseWatchMessage& InMessage, TFunctionRef<void(const FSkeletalMeshNamedCurve&)> Callback) const override;
 	virtual void EnumerateExternalMorphSets(const FSkeletalMeshPoseMessage& InMessage, TFunctionRef<void(const FExternalMorphWeightMessage&)> Callback) const override;
 	virtual bool ReadTickRecordTimeline(uint64 InObjectId, TFunctionRef<void(const TickRecordTimeline&)> Callback) const override;
+	virtual bool ReadInertializationTimeline(uint64 InObjectId, TFunctionRef<void(const InertializationTimeline&)> Callback) const override;
+	virtual void EnumerateInertializationNodes(uint64 InObjectId, TFunctionRef<void(int32 NodeId, EInertializationType Type)> Callback) const override;
 	virtual void EnumerateTickRecordIds(uint64 InObjectId, TFunctionRef<void(uint64, int32)> Callback) const override;
 	virtual void EnumerateAnimGraphTimelines(TFunctionRef<void(uint64 ObjectId, const AnimGraphTimeline&)> Callback) const override;
 	virtual bool ReadAnimGraphTimeline(uint64 InObjectId, TFunctionRef<void(const AnimGraphTimeline&)> Callback) const override;
@@ -121,6 +123,9 @@ public:
 	void AppendPoseWatch(uint64 InAnimInstanceId, double InTime, double InRecordingTime, uint64 PoseWatchId, const TArrayView<const float>& BoneTransformsRaw, const TArrayView<const uint16>& RequiredBones, const TArrayView<const float>& WorldTransformRaw, const bool bIsEnabled);
 	void AppendPoseWatch(uint64 InComponentId, uint64 InAnimInstanceId, double InTime, double InRecordingTime, uint64 PoseWatchId, uint32 NameId, FColor Color, const TArrayView<const float>& BoneTransformsRaw, const TArrayView<const uint32>& CurveIds, const TArrayView<const float>& CurveValues, const TArrayView<const uint16>& RequiredBones, const TArrayView<const float>& WorldTransformRaw, const bool bIsEnabled);
 
+	/** Append Inertialization data */
+	void AppendInertialization(uint64 InAnimInstanceId, double InProfileTime, double InRecordingTime, int32 NodeId, float InWeight, EInertializationType InType);
+
 private:
 	/** Add anim node values helper */
 	void AppendAnimNodeValue(uint64 InAnimInstanceId, double InTime, double InRecordingTime, uint16 InFrameCounter, int32 InNodeId, const TCHAR* InKey, FAnimNodeValueMessage& InMessage);
@@ -145,6 +150,7 @@ private:
 	TMap<uint64, uint32> ObjectIdToAnimNotifyStateTimelines;
 	TMap<uint64, uint32> ObjectIdToAnimNotifyTimelines;
 	TMap<uint64, uint32> ObjectIdToAnimMontageTimelines;
+	TMap<uint64, uint32> ObjectIdToInertializationTimelines;
 	TMap<uint64, uint32> ObjectIdToAnimAttributeTimelines;
 	TMap<uint64, uint32> ObjectIdToAnimSyncTimelines;
 	TMap<uint64, uint32> ObjectIdToPoseWatchTimelines;
@@ -188,6 +194,12 @@ private:
 		TSharedPtr<TraceServices::TPointTimeline<FAnimMontageMessage>> Timeline;
 		TSet<uint64> AllMontageIds;
 	};
+	
+	struct FInertializationTimelineStorage
+	{
+		TSharedPtr<TraceServices::TPointTimeline<FInertializationMessage>> Timeline;
+		TMap<int32,EInertializationType> NodeTypes;
+	};
 
 	/** Message storage */
 	TArray<TSharedRef<FTickRecordTimelineStorage>> TickRecordTimelineStorage;
@@ -202,6 +214,7 @@ private:
 	TArray<TSharedRef<FAnimNotifyStateTimelineStorage>> AnimNotifyStateTimelineStorage;
 	TArray<TSharedRef<TraceServices::TPointTimeline<FAnimNotifyMessage>>> AnimNotifyTimelines;
 	TArray<TSharedRef<FMontageTimelineStorage>> AnimMontageTimelineStorage;
+	TArray<TSharedRef<FInertializationTimelineStorage>> InertializationTimelineStorage;
 	TraceServices::TPagedArray<FTransform> SkeletalMeshPoseTransforms;
 	TraceServices::TPagedArray<FSkeletalMeshNamedCurve> SkeletalMeshCurves;
 	TraceServices::TPagedArray<FExternalMorphWeightMessage> ExternalMorphWeights;
diff --git a/Engine/Plugins/Animation/GameplayInsights/Source/GameplayInsights/Private/BlendWeightsTrack.cpp b/Engine/Plugins/Animation/GameplayInsights/Source/GameplayInsights/Private/BlendWeightsTrack.cpp
--- a/Engine/Plugins/Animation/GameplayInsights/Source/GameplayInsights/Private/BlendWeightsTrack.cpp
+++ b/Engine/Plugins/Animation/GameplayInsights/Source/GameplayInsights/Private/BlendWeightsTrack.cpp
@@ -38,6 +38,7 @@ void FBlendWeightsTrack::IterateSubTracksInternal(TFunction<void(TSharedPtr<FRew
 
 bool FBlendWeightsTrack::UpdateInternal()
 {
+	TRACE_CPUPROFILER_EVENT_SCOPE(FBlendWeights::UpdateInternal);
 	IRewindDebugger* RewindDebugger = IRewindDebugger::Instance();
 
 	struct TrackId
@@ -166,6 +167,7 @@ bool FBlendWeightTrack::UpdateInternal()
 
 	if(CurvesUpdateRequested > 10 && GameplayProvider && AnimationProvider)
 	{
+		TRACE_CPUPROFILER_EVENT_SCOPE(FBlendWeightTrack::UpdateCurvePointsInternal);
 		auto& CurvePoints = CurveData->Points;
 		CurvePoints.SetNum(0,false);
 		
@@ -342,6 +344,7 @@ TSharedPtr<SWidget> FBlendWeightTrack::GetDetailsViewInternal()
 
 bool FBlendWeightTrack::HandleDoubleClickInternal()
 {
+	TRACE_CPUPROFILER_EVENT_SCOPE(FBlendWeightsTrack::UpdateInternal);
 #if WITH_EDITOR
 	IRewindDebugger* RewindDebugger = IRewindDebugger::Instance();
 	if (const TraceServices::IAnalysisSession* AnalysisSession = RewindDebugger->GetAnalysisSession())
@@ -379,6 +382,7 @@ TSharedPtr<RewindDebugger::FRewindDebuggerTrack> FBlendWeightsTrackCreator::Crea
 
 bool FBlendWeightsTrackCreator::HasDebugInfoInternal(uint64 ObjectId) const
 {
+	TRACE_CPUPROFILER_EVENT_SCOPE(FBlendWeightsTrack::HasDebugInfoInternal);
 	const TraceServices::IAnalysisSession* AnalysisSession = IRewindDebugger::Instance()->GetAnalysisSession();
 	TraceServices::FAnalysisSessionReadScope SessionReadScope(*AnalysisSession);
 	bool bHasData = false;
diff --git a/Engine/Plugins/Animation/GameplayInsights/Source/GameplayInsights/Private/GameplayProvider.cpp b/Engine/Plugins/Animation/GameplayInsights/Source/GameplayInsights/Private/GameplayProvider.cpp
--- a/Engine/Plugins/Animation/GameplayInsights/Source/GameplayInsights/Private/GameplayProvider.cpp
+++ b/Engine/Plugins/Animation/GameplayInsights/Source/GameplayInsights/Private/GameplayProvider.cpp
@@ -593,6 +593,7 @@ TRange<double> FGameplayProvider::GetObjectTraceLifetime(uint64 ObjectId) const
 
 TRange<double> FGameplayProvider::GetObjectRecordingLifetime(uint64 ObjectId) const
 {
+	TRACE_CPUPROFILER_EVENT_SCOPE(FGameplayProvider::GetObjectRecordingLifetime);
 	Session.ReadAccessCheck();
 
 	if (const uint64 *FoundIndex = ActiveObjectRecordingLifetimes.Find(ObjectId))
diff --git a/Engine/Plugins/Animation/GameplayInsights/Source/GameplayInsights/Private/InertializationsTrack.cpp b/Engine/Plugins/Animation/GameplayInsights/Source/GameplayInsights/Private/InertializationsTrack.cpp
--- a/Engine/Plugins/Animation/GameplayInsights/Source/GameplayInsights/Private/InertializationsTrack.cpp
+++ b/Engine/Plugins/Animation/GameplayInsights/Source/GameplayInsights/Private/InertializationsTrack.cpp
@@ -36,10 +36,9 @@ void FInertializationsTrack::IterateSubTracksInternal(TFunction<void(TSharedPtr<
 
 bool FInertializationsTrack::UpdateInternal()
 {
+	TRACE_CPUPROFILER_EVENT_SCOPE(FInertializationsTrack::UpdateInternal);
 	IRewindDebugger* RewindDebugger = IRewindDebugger::Instance();
 
-	TSortedMap<int32, const TCHAR*, TInlineAllocator<8>> NodeNameMap;
-
 	const TraceServices::IAnalysisSession* AnalysisSession = RewindDebugger->GetAnalysisSession();
 
 	bool bChanged = false;
@@ -57,45 +56,33 @@ bool FInertializationsTrack::UpdateInternal()
 	{
 		TraceServices::FAnalysisSessionReadScope SessionReadScope(*AnalysisSession);
 
-		AnimationProvider->ReadAnimNodesTimeline(ObjectId, [&NodeNameMap, &AnimationProvider, StartTime, EndTime](const FAnimationProvider::AnimNodesTimeline& InTimeline)
+		TArray<TPair<int32,EInertializationType>, TInlineAllocator<8>> Nodes;
 		{
-			// this isn't very efficient, and it gets called every frame.  will need optimizing
-			InTimeline.EnumerateEvents(StartTime, EndTime, [&NodeNameMap, &AnimationProvider, StartTime, EndTime](double InStartTime, double InEndTime, uint32 InDepth, const FAnimNodeMessage& InMessage)
+			AnimationProvider->EnumerateInertializationNodes(ObjectId,  [&Nodes] (uint32 NodeId, EInertializationType Type)
 			{
-				if (InEndTime > StartTime && InStartTime < EndTime)
-				{
-					for (const TCHAR* InertializationNodeType : { TEXT("AnimNode_DeadBlending"), TEXT("AnimNode_Inertialization") })
-					{
-						if (FCString::Strcmp(InertializationNodeType, InMessage.NodeTypeName) == 0)
-						{
-							NodeNameMap.Add(InMessage.NodeId, InMessage.NodeName);
-							break;
-						}
-					}
-				}
-
-				return TraceServices::EEventEnumerate::Continue;
+				Nodes.Add({ NodeId, Type });
 			});
-		});
-
-		TArray<int32, TInlineAllocator<8>> NodeIds;
-		NodeNameMap.GetKeys(NodeIds);
+		}
 
-		if (Children.Num() != NodeIds.Num())
+		if (Children.Num() != Nodes.Num())
 		{
 			bChanged = true;
 		}
 
-		Children.SetNum(NodeIds.Num());
-		for(int32 NodeIdx = 0; NodeIdx < NodeIds.Num(); NodeIdx++)
+		Children.SetNum(Nodes.Num());
+		for(int32 NodeIdx = 0; NodeIdx < Nodes.Num(); NodeIdx++)
 		{
-			if (!Children[NodeIdx].IsValid())
+			if (!Children[NodeIdx].IsValid() || Children[NodeIdx]->GetNodeId() != Nodes[NodeIdx].Key)
 			{
-				Children[NodeIdx] = MakeShared<FInertializationTrack>(ObjectId, NodeIds[NodeIdx], FText::FromString(NodeNameMap[NodeIds[NodeIdx]]));
+				Children[NodeIdx] = MakeShared<FInertializationTrack>(ObjectId, Nodes[NodeIdx].Key,
+					Nodes[NodeIdx].Value == EInertializationType::DeadBlending ? LOCTEXT("DeadBlending","DeadBlending") : LOCTEXT("Inertialization","Inertialization"));
 				bChanged = true;
 			}
 
-			bChanged = bChanged || Children[NodeIdx]->Update();
+			if (Children[NodeIdx]->Update())
+			{
+				bChanged = true;
+			}
 		}
 	}
 
@@ -151,19 +138,18 @@ bool FInertializationTrack::UpdateInternal()
 		
 		TraceServices::FAnalysisSessionReadScope SessionReadScope(*AnalysisSession);
 		
-		AnimationProvider->ReadAnimNodeValuesTimeline(ObjectId, [AnalysisSession, StartTime, EndTime, &CurvePoints, this](const FAnimationProvider::AnimNodeValuesTimeline& InTimeline)
+		AnimationProvider->ReadInertializationTimeline(ObjectId, [AnalysisSession, StartTime, EndTime, &CurvePoints, this](const IAnimationProvider::InertializationTimeline& InTimeline)
 		{
-			InTimeline.EnumerateEvents(StartTime, EndTime, [&CurvePoints, AnalysisSession, this](double InStartTime, double InEndTime, uint32 InDepth, const FAnimNodeValueMessage& InMessage)
+			InTimeline.EnumerateEvents(StartTime, EndTime, [&CurvePoints, AnalysisSession, this](double InStartTime, double InEndTime, uint32 InDepth, const FInertializationMessage& InMessage)
 			{
-				if (InMessage.NodeId == NodeId && FCString::Strcmp(InMessage.Key, TEXT("Inertialization Weight")) == 0)
+				if (InMessage.NodeId == NodeId)
 				{
-					CurvePoints.Add({ InMessage.RecordingTime,	InMessage.Value.Float.Value });
+					CurvePoints.Add({ InMessage.RecordingTime,InMessage.Weight });
 				}
-				
 				return TraceServices::EEventEnumerate::Continue;
 			});
 		});
-
+		
 		CurvesUpdateRequested = 0;
 	}
 
@@ -275,32 +261,19 @@ TSharedPtr<RewindDebugger::FRewindDebuggerTrack> FInertializationsTrackCreator::
 
 bool FInertializationsTrackCreator::HasDebugInfoInternal(uint64 ObjectId) const
 {
+	TRACE_CPUPROFILER_EVENT_SCOPE(FInertializationsTrack::HasDebugInfoInternal);
 	const TraceServices::IAnalysisSession* AnalysisSession = IRewindDebugger::Instance()->GetAnalysisSession();
 	TraceServices::FAnalysisSessionReadScope SessionReadScope(*AnalysisSession);
 	
-	bool bHasData = false;
-
 	if (const FAnimationProvider* AnimationProvider = AnalysisSession->ReadProvider<FAnimationProvider>(FAnimationProvider::ProviderName))
 	{
-		AnimationProvider->ReadAnimNodesTimeline(ObjectId, [&AnimationProvider, &bHasData](const FAnimationProvider::AnimNodesTimeline& InTimeline)
+		if (AnimationProvider->ReadInertializationTimeline(ObjectId, [](const IAnimationProvider::InertializationTimeline& InTimeline){}))
 		{
-			InTimeline.EnumerateEvents(InTimeline.GetStartTime(), InTimeline.GetEndTime(), [&AnimationProvider, &bHasData](double InStartTime, double InEndTime, uint32 InDepth, const FAnimNodeMessage& InMessage)
-			{
-				for (const TCHAR* InertializationNodeType : { TEXT("AnimNode_DeadBlending"), TEXT("AnimNode_Inertialization") })
-				{
-					if (FCString::Strcmp(InertializationNodeType, InMessage.NodeTypeName) == 0)
-					{
-						bHasData = true;
-						return TraceServices::EEventEnumerate::Stop;
-					}
-				}
-
-				return TraceServices::EEventEnumerate::Continue;
-			});
-		});
+			return true;
+		}
 	}
-	
-	return bHasData;
+
+	return false;
 }
 
 }
diff --git a/Engine/Plugins/Animation/GameplayInsights/Source/GameplayInsights/Private/InertializationsTrack.h b/Engine/Plugins/Animation/GameplayInsights/Source/GameplayInsights/Private/InertializationsTrack.h
--- a/Engine/Plugins/Animation/GameplayInsights/Source/GameplayInsights/Private/InertializationsTrack.h
+++ b/Engine/Plugins/Animation/GameplayInsights/Source/GameplayInsights/Private/InertializationsTrack.h
@@ -16,6 +16,7 @@ class FInertializationTrack : public FRewindDebuggerTrack
 public:
 	
 	FInertializationTrack(uint64 InObjectId, int32 NodeId, const FText& Name);
+	int32 GetNodeId() const { return NodeId; }
 
 private:
 	virtual bool UpdateInternal() override;
diff --git a/Engine/Plugins/Animation/GameplayInsights/Source/GameplayInsights/Private/MontagesTrack.cpp b/Engine/Plugins/Animation/GameplayInsights/Source/GameplayInsights/Private/MontagesTrack.cpp
--- a/Engine/Plugins/Animation/GameplayInsights/Source/GameplayInsights/Private/MontagesTrack.cpp
+++ b/Engine/Plugins/Animation/GameplayInsights/Source/GameplayInsights/Private/MontagesTrack.cpp
@@ -35,6 +35,7 @@ void FMontagesTrack::IterateSubTracksInternal(TFunction<void(TSharedPtr<FRewindD
 
 bool FMontagesTrack::UpdateInternal()
 {
+	TRACE_CPUPROFILER_EVENT_SCOPE(FMontagesTrack::UpdateInternal);
 	TArray<uint64> UniqueTrackIds;
 
 	IRewindDebugger* RewindDebugger = IRewindDebugger::Instance();
@@ -334,6 +335,7 @@ TSharedPtr<RewindDebugger::FRewindDebuggerTrack> FMontagesTrackCreator::CreateTr
 
 bool FMontagesTrackCreator::HasDebugInfoInternal(uint64 ObjectId) const
 {
+	TRACE_CPUPROFILER_EVENT_SCOPE(FMontagesTrack::HasDebugInfoInternal);
 	const TraceServices::IAnalysisSession* AnalysisSession = IRewindDebugger::Instance()->GetAnalysisSession();
 	TraceServices::FAnalysisSessionReadScope SessionReadScope(*AnalysisSession);
 	bool bHasData = false;
diff --git a/Engine/Plugins/Animation/GameplayInsights/Source/GameplayInsights/Private/NotifiesTrack.cpp b/Engine/Plugins/Animation/GameplayInsights/Source/GameplayInsights/Private/NotifiesTrack.cpp
--- a/Engine/Plugins/Animation/GameplayInsights/Source/GameplayInsights/Private/NotifiesTrack.cpp
+++ b/Engine/Plugins/Animation/GameplayInsights/Source/GameplayInsights/Private/NotifiesTrack.cpp
@@ -74,6 +74,7 @@ bool FNotifyTrack::UpdateInternal()
 	
 	if(EventUpdateRequested > 10 && GameplayProvider && AnimationProvider)
 	{
+		TRACE_CPUPROFILER_EVENT_SCOPE(FNotifiesTrack::UpdateEventPointsInternal);
 		EventUpdateRequested = 0;
 		
 		auto& EventPoints = EventData->Points;
@@ -177,6 +178,7 @@ FNotifiesTrack::FNotifiesTrack(uint64 InObjectId) :
 
 bool FNotifiesTrack::UpdateInternal()
 {
+	TRACE_CPUPROFILER_EVENT_SCOPE(FNotifiesTrack::UpdateInternal);
 	IRewindDebugger* RewindDebugger = IRewindDebugger::Instance();
 	
 	TRange<double> TraceTimeRange = RewindDebugger->GetCurrentTraceRange();
@@ -258,6 +260,7 @@ void FNotifiesTrack::IterateSubTracksInternal(TFunction<void(TSharedPtr<FRewindD
 
 bool FNotifiesTrackCreator::HasDebugInfoInternal(uint64 ObjectId) const
 {
+	TRACE_CPUPROFILER_EVENT_SCOPE(FNotifiesTrack::HasDebugInfoInternal);
 	const TraceServices::IAnalysisSession* AnalysisSession = IRewindDebugger::Instance()->GetAnalysisSession();
 	
 	TraceServices::FAnalysisSessionReadScope SessionReadScope(*AnalysisSession);
diff --git a/Engine/Plugins/Animation/GameplayInsights/Source/GameplayInsights/Private/PoseWatchTrack.cpp b/Engine/Plugins/Animation/GameplayInsights/Source/GameplayInsights/Private/PoseWatchTrack.cpp
--- a/Engine/Plugins/Animation/GameplayInsights/Source/GameplayInsights/Private/PoseWatchTrack.cpp
+++ b/Engine/Plugins/Animation/GameplayInsights/Source/GameplayInsights/Private/PoseWatchTrack.cpp
@@ -346,6 +346,7 @@ void FPoseWatchesTrack::IterateSubTracksInternal(TFunction<void(TSharedPtr<FRewi
 
 bool FPoseWatchesTrackCreator::HasDebugInfoInternal(uint64 ObjectId) const
 {
+	TRACE_CPUPROFILER_EVENT_SCOPE(FPoseWatchesTrack::HasDebugInfoInternal);
 	const TraceServices::IAnalysisSession* AnalysisSession = IRewindDebugger::Instance()->GetAnalysisSession();
 	
 	TraceServices::FAnalysisSessionReadScope SessionReadScope(*AnalysisSession);
diff --git a/Engine/Plugins/Animation/GameplayInsights/Source/GameplayInsights/Private/PropertiesTrack.cpp b/Engine/Plugins/Animation/GameplayInsights/Source/GameplayInsights/Private/PropertiesTrack.cpp
--- a/Engine/Plugins/Animation/GameplayInsights/Source/GameplayInsights/Private/PropertiesTrack.cpp
+++ b/Engine/Plugins/Animation/GameplayInsights/Source/GameplayInsights/Private/PropertiesTrack.cpp
@@ -307,6 +307,7 @@ namespace RewindDebugger
 
 	bool FPropertiesTrackCreator::HasDebugInfoInternal(uint64 ObjectId) const
 	{
+		TRACE_CPUPROFILER_EVENT_SCOPE(FPropertiesTrack::HasDebugInfoInternal);
 		bool bHasData = false;
 
 		if (const TraceServices::IAnalysisSession* AnalysisSession = IRewindDebugger::Instance()->GetAnalysisSession())
diff --git a/Engine/Plugins/Animation/GameplayInsights/Source/GameplayInsights/Public/IAnimationProvider.h b/Engine/Plugins/Animation/GameplayInsights/Source/GameplayInsights/Public/IAnimationProvider.h
--- a/Engine/Plugins/Animation/GameplayInsights/Source/GameplayInsights/Public/IAnimationProvider.h
+++ b/Engine/Plugins/Animation/GameplayInsights/Source/GameplayInsights/Public/IAnimationProvider.h
@@ -285,6 +285,22 @@ struct FAnimMontageMessage
 	uint16 FrameCounter = 0;
 };
 
+enum class EInertializationType : uint8
+{
+	Inertialization = 0,
+	DeadBlending = 1
+};
+
+struct FInertializationMessage
+{
+	uint64 AnimInstanceId = 0;
+	double ProfileTime;
+	double RecordingTime;
+	int32 NodeId;
+	float Weight;
+	EInertializationType Type;
+};
+
 struct FAnimAttributeMessage
 {
 	int32 SourceNodeId = 0;
@@ -314,6 +330,7 @@ public:
 	typedef TraceServices::ITimeline<FAnimAttributeMessage> AnimAttributeTimeline;
 	typedef TraceServices::ITimeline<FAnimSyncMessage> AnimSyncTimeline;
 	typedef TraceServices::ITimeline<FPoseWatchMessage> PoseWatchTimeline;
+	typedef TraceServices::ITimeline<FInertializationMessage> InertializationTimeline;
 
 	virtual void EnumerateSkeletalMeshPoseTimelines(TFunctionRef<void(uint64 ObjectId, const SkeletalMeshPoseTimeline&)> Callback) const =0;
 	virtual bool ReadSkeletalMeshPoseTimeline(uint64 InObjectId, TFunctionRef<void(const SkeletalMeshPoseTimeline&, bool)> Callback) const = 0;
@@ -324,6 +341,8 @@ public:
 	virtual void EnumeratePoseWatchCurves(const FPoseWatchMessage& InMessage, TFunctionRef<void(const FSkeletalMeshNamedCurve&)> Callback) const = 0;
 	virtual void EnumerateExternalMorphSets(const FSkeletalMeshPoseMessage& InMessage, TFunctionRef<void(const FExternalMorphWeightMessage&)> Callback) const = 0;
 	virtual bool ReadTickRecordTimeline(uint64 InObjectId, TFunctionRef<void(const TickRecordTimeline&)> Callback) const = 0;
+	virtual bool ReadInertializationTimeline(uint64 InObjectId, TFunctionRef<void(const InertializationTimeline&)> Callback) const = 0;
+	virtual void EnumerateInertializationNodes(uint64 InObjectId, TFunctionRef<void(int32, EInertializationType)> Callback) const = 0;
 	virtual void EnumerateTickRecordIds(uint64 InObjectId, TFunctionRef<void(uint64, int32)> Callback) const = 0;
 	virtual void EnumerateAnimGraphTimelines(TFunctionRef<void(uint64 ObjectId, const AnimGraphTimeline&)> Callback) const =0;
 	virtual bool ReadAnimGraphTimeline(uint64 InObjectId, TFunctionRef<void(const AnimGraphTimeline&)> Callback) const = 0;
diff --git a/Engine/Plugins/Animation/GameplayInsights/Source/RewindDebugger/Private/RewindDebuggerObjectTrack.cpp b/Engine/Plugins/Animation/GameplayInsights/Source/RewindDebugger/Private/RewindDebuggerObjectTrack.cpp
--- a/Engine/Plugins/Animation/GameplayInsights/Source/RewindDebugger/Private/RewindDebuggerObjectTrack.cpp
+++ b/Engine/Plugins/Animation/GameplayInsights/Source/RewindDebugger/Private/RewindDebuggerObjectTrack.cpp
@@ -14,14 +14,58 @@
 
 namespace RewindDebugger
 {
+	
+// check if an object is or is a subclass of a type by name, based on Insights traced type info
+static bool IsTargetType(uint64 ObjectId, FName TargetTypeName, const TraceServices::IAnalysisSession& Session)
+{
+	TraceServices::FAnalysisSessionReadScope SessionReadScope(Session);
+
+	const IGameplayProvider* GameplayProvider = Session.ReadProvider<IGameplayProvider>("GameplayProvider");
+	const FObjectInfo& ObjectInfo = GameplayProvider->GetObjectInfo(ObjectId);
+	uint64 ClassId = ObjectInfo.ClassId;
+
+	while (ClassId != 0)
+	{
+		const FClassInfo& ClassInfo = GameplayProvider->GetClassInfo(ClassId);
+		if (ClassInfo.Name == TargetTypeName)
+		{
+			return true;
+		}
+		ClassId = ClassInfo.SuperId;
+	}
+
+	return false;
+}
+	
+FRewindDebuggerObjectTrack::FRewindDebuggerObjectTrack(uint64 InObjectId, const FString& InObjectName, bool bInAddController)
+	: ObjectName(InObjectName)
+	, ObjectId(InObjectId)
+	, bAddController(bInAddController)
+	, bDisplayNameValid(false)
+{
+	ExistenceRange = MakeShared<SSegmentedTimelineView::FSegmentData>();
+	ExistenceRange->Segments.SetNumUninitialized(1);
+	
+	IRewindDebugger* RewindDebugger = IRewindDebugger::Instance();
+	const TraceServices::IAnalysisSession* Session = RewindDebugger->GetAnalysisSession();
+	TraceServices::FAnalysisSessionReadScope SessionReadScope(*Session);
+	FRewindDebuggerTrackCreators::EnumerateCreators([this, &Session, RewindDebugger](const IRewindDebuggerTrackCreator* Creator)
+	{
+		if (IsTargetType(ObjectId, Creator->GetTargetTypeName(), *Session))
+		{
+			TrackChildren.Push({Creator, nullptr});
+		}
+	});
 
+	// sort by creator priority + name
+}
+	
 TSharedPtr<SWidget> FRewindDebuggerObjectTrack::GetTimelineViewInternal()
 {
 	return SNew(SSegmentedTimelineView)
 		.ViewRange_Lambda([]() { return IRewindDebugger::Instance()->GetCurrentViewRange(); })
 		.FillColor(FLinearColor(0.02f, 0.02f, 0.02f, 0.5f))
 		.SegmentData_Raw(this, &FRewindDebuggerObjectTrack::GetExistenceRange);
-
 }
 
 bool FRewindDebuggerObjectTrack::HandleDoubleClickInternal()
@@ -67,30 +111,17 @@ bool FRewindDebuggerObjectTrack::HandleDoubleClickInternal()
 	return true;
 }
 
-// check if an object is or is a subclass of a type by name, based on Insights traced type info
-static bool IsTargetType(uint64 ObjectId, FName TargetTypeName, const TraceServices::IAnalysisSession& Session)
-{
-	TraceServices::FAnalysisSessionReadScope SessionReadScope(Session);
-
-	const IGameplayProvider* GameplayProvider = Session.ReadProvider<IGameplayProvider>("GameplayProvider");
-	const FObjectInfo& ObjectInfo = GameplayProvider->GetObjectInfo(ObjectId);
-	uint64 ClassId = ObjectInfo.ClassId;
 
-	while (ClassId != 0)
+void FRewindDebuggerObjectTrack::IterateSubTracksInternal(TFunction<void(TSharedPtr<FRewindDebuggerTrack> SubTrack)> IteratorFunction)
+{
+	for (FTrackCreatorAndTrack& TrackChild : TrackChildren)
 	{
-		const FClassInfo& ClassInfo = GameplayProvider->GetClassInfo(ClassId);
-		if (ClassInfo.Name == TargetTypeName)
+		if (TrackChild.Track.IsValid())
 		{
-			return true;
+			IteratorFunction(TrackChild.Track);
 		}
-		ClassId = ClassInfo.SuperId;
 	}
-
-	return false;
-}
-
-void FRewindDebuggerObjectTrack::IterateSubTracksInternal(TFunction<void(TSharedPtr<FRewindDebuggerTrack> SubTrack)> IteratorFunction)
-{
+	
 	for (TSharedPtr<FRewindDebuggerTrack>& Track : Children)
 	{
 		IteratorFunction(Track);
@@ -150,6 +181,7 @@ bool FRewindDebuggerObjectTrack::UpdateInternal()
 
 	if (!Icon.IsSet())
 	{
+		TRACE_CPUPROFILER_EVENT_SCOPE(FRewindDebuggerObjectTrack::FindIcon);
 		if (const FObjectInfo* ObjectInfo = GameplayProvider->FindObjectInfo(ObjectId))
 		{
 			if (const UClass* FoundClass = GameplayProvider->FindClass(ObjectInfo->ClassId))
@@ -165,60 +197,59 @@ bool FRewindDebuggerObjectTrack::UpdateInternal()
 	FoundObjects.Add(ObjectId); // prevent debug views from being removed
 
 	// add debug views as children
-	FRewindDebuggerTrackCreators::EnumerateCreators([this, &Session, &bChanged, RewindDebugger](const IRewindDebuggerTrackCreator* Creator)
-		{
-			int32 FoundIndex = Children.FindLastByPredicate([&bChanged, this, Creator, RewindDebugger](const TSharedPtr<FRewindDebuggerTrack>& Track) { return Track->GetName() == Creator->GetName(); });
 
-			const bool bHasDebugInfo = Creator->HasDebugInfo(ObjectId) && IsTargetType(ObjectId, Creator->GetTargetTypeName(), *Session);
+	for(FTrackCreatorAndTrack& TrackChild : TrackChildren)
+	{
+		const bool bHasDebugInfo = TrackChild.Creator->HasDebugInfo(ObjectId);
 
-			if (FoundIndex >= 0)
+		if (TrackChild.Track.IsValid())
+		{
+			if (!bHasDebugInfo)
 			{
-				if (!bHasDebugInfo)
-				{
-					bChanged = true;
-					Children.RemoveAt(FoundIndex);
-				}
+				bChanged = true;
+				TrackChild.Track.Reset();
 			}
-			else
+		}
+		else
+		{
+			if (bHasDebugInfo)
 			{
-				if (bHasDebugInfo)
-				{
-					bChanged = true;
-					if (TSharedPtr<FRewindDebuggerTrack> Track = Creator->CreateTrack(ObjectId))
-					{
-						Children.Add(Track);
-					}
-				}
+				bChanged = true;
+				TrackChild.Track = TrackChild.Creator->CreateTrack(ObjectId);
 			}
-		});
+		}
+	}
 
 	// Fallback codepath to add views with no track implementation
-	FRewindDebuggerViewCreators::EnumerateCreators([this, &Session, &bChanged, RewindDebugger](const IRewindDebuggerViewCreator* Creator)
-		{
-			int32 FoundIndex = Children.FindLastByPredicate([&bChanged, this, Creator, RewindDebugger](const TSharedPtr<FRewindDebuggerTrack>& Track) { return Track->GetName() == Creator->GetName(); });
+	{
+		TRACE_CPUPROFILER_EVENT_SCOPE(FRewindDebuggerObjectTrack::UpdateInternal_AddViews);
+		FRewindDebuggerViewCreators::EnumerateCreators([this, &Session, &bChanged, RewindDebugger](const IRewindDebuggerViewCreator* Creator)
+			{
+				int32 FoundIndex = Children.FindLastByPredicate([&bChanged, this, Creator, RewindDebugger](const TSharedPtr<FRewindDebuggerTrack>& Track) { return Track->GetName() == Creator->GetName(); });
 
-			bool bHasDebugInfo = Creator->HasDebugInfo(ObjectId)
-				&& IsTargetType(ObjectId, Creator->GetTargetTypeName(), *Session);
+				bool bHasDebugInfo = Creator->HasDebugInfo(ObjectId)
+					&& IsTargetType(ObjectId, Creator->GetTargetTypeName(), *Session);
 
-			if (FoundIndex >= 0)
-			{
-				if (!bHasDebugInfo)
+				if (FoundIndex >= 0)
 				{
-					bChanged = true;
-					Children.RemoveAt(FoundIndex);
+					if (!bHasDebugInfo)
+					{
+						bChanged = true;
+						Children.RemoveAt(FoundIndex);
+					}
 				}
-			}
-			else
-			{
-				if (bHasDebugInfo)
+				else
 				{
-					bChanged = true;
-					TSharedPtr<FRewindDebuggerTrack> Track = MakeShared<FRewindDebuggerFallbackTrack>(ObjectId, Creator);
-					Children.Add(Track);
+					if (bHasDebugInfo)
+					{
+						bChanged = true;
+						TSharedPtr<FRewindDebuggerTrack> Track = MakeShared<FRewindDebuggerFallbackTrack>(ObjectId, Creator);
+						Children.Add(Track);
+					}
 				}
 			}
-		}
-	);
+		);
+	}
 
 	// add child objects/components
 	{
@@ -253,6 +284,7 @@ bool FRewindDebuggerObjectTrack::UpdateInternal()
 	// add controller and it's component hierarchy if one is attached
 	if (bAddController)
 	{
+		TRACE_CPUPROFILER_EVENT_SCOPE(FRewindDebuggerObjectTrack::FindController);
 		// Should probably update this to use a time range and return all posessing controllers from the visible time range.  For now just returns the one at the current time.
 		if (uint64 ControllerId = GameplayProvider->FindPossessingController(ObjectId, RewindDebugger->CurrentTraceTime()))
 		{
@@ -281,6 +313,7 @@ bool FRewindDebuggerObjectTrack::UpdateInternal()
 
 	if (bChanged)
 	{
+		TRACE_CPUPROFILER_EVENT_SCOPE(FRewindDebuggerObjectTrack::Sort);
 		Children.Sort([](const TSharedPtr<FRewindDebuggerTrack>& A, const TSharedPtr<FRewindDebuggerTrack>& B)
 		{
 			const int SortOrderPriorityA = A->GetSortOrderPriority();
@@ -295,11 +328,27 @@ bool FRewindDebuggerObjectTrack::UpdateInternal()
 		});
 	}
 
-	for (auto& Child : Children)
 	{
-		if (Child->Update())
+		TRACE_CPUPROFILER_EVENT_SCOPE(FRewindDebuggerObjectTrack::UpdateChilden);
+		for (auto& Child : Children)
 		{
-			bChanged = true;
+			if (Child->Update())
+			{
+				bChanged = true;
+			}
+		}
+	}
+	{
+		TRACE_CPUPROFILER_EVENT_SCOPE(FRewindDebuggerObjectTrack::UpdateTrackChilden);
+		for (auto& TrackChild : TrackChildren)
+		{
+			if (TrackChild.Track.IsValid())
+			{
+				if (TrackChild.Track->Update())
+				{
+					bChanged = true;
+				}
+			}
 		}
 	}
 
diff --git a/Engine/Plugins/Animation/GameplayInsights/Source/RewindDebugger/Private/RewindDebuggerObjectTrack.h b/Engine/Plugins/Animation/GameplayInsights/Source/RewindDebugger/Private/RewindDebuggerObjectTrack.h
--- a/Engine/Plugins/Animation/GameplayInsights/Source/RewindDebugger/Private/RewindDebuggerObjectTrack.h
+++ b/Engine/Plugins/Animation/GameplayInsights/Source/RewindDebugger/Private/RewindDebuggerObjectTrack.h
@@ -4,22 +4,23 @@
 #include "RewindDebuggerTrack.h"
 #include "SSegmentedTimelineView.h"
 
+
 namespace RewindDebugger
 {
+	
+class IRewindDebuggerTrackCreator;
+
+struct FTrackCreatorAndTrack
+{
+	const IRewindDebuggerTrackCreator* Creator;
+	TSharedPtr<FRewindDebuggerTrack> Track;
+};
 
 class FRewindDebuggerObjectTrack : public FRewindDebuggerTrack
 {
 public:
 
-	FRewindDebuggerObjectTrack(uint64 InObjectId, const FString& InObjectName, bool bInAddController = false)
-		: ObjectName(InObjectName)
-		, ObjectId(InObjectId)
-		, bAddController(bInAddController)
-		, bDisplayNameValid(false)
-	{
-		ExistenceRange = MakeShared<SSegmentedTimelineView::FSegmentData>();
-		ExistenceRange->Segments.SetNumUninitialized(1);
-	}
+	FRewindDebuggerObjectTrack(uint64 InObjectId, const FString& InObjectName, bool bInAddController = false);
 
 	TSharedPtr<SSegmentedTimelineView::FSegmentData> GetExistenceRange() const { return ExistenceRange; }
 
@@ -41,10 +42,12 @@ private:
 	FSlateIcon Icon;
 	TSharedPtr<SSegmentedTimelineView::FSegmentData> ExistenceRange;
 	uint64 ObjectId;
+	TArray<FTrackCreatorAndTrack> TrackChildren;
 	TArray<TSharedPtr<FRewindDebuggerTrack>> Children;
 
 	bool bAddController;
 	mutable bool bDisplayNameValid;
+	
 };
 
 }
\ No newline at end of file
diff --git a/Engine/Plugins/Animation/GameplayInsights/Source/RewindDebuggerVLog/Private/VisualLogTrack.cpp b/Engine/Plugins/Animation/GameplayInsights/Source/RewindDebuggerVLog/Private/VisualLogTrack.cpp
--- a/Engine/Plugins/Animation/GameplayInsights/Source/RewindDebuggerVLog/Private/VisualLogTrack.cpp
+++ b/Engine/Plugins/Animation/GameplayInsights/Source/RewindDebuggerVLog/Private/VisualLogTrack.cpp
@@ -85,6 +85,7 @@ bool FVisualLogCategoryTrack::UpdateInternal()
 	{
 		if(EventUpdateRequested > 10)
 		{
+			TRACE_CPUPROFILER_EVENT_SCOPE(FVisualLogTrack::UpdateEventPointsInternal);
 			EventUpdateRequested = 0;
 			
 			EventData->Points.SetNum(0,false);
@@ -191,6 +192,7 @@ FVisualLogTrack::FVisualLogTrack(uint64 InObjectId) :
 
 bool FVisualLogTrack::UpdateInternal()
 {
+	TRACE_CPUPROFILER_EVENT_SCOPE(FVisualLogTrack::UpdateInternal);
 	IRewindDebugger* RewindDebugger = IRewindDebugger::Instance();
 	
 	TRange<double> TraceTimeRange = RewindDebugger->GetCurrentTraceRange();
@@ -255,6 +257,7 @@ void FVisualLogTrack::IterateSubTracksInternal(TFunction<void(TSharedPtr<FRewind
 
 bool FVisualLogTrackCreator::HasDebugInfoInternal(uint64 ObjectId) const
 {
+	TRACE_CPUPROFILER_EVENT_SCOPE(FVisualLogTrack::HasDebugInfoInternal);
 	const TraceServices::IAnalysisSession* AnalysisSession = IRewindDebugger::Instance()->GetAnalysisSession();
 	
 	TraceServices::FAnalysisSessionReadScope SessionReadScope(*AnalysisSession);
diff --git a/Engine/Plugins/Experimental/Animation/PoseSearch/Source/Editor/Private/PoseSearchDebugger.cpp b/Engine/Plugins/Experimental/Animation/PoseSearch/Source/Editor/Private/PoseSearchDebugger.cpp
--- a/Engine/Plugins/Experimental/Animation/PoseSearch/Source/Editor/Private/PoseSearchDebugger.cpp
+++ b/Engine/Plugins/Experimental/Animation/PoseSearch/Source/Editor/Private/PoseSearchDebugger.cpp
@@ -386,6 +386,7 @@ FSlateIcon FDebuggerTrack::GetIconInternal()
 
 bool FDebuggerTrack::UpdateInternal()
 {
+	TRACE_CPUPROFILER_EVENT_SCOPE(PoseSearchDebuggerTrack::UpdateInternal);
 	CostTimelineView->UpdateInternal(ObjectId);
 
 	if (TSharedPtr<IRewindDebuggerView> PinnedView = View.Pin())
@@ -441,6 +442,7 @@ TSharedPtr<RewindDebugger::FRewindDebuggerTrack> FDebuggerTrackCreator::CreateTr
 
 bool FDebuggerTrackCreator::HasDebugInfoInternal(uint64 ObjectId) const
 {
+	TRACE_CPUPROFILER_EVENT_SCOPE(PoseSearchDebugger::HasDebugInfoInternal);
 	// Get provider and validate
 	const TraceServices::IAnalysisSession* Session = IRewindDebugger::Instance()->GetAnalysisSession();
 	TraceServices::FAnalysisSessionReadScope SessionReadScope(*Session);
diff --git a/Engine/Plugins/Experimental/Chooser/Source/ChooserEditor/Private/ChooserTrack.cpp b/Engine/Plugins/Experimental/Chooser/Source/ChooserEditor/Private/ChooserTrack.cpp
--- a/Engine/Plugins/Experimental/Chooser/Source/ChooserEditor/Private/ChooserTrack.cpp
+++ b/Engine/Plugins/Experimental/Chooser/Source/ChooserEditor/Private/ChooserTrack.cpp
@@ -71,6 +71,7 @@ bool FChooserTrack::UpdateInternal()
 	
 	if(EventUpdateRequested > 10 && ChooserProvider)
 	{
+		TRACE_CPUPROFILER_EVENT_SCOPE(FChooserTrack::UpdateEventPointsInternal);
 		EventUpdateRequested = 0;
 		
 		EventData->Points.SetNum(0,false);
@@ -165,6 +166,7 @@ FChoosersTrack::FChoosersTrack(uint64 InObjectId) :
 
 bool FChoosersTrack::UpdateInternal()
 {
+	TRACE_CPUPROFILER_EVENT_SCOPE(FChoosersTrack::UpdateInternal);
 	IRewindDebugger* RewindDebugger = IRewindDebugger::Instance();
 	
 	TRange<double> RecordingTimeRange = RewindDebugger->GetCurrentViewRange();
@@ -226,6 +228,7 @@ void FChoosersTrack::IterateSubTracksInternal(TFunction<void(TSharedPtr<FRewindD
 
 bool FChoosersTrackCreator::HasDebugInfoInternal(uint64 ObjectId) const
 {
+	TRACE_CPUPROFILER_EVENT_SCOPE(FChoosersTrack::HasDebugInfoInternal);
 	const TraceServices::IAnalysisSession* AnalysisSession = IRewindDebugger::Instance()->GetAnalysisSession();
 	
 	TraceServices::FAnalysisSessionReadScope SessionReadScope(*AnalysisSession);
diff --git a/Engine/Source/Runtime/Engine/Private/Animation/AnimNode_DeadBlending.cpp b/Engine/Source/Runtime/Engine/Private/Animation/AnimNode_DeadBlending.cpp
--- a/Engine/Source/Runtime/Engine/Private/Animation/AnimNode_DeadBlending.cpp
+++ b/Engine/Source/Runtime/Engine/Private/Animation/AnimNode_DeadBlending.cpp
@@ -925,6 +925,8 @@ void FAnimNode_DeadBlending::Evaluate_AnyThread(FPoseContext& Output)
 	TRACE_ANIM_NODE_VALUE_WITH_ID(Output, GetNodeIndex(), TEXT("Inertialization Weight"), InertializationWeight);
 	TRACE_ANIM_NODE_VALUE_WITH_ID(Output, GetNodeIndex(), TEXT("Request Description"), *InertializationRequestDescription);
 	TRACE_ANIM_NODE_VALUE_WITH_ID_ANIM_NODE(Output, GetNodeIndex(), TEXT("Request Node"), InertializationRequestNodeId, InertializationRequestAnimInstance);
+
+	TRACE_ANIM_INERTIALIZATION(*Output.AnimInstanceProxy, GetNodeIndex(), InertializationWeight, FAnimTrace::EInertializationType::DeadBlending);
 }
 
 bool FAnimNode_DeadBlending::NeedsDynamicReset() const
diff --git a/Engine/Source/Runtime/Engine/Private/Animation/AnimNode_Inertialization.cpp b/Engine/Source/Runtime/Engine/Private/Animation/AnimNode_Inertialization.cpp
--- a/Engine/Source/Runtime/Engine/Private/Animation/AnimNode_Inertialization.cpp
+++ b/Engine/Source/Runtime/Engine/Private/Animation/AnimNode_Inertialization.cpp
@@ -489,6 +489,8 @@ void FAnimNode_Inertialization::Evaluate_AnyThread(FPoseContext& Output)
 	TRACE_ANIM_NODE_VALUE_WITH_ID(Output, GetNodeIndex(), TEXT("Inertialization Weight"), InertializationWeight);
 	TRACE_ANIM_NODE_VALUE_WITH_ID(Output, GetNodeIndex(), TEXT("Request Description"), *InertializationRequestDescription);
 	TRACE_ANIM_NODE_VALUE_WITH_ID_ANIM_NODE(Output, GetNodeIndex(), TEXT("Request Node"), InertializationRequestNodeId, InertializationRequestAnimInstance);
+
+	TRACE_ANIM_INERTIALIZATION(*Output.AnimInstanceProxy, GetNodeIndex(), InertializationWeight, FAnimTrace::EInertializationType::Inertialization);
 }
 
 
diff --git a/Engine/Source/Runtime/Engine/Private/Animation/AnimTrace.cpp b/Engine/Source/Runtime/Engine/Private/Animation/AnimTrace.cpp
--- a/Engine/Source/Runtime/Engine/Private/Animation/AnimTrace.cpp
+++ b/Engine/Source/Runtime/Engine/Private/Animation/AnimTrace.cpp
@@ -285,6 +285,16 @@ UE_TRACE_EVENT_BEGIN(Animation, PoseWatch2)
 	UE_TRACE_EVENT_FIELD(bool, bIsEnabled)
 UE_TRACE_EVENT_END()
 
+UE_TRACE_EVENT_BEGIN(Animation, Inertialization)
+	UE_TRACE_EVENT_FIELD(uint64, Cycle)
+	UE_TRACE_EVENT_FIELD(double, RecordingTime)
+	UE_TRACE_EVENT_FIELD(uint64, ComponentId)
+	UE_TRACE_EVENT_FIELD(uint64, AnimInstanceId)
+	UE_TRACE_EVENT_FIELD(uint32, NodeId)
+	UE_TRACE_EVENT_FIELD(float, Weight)
+	UE_TRACE_EVENT_FIELD(uint8, Type)
+UE_TRACE_EVENT_END()
+
 FAutoConsoleVariable CVarRecordExternalMorphTargets(
 	TEXT("RecordExternalMorphTargets"),
 	false,
@@ -1361,4 +1371,18 @@ void FAnimTrace::OutputPoseWatch(const FAnimInstanceProxy& InSourceProxy, UPoseW
 		<< PoseWatch2.bIsEnabled(bIsEnabled);
 }
 
+void FAnimTrace::OutputInertialization(const FAnimInstanceProxy& InSourceProxy, int32 NodeId, float Weight, EInertializationType Type)
+{
+	const UAnimInstance* AnimInstance = CastChecked<UAnimInstance>(InSourceProxy.GetAnimInstanceObject());
+	TRACE_OBJECT(AnimInstance);
+	
+	UE_TRACE_LOG(Animation, Inertialization, AnimationChannel)
+	<< Inertialization.Cycle(FPlatformTime::Cycles64())
+	<< Inertialization.RecordingTime(FObjectTrace::GetWorldElapsedTime(AnimInstance->GetWorld()))
+	<< Inertialization.AnimInstanceId(FObjectTrace::GetObjectId(InSourceProxy.GetAnimInstanceObject()))
+	<< Inertialization.NodeId(NodeId)
+	<< Inertialization.Weight(Weight)
+	<< Inertialization.Type(static_cast<uint8>(Type));
+}
+
 #endif
diff --git a/Engine/Source/Runtime/Engine/Public/Animation/AnimTrace.h b/Engine/Source/Runtime/Engine/Public/Animation/AnimTrace.h
--- a/Engine/Source/Runtime/Engine/Public/Animation/AnimTrace.h
+++ b/Engine/Source/Runtime/Engine/Public/Animation/AnimTrace.h
@@ -59,6 +59,12 @@ struct FAnimTrace
 		Tick = 3,
 		SyncMarker = 4	// We 'fake' sync markers with a notify type for convenience
 	};
+	
+	enum class EInertializationType : uint8
+    {
+    	Inertialization = 0,
+    	DeadBlending = 1
+    };
 
 	/** Helper for outputting anim nodes */
 	struct FScopedAnimNodeTrace
@@ -196,6 +202,9 @@ struct FAnimTrace
 
 	/** Helper function to output a pose watch record */
 	ENGINE_API static void OutputPoseWatch(const FAnimInstanceProxy& InSourceProxy, UPoseWatchPoseElement* InPoseWatchElement, int32 InPoseWatchId, const TArray<FTransform>& BoneTransforms, const FBlendedHeapCurve& InCurves, const TArray<FBoneIndexType>& RequiredBones, const FTransform& WorldTransform, const bool bIsEnabled);
+	
+	/** Helper function to output inertialization state */
+	ENGINE_API static void OutputInertialization(const FAnimInstanceProxy& InSourceProxy, int32 InNodeId, float InWeight, EInertializationType InType);
 };
 
 #define TRACE_ANIM_TICK_RECORD(Context, TickRecord) \
@@ -255,6 +264,9 @@ struct FAnimTrace
 #define TRACE_ANIM_POSE_WATCH(SourceProxy, PoseWatchElement, PoseWatchId, BoneTransforms, Curves, RequiredBones, WorldTransform, bIsEnabled) \
 	FAnimTrace::OutputPoseWatch(SourceProxy, PoseWatchElement, PoseWatchId, BoneTransforms, Curves, RequiredBones, WorldTransform, bIsEnabled);
 
+#define TRACE_ANIM_INERTIALIZATION(SourceProxy, NodeId, Weight, Type) \
+	FAnimTrace::OutputInertialization(SourceProxy, NodeId, Weight, Type);
+
 #else
 
 #define TRACE_ANIM_TICK_RECORD(Context, TickRecord)
@@ -276,5 +288,6 @@ struct FAnimTrace
 #define TRACE_ANIM_MONTAGE(AnimInstance, MontageInstance)
 #define TRACE_ANIM_NODE_SYNC(SourceProxy, SourceNodeId, GroupName)
 #define TRACE_ANIM_POSE_WATCH(SourceProxy, PoseWatchElement, PoseWatchId, BoneTransforms, Curves, RequiredBones, WorldTransform, bIsEnabled)
+#define TRACE_ANIM_INERTIALIZATION(SourceProxy, NodeId, Weight, Type)
 
 #endif
