fix track sort order in rewind debugger

diff --git a/Engine/Plugins/Animation/GameplayInsights/Source/GameplayInsights/Private/GameplayInsightsModule.cpp b/Engine/Plugins/Animation/GameplayInsights/Source/GameplayInsights/Private/GameplayInsightsModule.cpp
--- a/Engine/Plugins/Animation/GameplayInsights/Source/GameplayInsights/Private/GameplayInsightsModule.cpp
+++ b/Engine/Plugins/Animation/GameplayInsights/Source/GameplayInsights/Private/GameplayInsightsModule.cpp
@@ -75,10 +75,9 @@ void FGameplayInsightsModule::StartupModule()
 	});
 
 #if WITH_EDITOR
-	// register rewind debugger view creators
-	static FAnimGraphSchematicViewCreator AnimGraphSchematicViewCreator;
-	IModularFeatures::Get().RegisterModularFeature(IRewindDebuggerViewCreator::ModularFeatureName, &AnimGraphSchematicViewCreator);
-	
+	// register rewind debugger track creators
+	static FAnimGraphSchematicTrackCreator AnimGraphSchematicTrackCreator;
+	IModularFeatures::Get().RegisterModularFeature(RewindDebugger::IRewindDebuggerTrackCreator::ModularFeatureName, &AnimGraphSchematicTrackCreator);
 	static RewindDebugger::FAnimationCurvesTrackCreator AnimationCurvesTrackCreator;
 	IModularFeatures::Get().RegisterModularFeature(RewindDebugger::IRewindDebuggerTrackCreator::ModularFeatureName, &AnimationCurvesTrackCreator);
 	static RewindDebugger::FInertializationsTrackCreator InertializationsTrackCreator;
diff --git a/Engine/Plugins/Animation/GameplayInsights/Source/GameplayInsights/Private/SAnimGraphSchematicView.cpp b/Engine/Plugins/Animation/GameplayInsights/Source/GameplayInsights/Private/SAnimGraphSchematicView.cpp
--- a/Engine/Plugins/Animation/GameplayInsights/Source/GameplayInsights/Private/SAnimGraphSchematicView.cpp
+++ b/Engine/Plugins/Animation/GameplayInsights/Source/GameplayInsights/Private/SAnimGraphSchematicView.cpp
@@ -939,23 +939,33 @@ void SAnimGraphSchematicView::RefreshDetails(const TArray<TSharedRef<FAnimGraphS
 	}
 }
 
-FName FAnimGraphSchematicViewCreator::GetTargetTypeName() const
+FName FAnimGraphSchematicTrackCreator::GetTargetTypeNameInternal() const
 {
 	static FName TargetTypeName = "AnimInstance";
 	return TargetTypeName;
 }
 
-FName FAnimGraphSchematicViewCreator::GetName() const
+FName FAnimGraphSchematicTrackCreator::GetNameInternal() const
 {
 	return AnimGraphSchematicName;
 }
 
-FText FAnimGraphSchematicViewCreator::GetTitle() const
+TSharedPtr<RewindDebugger::FRewindDebuggerTrack> FAnimGraphSchematicTrackCreator::CreateTrackInternal(uint64 ObjectId) const
+{
+	return MakeShared<FAnimGraphSchematicTrack>(ObjectId);
+}
+
+FText FAnimGraphSchematicTrack::GetDisplayNameInternal() const
 {
 	return LOCTEXT("Anim Graph Update", "Anim Graph");
 }
 
-FSlateIcon FAnimGraphSchematicViewCreator::GetIcon() const
+FName FAnimGraphSchematicTrack::GetNameInternal() const
+{
+	return AnimGraphSchematicName;
+}
+
+FSlateIcon FAnimGraphSchematicTrack::GetIconInternal()
 {
 #if WITH_EDITOR
 	return FSlateIconFinder::FindIconForClass(UAnimInstance::StaticClass());
@@ -964,12 +974,38 @@ FSlateIcon FAnimGraphSchematicViewCreator::GetIcon() const
 #endif
 }
 
-TSharedPtr<IRewindDebuggerView> FAnimGraphSchematicViewCreator::CreateDebugView(uint64 ObjectId, double CurrentTime, const TraceServices::IAnalysisSession& AnalysisSession) const
+TSharedPtr<SWidget> FAnimGraphSchematicTrack::GetDetailsViewInternal()
 {
-	return SNew(SAnimGraphSchematicView, ObjectId, CurrentTime, AnalysisSession);
+	TSharedPtr<SAnimGraphSchematicView> NewView;
+	
+	if (IRewindDebugger* RewindDebugger = IRewindDebugger::Instance())
+	{
+		if (const TraceServices::IAnalysisSession* AnalysisSession = RewindDebugger->GetAnalysisSession())
+		{
+			if (!View.IsValid())
+			{
+				NewView = SNew(SAnimGraphSchematicView, AnimInstanceId, RewindDebugger->CurrentTraceTime(), *AnalysisSession);
+				View = NewView.ToWeakPtr();
+			}
+		}
+	}
+
+	return NewView;
+}
+
+bool FAnimGraphSchematicTrack::UpdateInternal()
+{
+	IRewindDebugger* RewindDebugger = IRewindDebugger::Instance();
+	TSharedPtr<IRewindDebuggerView> PinnedView = View.Pin();
+	if (PinnedView.IsValid())
+	{
+		PinnedView->SetTimeMarker(RewindDebugger->CurrentTraceTime());
+	}
+
+	return false;
 }
 
-bool FAnimGraphSchematicViewCreator::HasDebugInfo(uint64 ObjectId) const
+bool FAnimGraphSchematicTrackCreator::HasDebugInfoInternal(uint64 ObjectId) const
 {
 	const TraceServices::IAnalysisSession* AnalysisSession = IRewindDebugger::Instance()->GetAnalysisSession();
 	
diff --git a/Engine/Plugins/Animation/GameplayInsights/Source/GameplayInsights/Private/SAnimGraphSchematicView.h b/Engine/Plugins/Animation/GameplayInsights/Source/GameplayInsights/Private/SAnimGraphSchematicView.h
--- a/Engine/Plugins/Animation/GameplayInsights/Source/GameplayInsights/Private/SAnimGraphSchematicView.h
+++ b/Engine/Plugins/Animation/GameplayInsights/Source/GameplayInsights/Private/SAnimGraphSchematicView.h
@@ -3,8 +3,9 @@
 #pragma once
 
 #include "CoreMinimal.h"
+#include "IRewindDebuggerTrackCreator.h"
 #include "IRewindDebuggerView.h"
-#include "IRewindDebuggerViewCreator.h"
+#include "RewindDebuggerTrack.h"
 #include "Widgets/SCompoundWidget.h"
 #include "Widgets/Views/STreeView.h"
 
@@ -116,13 +117,33 @@ private:
 	uint64 AnimInstanceId;
 };
 
-class FAnimGraphSchematicViewCreator : public IRewindDebuggerViewCreator
+class FAnimGraphSchematicTrack : public RewindDebugger::FRewindDebuggerTrack
 {
-	public:
-		virtual FName GetTargetTypeName() const;
-		virtual FName GetName() const override;
-		virtual FText GetTitle() const override;
-		virtual FSlateIcon GetIcon() const override;
-		virtual TSharedPtr<IRewindDebuggerView> CreateDebugView(uint64 ObjectId, double CurrentTime, const TraceServices::IAnalysisSession& InAnalysisSession) const override;
-		virtual bool HasDebugInfo(uint64 ObjectId) const override;
+public:
+
+	FAnimGraphSchematicTrack(uint64 InObjectId)
+		: AnimInstanceId(InObjectId)
+	{
+	}
+
+private:
+	virtual FSlateIcon GetIconInternal() override;
+	virtual TSharedPtr<SWidget> GetDetailsViewInternal() override;
+	virtual FName GetNameInternal() const override;
+	virtual FText GetDisplayNameInternal() const override;
+	virtual uint64 GetObjectIdInternal() const override { return AnimInstanceId; }
+	virtual bool UpdateInternal() override;
+
+	TWeakPtr<SAnimGraphSchematicView> View;
+	FSlateIcon Icon;
+	uint64 AnimInstanceId;
+};
+
+class FAnimGraphSchematicTrackCreator : public RewindDebugger::IRewindDebuggerTrackCreator
+{
+public:
+	virtual FName GetTargetTypeNameInternal() const override;
+	virtual FName GetNameInternal() const override;
+	virtual TSharedPtr<RewindDebugger::FRewindDebuggerTrack> CreateTrackInternal(uint64 ObjectId) const override;
+	virtual bool HasDebugInfoInternal(uint64 ObjectId) const override;
 };
\ No newline at end of file
diff --git a/Engine/Plugins/Animation/GameplayInsights/Source/GameplayInsights/Private/SMontageView.h b/Engine/Plugins/Animation/GameplayInsights/Source/GameplayInsights/Private/SMontageView.h
--- a/Engine/Plugins/Animation/GameplayInsights/Source/GameplayInsights/Private/SMontageView.h
+++ b/Engine/Plugins/Animation/GameplayInsights/Source/GameplayInsights/Private/SMontageView.h
@@ -3,8 +3,6 @@
 #pragma once
 
 #include "CoreMinimal.h"
-#include "IRewindDebuggerView.h"
-#include "IRewindDebuggerViewCreator.h"
 #include "SPropertiesDebugViewBase.h"
 
 namespace TraceServices { class IAnalysisSession; }
diff --git a/Engine/Plugins/Animation/GameplayInsights/Source/RewindDebugger/Private/RewindDebuggerFallbackTrack.cpp b/Engine/Plugins/Animation/GameplayInsights/Source/RewindDebugger/Private/RewindDebuggerFallbackTrack.cpp
--- a/Engine/Plugins/Animation/GameplayInsights/Source/RewindDebugger/Private/RewindDebuggerFallbackTrack.cpp
+++ b/Engine/Plugins/Animation/GameplayInsights/Source/RewindDebugger/Private/RewindDebuggerFallbackTrack.cpp
@@ -4,6 +4,7 @@
 #include "IRewindDebugger.h"
 #include "IGameplayProvider.h"
 #include "IAnimationProvider.h"
+#include "IRewindDebuggerView.h"
 #include "Styling/SlateIconFinder.h"
 #include "ObjectTrace.h"
 
diff --git a/Engine/Plugins/Animation/GameplayInsights/Source/RewindDebugger/Private/RewindDebuggerFallbackTrack.h b/Engine/Plugins/Animation/GameplayInsights/Source/RewindDebugger/Private/RewindDebuggerFallbackTrack.h
--- a/Engine/Plugins/Animation/GameplayInsights/Source/RewindDebugger/Private/RewindDebuggerFallbackTrack.h
+++ b/Engine/Plugins/Animation/GameplayInsights/Source/RewindDebugger/Private/RewindDebuggerFallbackTrack.h
@@ -2,6 +2,7 @@
 #pragma once
 
 #include "RewindDebuggerTrack.h"
+#include "IRewindDebuggerViewCreator.h"
 
 namespace RewindDebugger
 {
diff --git a/Engine/Plugins/Animation/GameplayInsights/Source/RewindDebugger/Private/RewindDebuggerObjectTrack.cpp b/Engine/Plugins/Animation/GameplayInsights/Source/RewindDebugger/Private/RewindDebuggerObjectTrack.cpp
--- a/Engine/Plugins/Animation/GameplayInsights/Source/RewindDebugger/Private/RewindDebuggerObjectTrack.cpp
+++ b/Engine/Plugins/Animation/GameplayInsights/Source/RewindDebugger/Private/RewindDebuggerObjectTrack.cpp
@@ -57,7 +57,19 @@ FRewindDebuggerObjectTrack::FRewindDebuggerObjectTrack(uint64 InObjectId, const
 		}
 	});
 
-	// sort by creator priority + name
+	// sort by creators by priority + name
+	TrackChildren.Sort([](const FTrackCreatorAndTrack& A, const FTrackCreatorAndTrack& B)
+		{
+			const int32 SortOrderPriorityA = A.Creator->GetSortOrderPriority();
+			const int32 SortOrderPriorityB = B.Creator->GetSortOrderPriority();
+			
+			if (SortOrderPriorityA != SortOrderPriorityB)
+			{
+				return SortOrderPriorityA > SortOrderPriorityB;
+			}
+			
+			return A.Creator->GetName().ToString() < B.Creator->GetName().ToString();
+		});
 }
 	
 TSharedPtr<SWidget> FRewindDebuggerObjectTrack::GetTimelineViewInternal()
@@ -313,17 +325,9 @@ bool FRewindDebuggerObjectTrack::UpdateInternal()
 
 	if (bChanged)
 	{
-		TRACE_CPUPROFILER_EVENT_SCOPE(FRewindDebuggerObjectTrack::Sort);
+		// sort child object tracks by name
 		Children.Sort([](const TSharedPtr<FRewindDebuggerTrack>& A, const TSharedPtr<FRewindDebuggerTrack>& B)
 		{
-			const int SortOrderPriorityA = A->GetSortOrderPriority();
-			const int SortOrderPriorityB = B->GetSortOrderPriority();
-			
-			if (SortOrderPriorityA != SortOrderPriorityB)
-			{
-				return SortOrderPriorityA > SortOrderPriorityB;
-			}
-			
 			return A->GetDisplayName().ToString() < B->GetDisplayName().ToString();
 		});
 	}
diff --git a/Engine/Plugins/Animation/GameplayInsights/Source/RewindDebugger/Private/RewindDebuggerObjectTrack.h b/Engine/Plugins/Animation/GameplayInsights/Source/RewindDebugger/Private/RewindDebuggerObjectTrack.h
--- a/Engine/Plugins/Animation/GameplayInsights/Source/RewindDebugger/Private/RewindDebuggerObjectTrack.h
+++ b/Engine/Plugins/Animation/GameplayInsights/Source/RewindDebugger/Private/RewindDebuggerObjectTrack.h
@@ -35,7 +35,6 @@ private:
 	virtual uint64 GetObjectIdInternal() const override { return ObjectId; }
 	virtual bool HasDebugDataInternal() const override { return false; }
 	virtual bool HandleDoubleClickInternal() override;
-	virtual int GetSortOrderPriorityInternal() const override { return -10; };
 	
 	mutable FText DisplayName;
 	FString ObjectName;
diff --git a/Engine/Plugins/Animation/GameplayInsights/Source/RewindDebugger/Private/SRewindDebuggerDetails.cpp b/Engine/Plugins/Animation/GameplayInsights/Source/RewindDebugger/Private/SRewindDebuggerDetails.cpp
--- a/Engine/Plugins/Animation/GameplayInsights/Source/RewindDebugger/Private/SRewindDebuggerDetails.cpp
+++ b/Engine/Plugins/Animation/GameplayInsights/Source/RewindDebugger/Private/SRewindDebuggerDetails.cpp
@@ -25,7 +25,6 @@
 #include "Widgets/Input/SButton.h"
 #include "Widgets/Layout/SExpandableArea.h"
 #include "Widgets/Layout/SScrollBox.h"
-#include "IRewindDebuggerViewCreator.h"
 #include "RewindDebuggerViewCreators.h"
 #include "IRewindDebuggerDoubleClickHandler.h"
 
diff --git a/Engine/Source/Editor/RewindDebuggerInterface/Public/IRewindDebuggerTrackCreator.h b/Engine/Source/Editor/RewindDebuggerInterface/Public/IRewindDebuggerTrackCreator.h
--- a/Engine/Source/Editor/RewindDebuggerInterface/Public/IRewindDebuggerTrackCreator.h
+++ b/Engine/Source/Editor/RewindDebuggerInterface/Public/IRewindDebuggerTrackCreator.h
@@ -38,6 +38,12 @@ class REWINDDEBUGGERINTERFACE_API IRewindDebuggerTrackCreator : public IModularF
 		return GetNameInternal();
 	}
 
+	// returns an integer.  Higher values will show higher in the track list (default is 0)
+	int32 GetSortOrderPriority() const
+	{
+		return GetSortOrderPriorityInternal();
+	}
+
 	// optional additional filter, to prevent debug views from being listed if they have no data
 	bool HasDebugInfo(uint64 ObjectId) const
 	{
@@ -56,6 +62,8 @@ class REWINDDEBUGGERINTERFACE_API IRewindDebuggerTrackCreator : public IModularF
 	
 	virtual FName GetNameInternal() const { return FName(); }
 	
+	virtual int32 GetSortOrderPriorityInternal() const { return 0; }
+	
 	virtual bool HasDebugInfoInternal(uint64 ObjectId) const
 	{
 		return true;
diff --git a/Engine/Source/Editor/RewindDebuggerInterface/Public/RewindDebuggerTrack.h b/Engine/Source/Editor/RewindDebuggerInterface/Public/RewindDebuggerTrack.h
--- a/Engine/Source/Editor/RewindDebuggerInterface/Public/RewindDebuggerTrack.h
+++ b/Engine/Source/Editor/RewindDebuggerInterface/Public/RewindDebuggerTrack.h
@@ -2,8 +2,9 @@
 #pragma once
 
 #include "CoreMinimal.h"
-#include "IRewindDebuggerViewCreator.h"
 #include "IRewindDebugger.h"
+#include "Widgets/SWidget.h"
+#include "Textures/SlateIcon.h"
 
 namespace TraceServices
 {
@@ -59,12 +60,6 @@ public:
 	{
 		return GetNameInternal();
 	}
-	
-	// unique name for track (must match creator name if track is created by an IRewindDebuggerViewCreator) 
-	int GetSortOrderPriority() const
-	{
-		return GetSortOrderPriorityInternal();
-	}
 
 	// icon to display in the tree view
 	FSlateIcon GetIcon()
@@ -115,7 +110,6 @@ private:
 	virtual TSharedPtr<SWidget> GetDetailsViewInternal() { return TSharedPtr<SWidget>(); }
 	virtual FSlateIcon GetIconInternal() { return FSlateIcon(); }
 	virtual FName GetNameInternal() const { return ""; }
-	virtual int GetSortOrderPriorityInternal() const { return 0; }
 	virtual FText GetDisplayNameInternal() const { return FText(); }
 	virtual uint64 GetObjectIdInternal() const { return 0; }
 	virtual bool HasDebugDataInternal() const { return true; }
