Make Rewind Debugger links into the animation blueprint attach debugging to the correct instance when they open

diff --git a/Engine/Plugins/Animation/GameplayInsights/Source/GameplayInsights/Private/InertializationsTrack.cpp b/Engine/Plugins/Animation/GameplayInsights/Source/GameplayInsights/Private/InertializationsTrack.cpp
--- a/Engine/Plugins/Animation/GameplayInsights/Source/GameplayInsights/Private/InertializationsTrack.cpp
+++ b/Engine/Plugins/Animation/GameplayInsights/Source/GameplayInsights/Private/InertializationsTrack.cpp
@@ -14,6 +14,7 @@
 #include "IAnimationBlueprintEditor.h"
 #include "Animation/AnimBlueprint.h"
 #include "Subsystems/AssetEditorSubsystem.h"
+#include "ObjectTrace.h"
 #endif
 
 #define LOCTEXT_NAMESPACE "InertializationsTrack"
@@ -220,6 +221,11 @@ bool FInertializationTrack::HandleDoubleClickInternal()
 						{
 							GEditor->GetEditorSubsystem<UAssetEditorSubsystem>()->OpenEditorForAsset(AnimBlueprint);
 
+							if (UObject* SelectedInstance = FObjectTrace::GetObjectFromId(ObjectId))
+							{
+								AnimBlueprint->SetObjectBeingDebugged(SelectedInstance);
+							}
+
 							if (IAnimationBlueprintEditor* AnimBlueprintEditor = static_cast<IAnimationBlueprintEditor*>(GEditor->GetEditorSubsystem<UAssetEditorSubsystem>()->FindEditorForAsset(AnimBlueprint, true)))
 							{
 								int32 AnimNodeIndex = InstanceClass.Get()->GetAnimNodeProperties().Num() - NodeId - 1;
diff --git a/Engine/Plugins/Animation/GameplayInsights/Source/GameplayInsights/Private/SAnimGraphSchematicView.cpp b/Engine/Plugins/Animation/GameplayInsights/Source/GameplayInsights/Private/SAnimGraphSchematicView.cpp
--- a/Engine/Plugins/Animation/GameplayInsights/Source/GameplayInsights/Private/SAnimGraphSchematicView.cpp
+++ b/Engine/Plugins/Animation/GameplayInsights/Source/GameplayInsights/Private/SAnimGraphSchematicView.cpp
@@ -28,6 +28,7 @@
 #include "Subsystems/AssetEditorSubsystem.h"
 #include "Editor.h"
 #include "IAnimationBlueprintEditor.h"
+#include "ObjectTrace.h"
 #endif
 
 #define LOCTEXT_NAMESPACE "SAnimGraphSchematicView"
@@ -245,6 +246,7 @@ class SAnimGraphSchematicNode : public SMultiColumnTableRow<TSharedRef<FAnimGrap
 					const FObjectInfo* ObjectInfo = GameplayProvider->FindObjectInfo(Node->AnimInstanceId);
 					const FClassInfo* AnimInstanceClassInfo = ObjectInfo ? GameplayProvider->FindClassInfo(ObjectInfo->ClassId) : nullptr;
 					const int32 AnimNodeId = Node->NodeId;
+					const uint64 AnimInstanceId = Node->AnimInstanceId;
 
 					if (ObjectInfo && AnimInstanceClassInfo)
 					{
@@ -271,7 +273,7 @@ class SAnimGraphSchematicNode : public SMultiColumnTableRow<TSharedRef<FAnimGrap
 									.TextStyle(&FCoreStyle::Get().GetWidgetStyle<FTextBlockStyle>("SmallText"))
 									.ToolTipText(FText::Format(LOCTEXT("AssetHyperlinkTooltipFormat", "Open node '{0}'"), Node->Type))
 									.HighlightText(FilterText)
-									.OnNavigate_Lambda([AnimNodeId, AnimInstanceClassInfo]()
+									.OnNavigate_Lambda([AnimNodeId, AnimInstanceClassInfo, AnimInstanceId]()
 									{
 										TSoftObjectPtr<UAnimBlueprintGeneratedClass> InstanceClass;
 										InstanceClass = FSoftObjectPath(AnimInstanceClassInfo->PathName);
@@ -281,6 +283,11 @@ class SAnimGraphSchematicNode : public SMultiColumnTableRow<TSharedRef<FAnimGrap
 											if (UAnimBlueprint* AnimBlueprint = Cast<UAnimBlueprint>(InstanceClass.Get()->ClassGeneratedBy))
 											{
 												GEditor->GetEditorSubsystem<UAssetEditorSubsystem>()->OpenEditorForAsset(AnimBlueprint);
+												
+												if (UObject* SelectedInstance = FObjectTrace::GetObjectFromId(AnimInstanceId))
+												{
+													AnimBlueprint->SetObjectBeingDebugged(SelectedInstance);
+												}
 
 												if (IAnimationBlueprintEditor* AnimBlueprintEditor = static_cast<IAnimationBlueprintEditor*>(GEditor->GetEditorSubsystem<UAssetEditorSubsystem>()->FindEditorForAsset(AnimBlueprint, true)))
 												{
