Rewind Debugger Track Type filtering

diff --git a/Engine/Plugins/Animation/GameplayInsights/Source/GameplayInsights/Private/AnimCurvesTrack.cpp b/Engine/Plugins/Animation/GameplayInsights/Source/GameplayInsights/Private/AnimCurvesTrack.cpp
--- a/Engine/Plugins/Animation/GameplayInsights/Source/GameplayInsights/Private/AnimCurvesTrack.cpp
+++ b/Engine/Plugins/Animation/GameplayInsights/Source/GameplayInsights/Private/AnimCurvesTrack.cpp
@@ -226,6 +226,11 @@ FName FAnimationCurvesTrackCreator::GetNameInternal() const
 {
 	return AnimationCurvesName;
 }
+		
+void FAnimationCurvesTrackCreator::GetTrackTypesInternal(TArray<FRewindDebuggerTrackType>& Types) const
+{
+	Types.Add({AnimationCurvesName, LOCTEXT("Curves", "Curves")});
+}
 
 TSharedPtr<RewindDebugger::FRewindDebuggerTrack> FAnimationCurvesTrackCreator::CreateTrackInternal(uint64 ObjectId) const
 {
diff --git a/Engine/Plugins/Animation/GameplayInsights/Source/GameplayInsights/Private/AnimCurvesTrack.h b/Engine/Plugins/Animation/GameplayInsights/Source/GameplayInsights/Private/AnimCurvesTrack.h
--- a/Engine/Plugins/Animation/GameplayInsights/Source/GameplayInsights/Private/AnimCurvesTrack.h
+++ b/Engine/Plugins/Animation/GameplayInsights/Source/GameplayInsights/Private/AnimCurvesTrack.h
@@ -73,6 +73,7 @@ class FAnimationCurvesTrackCreator : public IRewindDebuggerTrackCreator
 private:
 	virtual FName GetTargetTypeNameInternal() const override;
 	virtual FName GetNameInternal() const override;
+	virtual void GetTrackTypesInternal(TArray<FRewindDebuggerTrackType>& Types) const override;
 	virtual TSharedPtr<FRewindDebuggerTrack> CreateTrackInternal(uint64 ObjectId) const override;
 	virtual bool HasDebugInfoInternal(uint64 ObjectId) const override;
 };
diff --git a/Engine/Plugins/Animation/GameplayInsights/Source/GameplayInsights/Private/BlendWeightsTrack.cpp b/Engine/Plugins/Animation/GameplayInsights/Source/GameplayInsights/Private/BlendWeightsTrack.cpp
--- a/Engine/Plugins/Animation/GameplayInsights/Source/GameplayInsights/Private/BlendWeightsTrack.cpp
+++ b/Engine/Plugins/Animation/GameplayInsights/Source/GameplayInsights/Private/BlendWeightsTrack.cpp
@@ -368,12 +368,18 @@ FName FBlendWeightsTrackCreator::GetTargetTypeNameInternal() const
 	static const FName TargetTypeName = "AnimInstance";
 	return TargetTypeName;
 }
+	
+static const FName BlendWeightsName("BlendWeights");
 
 FName FBlendWeightsTrackCreator::GetNameInternal() const
 {
-	static const FName BlendWeightsName("BlendWeights");
 	return BlendWeightsName;
 }
+	
+void FBlendWeightsTrackCreator::GetTrackTypesInternal(TArray<FRewindDebuggerTrackType>& Types) const
+{
+	Types.Add({BlendWeightsName, LOCTEXT("Blend Weights", "Blend Weights")});
+}
 
 TSharedPtr<RewindDebugger::FRewindDebuggerTrack> FBlendWeightsTrackCreator::CreateTrackInternal(uint64 ObjectId) const
 {
@@ -388,7 +394,7 @@ bool FBlendWeightsTrackCreator::HasDebugInfoInternal(uint64 ObjectId) const
 	bool bHasData = false;
 	if (const FAnimationProvider* AnimationProvider = AnalysisSession->ReadProvider<FAnimationProvider>(FAnimationProvider::ProviderName))
 	{
-		AnimationProvider->ReadAnimGraphTimeline(ObjectId, [&bHasData](const FAnimationProvider::AnimGraphTimeline& InGraphTimeline)
+		AnimationProvider->ReadTickRecordTimeline(ObjectId, [&bHasData](const FAnimationProvider::TickRecordTimeline& InGraphTimeline)
 		{
 			bHasData = true;
 		});
diff --git a/Engine/Plugins/Animation/GameplayInsights/Source/GameplayInsights/Private/BlendWeightsTrack.h b/Engine/Plugins/Animation/GameplayInsights/Source/GameplayInsights/Private/BlendWeightsTrack.h
--- a/Engine/Plugins/Animation/GameplayInsights/Source/GameplayInsights/Private/BlendWeightsTrack.h
+++ b/Engine/Plugins/Animation/GameplayInsights/Source/GameplayInsights/Private/BlendWeightsTrack.h
@@ -86,6 +86,7 @@ class FBlendWeightsTrackCreator : public IRewindDebuggerTrackCreator
 public:
 	virtual FName GetTargetTypeNameInternal() const;
 	virtual FName GetNameInternal() const override;
+	virtual void GetTrackTypesInternal(TArray<FRewindDebuggerTrackType>& Types) const override;
 	virtual TSharedPtr<FRewindDebuggerTrack> CreateTrackInternal(uint64 ObjectId) const override;
 	virtual bool HasDebugInfoInternal(uint64 ObjectId) const override;
 };
diff --git a/Engine/Plugins/Animation/GameplayInsights/Source/GameplayInsights/Private/ExternalMorphTrack.cpp b/Engine/Plugins/Animation/GameplayInsights/Source/GameplayInsights/Private/ExternalMorphTrack.cpp
--- a/Engine/Plugins/Animation/GameplayInsights/Source/GameplayInsights/Private/ExternalMorphTrack.cpp
+++ b/Engine/Plugins/Animation/GameplayInsights/Source/GameplayInsights/Private/ExternalMorphTrack.cpp
@@ -275,11 +275,17 @@ namespace RewindDebugger
 		return SkelMeshName;
 	}
 	
+	static const FName ExternalMorphSetName("ExternalMorphSetGroup");
+	
 	FName FExternalMorphSetGroupTrackCreator::GetNameInternal() const
 	{
-		static const FName ExternalMorphSetName("ExternalMorphSetGroup");
 		return ExternalMorphSetName;
 	}
+	
+	void FExternalMorphSetGroupTrackCreator::GetTrackTypesInternal(TArray<FRewindDebuggerTrackType>& Types) const 
+    {
+    	Types.Add({ExternalMorphSetName, LOCTEXT("External Morph Sets", "External Morph Sets")});
+    }
 
 	TSharedPtr<RewindDebugger::FRewindDebuggerTrack> FExternalMorphSetGroupTrackCreator::CreateTrackInternal(uint64 ObjectId) const
 	{
diff --git a/Engine/Plugins/Animation/GameplayInsights/Source/GameplayInsights/Private/ExternalMorphTrack.h b/Engine/Plugins/Animation/GameplayInsights/Source/GameplayInsights/Private/ExternalMorphTrack.h
--- a/Engine/Plugins/Animation/GameplayInsights/Source/GameplayInsights/Private/ExternalMorphTrack.h
+++ b/Engine/Plugins/Animation/GameplayInsights/Source/GameplayInsights/Private/ExternalMorphTrack.h
@@ -105,6 +105,7 @@ namespace RewindDebugger
 		// IRewindDebuggerTrackCreator overrides.
 		virtual FName GetTargetTypeNameInternal() const override;
 		virtual FName GetNameInternal() const override;
+		virtual void GetTrackTypesInternal(TArray<FRewindDebuggerTrackType>& Types) const override;
 		virtual TSharedPtr<FRewindDebuggerTrack> CreateTrackInternal(uint64 ObjectId) const override;
 		virtual bool HasDebugInfoInternal(uint64 ObjectId) const override;
 		// ~END IRewindDebuggerTrackCreator overrides.
diff --git a/Engine/Plugins/Animation/GameplayInsights/Source/GameplayInsights/Private/InertializationsTrack.cpp b/Engine/Plugins/Animation/GameplayInsights/Source/GameplayInsights/Private/InertializationsTrack.cpp
--- a/Engine/Plugins/Animation/GameplayInsights/Source/GameplayInsights/Private/InertializationsTrack.cpp
+++ b/Engine/Plugins/Animation/GameplayInsights/Source/GameplayInsights/Private/InertializationsTrack.cpp
@@ -254,12 +254,17 @@ FName FInertializationsTrackCreator::GetTargetTypeNameInternal() const
 	return TargetTypeName;
 }
 
+static const FName InertializationsName("Inertializations");
 FName FInertializationsTrackCreator::GetNameInternal() const
 {
-	static const FName InertializationsName("Inertializations");
 	return InertializationsName;
 }
 
+void FInertializationsTrackCreator::GetTrackTypesInternal(TArray<FRewindDebuggerTrackType>& Types) const 
+{
+	Types.Add({InertializationsName, LOCTEXT("Inertializations", "Inertializations")});
+}
+
 TSharedPtr<RewindDebugger::FRewindDebuggerTrack> FInertializationsTrackCreator::CreateTrackInternal(uint64 ObjectId) const
 {
  	return MakeShared<RewindDebugger::FInertializationsTrack>(ObjectId);
diff --git a/Engine/Plugins/Animation/GameplayInsights/Source/GameplayInsights/Private/InertializationsTrack.h b/Engine/Plugins/Animation/GameplayInsights/Source/GameplayInsights/Private/InertializationsTrack.h
--- a/Engine/Plugins/Animation/GameplayInsights/Source/GameplayInsights/Private/InertializationsTrack.h
+++ b/Engine/Plugins/Animation/GameplayInsights/Source/GameplayInsights/Private/InertializationsTrack.h
@@ -67,6 +67,7 @@ class FInertializationsTrackCreator : public IRewindDebuggerTrackCreator
 public:
 	virtual FName GetTargetTypeNameInternal() const;
 	virtual FName GetNameInternal() const override;
+	virtual void GetTrackTypesInternal(TArray<FRewindDebuggerTrackType>& Types) const override;
 	virtual TSharedPtr<FRewindDebuggerTrack> CreateTrackInternal(uint64 ObjectId) const override;
 	virtual bool HasDebugInfoInternal(uint64 ObjectId) const override;
 };
diff --git a/Engine/Plugins/Animation/GameplayInsights/Source/GameplayInsights/Private/MontagesTrack.cpp b/Engine/Plugins/Animation/GameplayInsights/Source/GameplayInsights/Private/MontagesTrack.cpp
--- a/Engine/Plugins/Animation/GameplayInsights/Source/GameplayInsights/Private/MontagesTrack.cpp
+++ b/Engine/Plugins/Animation/GameplayInsights/Source/GameplayInsights/Private/MontagesTrack.cpp
@@ -323,6 +323,11 @@ FName FMontagesTrackCreator::GetNameInternal() const
 	return MontagesName;
 }
 	
+void FMontagesTrackCreator::GetTrackTypesInternal(TArray<FRewindDebuggerTrackType>& Types) const 
+{
+	Types.Add({MontagesName, LOCTEXT("Montages", "Montages")});
+}
+	
 FName FMontagesTrack::GetNameInternal() const
 {
 	return MontagesName;
diff --git a/Engine/Plugins/Animation/GameplayInsights/Source/GameplayInsights/Private/MontagesTrack.h b/Engine/Plugins/Animation/GameplayInsights/Source/GameplayInsights/Private/MontagesTrack.h
--- a/Engine/Plugins/Animation/GameplayInsights/Source/GameplayInsights/Private/MontagesTrack.h
+++ b/Engine/Plugins/Animation/GameplayInsights/Source/GameplayInsights/Private/MontagesTrack.h
@@ -79,6 +79,7 @@ class FMontagesTrackCreator : public IRewindDebuggerTrackCreator
 public:
 	virtual FName GetTargetTypeNameInternal() const;
 	virtual FName GetNameInternal() const override;
+	virtual void GetTrackTypesInternal(TArray<FRewindDebuggerTrackType>& Types) const override;
 	virtual TSharedPtr<FRewindDebuggerTrack> CreateTrackInternal(uint64 ObjectId) const override;
 	virtual bool HasDebugInfoInternal(uint64 ObjectId) const override;
 };
diff --git a/Engine/Plugins/Animation/GameplayInsights/Source/GameplayInsights/Private/NotifiesTrack.cpp b/Engine/Plugins/Animation/GameplayInsights/Source/GameplayInsights/Private/NotifiesTrack.cpp
--- a/Engine/Plugins/Animation/GameplayInsights/Source/GameplayInsights/Private/NotifiesTrack.cpp
+++ b/Engine/Plugins/Animation/GameplayInsights/Source/GameplayInsights/Private/NotifiesTrack.cpp
@@ -163,6 +163,12 @@ FName FNotifiesTrackCreator::GetNameInternal() const
 	return NotifiesName;
 }
 
+void FNotifiesTrackCreator::GetTrackTypesInternal(TArray<FRewindDebuggerTrackType>& Types) const
+{
+	Types.Add({ NotifiesName, LOCTEXT("Notifies", "Notifies") });
+}
+
+
 TSharedPtr<RewindDebugger::FRewindDebuggerTrack> FNotifiesTrackCreator::CreateTrackInternal(uint64 ObjectId) const
 {
 	return MakeShared<RewindDebugger::FNotifiesTrack>(ObjectId);
diff --git a/Engine/Plugins/Animation/GameplayInsights/Source/GameplayInsights/Private/NotifiesTrack.h b/Engine/Plugins/Animation/GameplayInsights/Source/GameplayInsights/Private/NotifiesTrack.h
--- a/Engine/Plugins/Animation/GameplayInsights/Source/GameplayInsights/Private/NotifiesTrack.h
+++ b/Engine/Plugins/Animation/GameplayInsights/Source/GameplayInsights/Private/NotifiesTrack.h
@@ -101,6 +101,7 @@ class FNotifiesTrackCreator : public IRewindDebuggerTrackCreator
 private:
 	virtual FName GetTargetTypeNameInternal() const;
 	virtual FName GetNameInternal() const override;
+	virtual void GetTrackTypesInternal(TArray<FRewindDebuggerTrackType>& Types) const override;
 	virtual TSharedPtr<FRewindDebuggerTrack> CreateTrackInternal(uint64 ObjectId) const override;
 	virtual bool HasDebugInfoInternal(uint64 ObjectId) const override;
 };
diff --git a/Engine/Plugins/Animation/GameplayInsights/Source/GameplayInsights/Private/PoseWatchTrack.cpp b/Engine/Plugins/Animation/GameplayInsights/Source/GameplayInsights/Private/PoseWatchTrack.cpp
--- a/Engine/Plugins/Animation/GameplayInsights/Source/GameplayInsights/Private/PoseWatchTrack.cpp
+++ b/Engine/Plugins/Animation/GameplayInsights/Source/GameplayInsights/Private/PoseWatchTrack.cpp
@@ -247,11 +247,16 @@ FName FPoseWatchesTrackCreator::GetTargetTypeNameInternal() const
 	return AnimInstanceName;
 }
 	
+static const FName PoseWatchesName("PoseWatches");
 FName FPoseWatchesTrackCreator::GetNameInternal() const
 {
-	static const FName PoseWatchesName("PoseWatches");
 	return PoseWatchesName;
 }
+		
+void FPoseWatchesTrackCreator::GetTrackTypesInternal(TArray<FRewindDebuggerTrackType>& Types) const 
+{
+	Types.Add({PoseWatchesName, LOCTEXT("Pose Watches", "Pose Watches")});
+}
 
 TSharedPtr<RewindDebugger::FRewindDebuggerTrack> FPoseWatchesTrackCreator::CreateTrackInternal(uint64 ObjectId) const
 {
diff --git a/Engine/Plugins/Animation/GameplayInsights/Source/GameplayInsights/Private/PoseWatchTrack.h b/Engine/Plugins/Animation/GameplayInsights/Source/GameplayInsights/Private/PoseWatchTrack.h
--- a/Engine/Plugins/Animation/GameplayInsights/Source/GameplayInsights/Private/PoseWatchTrack.h
+++ b/Engine/Plugins/Animation/GameplayInsights/Source/GameplayInsights/Private/PoseWatchTrack.h
@@ -89,6 +89,7 @@ class FPoseWatchesTrackCreator : public IRewindDebuggerTrackCreator
 private:
 	virtual FName GetTargetTypeNameInternal() const;
 	virtual FName GetNameInternal() const override;
+	virtual void GetTrackTypesInternal(TArray<FRewindDebuggerTrackType>& Types) const override;
 	virtual TSharedPtr<FRewindDebuggerTrack> CreateTrackInternal(uint64 ObjectId) const override;
 	virtual bool HasDebugInfoInternal(uint64 ObjectId) const override;
 };
diff --git a/Engine/Plugins/Animation/GameplayInsights/Source/GameplayInsights/Private/PropertiesTrack.cpp b/Engine/Plugins/Animation/GameplayInsights/Source/GameplayInsights/Private/PropertiesTrack.cpp
--- a/Engine/Plugins/Animation/GameplayInsights/Source/GameplayInsights/Private/PropertiesTrack.cpp
+++ b/Engine/Plugins/Animation/GameplayInsights/Source/GameplayInsights/Private/PropertiesTrack.cpp
@@ -290,12 +290,18 @@ namespace RewindDebugger
 		return TargetTypeName;
 	}
 
+	static const FName PropertiesTrackCreatorName("PropertyWatch");
 	FName FPropertiesTrackCreator::GetNameInternal() const
 	{
-		static const FName PropertiesTrackCreatorName("PropertyWatch");
 		return PropertiesTrackCreatorName;
 	}
 
+	void FPropertiesTrackCreator::GetTrackTypesInternal(TArray<FRewindDebuggerTrackType>& Types) const 
+    {
+    	Types.Add({PropertiesTrackCreatorName, LOCTEXT("Properties", "Properties")});
+    }
+
+
 	TSharedPtr<FRewindDebuggerTrack> FPropertiesTrackCreator::CreateTrackInternal(uint64 ObjectId) const
 	{
 		TSharedRef<FPropertiesTrack> NewTrack = MakeShared<FPropertiesTrack>(ObjectId);
diff --git a/Engine/Plugins/Animation/GameplayInsights/Source/GameplayInsights/Private/PropertiesTrack.h b/Engine/Plugins/Animation/GameplayInsights/Source/GameplayInsights/Private/PropertiesTrack.h
--- a/Engine/Plugins/Animation/GameplayInsights/Source/GameplayInsights/Private/PropertiesTrack.h
+++ b/Engine/Plugins/Animation/GameplayInsights/Source/GameplayInsights/Private/PropertiesTrack.h
@@ -75,6 +75,7 @@ namespace RewindDebugger
 		/** Begin IRewindDebuggerTrackCreator interface */
 		virtual FName GetTargetTypeNameInternal() const override;
 		virtual FName GetNameInternal() const override;
+		virtual void GetTrackTypesInternal(TArray<FRewindDebuggerTrackType>& Types) const override;
 		virtual TSharedPtr<FRewindDebuggerTrack> CreateTrackInternal(uint64 ObjectId) const override;
 		virtual bool HasDebugInfoInternal(uint64 ObjectId) const override;
 		/** End IRewindDebuggerTrackCreator interface */
diff --git a/Engine/Plugins/Animation/GameplayInsights/Source/GameplayInsights/Private/SAnimGraphSchematicView.cpp b/Engine/Plugins/Animation/GameplayInsights/Source/GameplayInsights/Private/SAnimGraphSchematicView.cpp
--- a/Engine/Plugins/Animation/GameplayInsights/Source/GameplayInsights/Private/SAnimGraphSchematicView.cpp
+++ b/Engine/Plugins/Animation/GameplayInsights/Source/GameplayInsights/Private/SAnimGraphSchematicView.cpp
@@ -957,6 +957,11 @@ FName FAnimGraphSchematicTrackCreator::GetNameInternal() const
 	return AnimGraphSchematicName;
 }
 
+void FAnimGraphSchematicTrackCreator::GetTrackTypesInternal(TArray<RewindDebugger::FRewindDebuggerTrackType>& Types) const
+{
+	Types.Add({AnimGraphSchematicName, LOCTEXT("Anim Graph", "Anim Graph")});
+}
+
 TSharedPtr<RewindDebugger::FRewindDebuggerTrack> FAnimGraphSchematicTrackCreator::CreateTrackInternal(uint64 ObjectId) const
 {
 	return MakeShared<FAnimGraphSchematicTrack>(ObjectId);
diff --git a/Engine/Plugins/Animation/GameplayInsights/Source/GameplayInsights/Private/SAnimGraphSchematicView.h b/Engine/Plugins/Animation/GameplayInsights/Source/GameplayInsights/Private/SAnimGraphSchematicView.h
--- a/Engine/Plugins/Animation/GameplayInsights/Source/GameplayInsights/Private/SAnimGraphSchematicView.h
+++ b/Engine/Plugins/Animation/GameplayInsights/Source/GameplayInsights/Private/SAnimGraphSchematicView.h
@@ -144,6 +144,7 @@ class FAnimGraphSchematicTrackCreator : public RewindDebugger::IRewindDebuggerTr
 public:
 	virtual FName GetTargetTypeNameInternal() const override;
 	virtual FName GetNameInternal() const override;
+	virtual void GetTrackTypesInternal(TArray<RewindDebugger::FRewindDebuggerTrackType>& Types) const override;
 	virtual TSharedPtr<RewindDebugger::FRewindDebuggerTrack> CreateTrackInternal(uint64 ObjectId) const override;
 	virtual bool HasDebugInfoInternal(uint64 ObjectId) const override;
 };
\ No newline at end of file
diff --git a/Engine/Plugins/Animation/GameplayInsights/Source/RewindDebugger/Private/RewindDebugger.cpp b/Engine/Plugins/Animation/GameplayInsights/Source/RewindDebugger/Private/RewindDebugger.cpp
--- a/Engine/Plugins/Animation/GameplayInsights/Source/RewindDebugger/Private/RewindDebugger.cpp
+++ b/Engine/Plugins/Animation/GameplayInsights/Source/RewindDebugger/Private/RewindDebugger.cpp
@@ -33,6 +33,7 @@
 #include "Widgets/Docking/SDockTab.h"
 #include "Widgets/Layout/SSpacer.h"
 #include "RewindDebuggerCommands.h"
+#include "RewindDebuggerTrackCreators.h"
 
 #define LOCTEXT_NAMESPACE "RewindDebugger"
 
@@ -73,6 +74,11 @@ FRewindDebugger::FRewindDebugger()  :
 	bTargetActorPositionValid(false),
 	bIsDetailsPanelOpen(true)
 {
+	RewindDebugger::FRewindDebuggerTrackCreators::EnumerateCreators([this](const RewindDebugger::IRewindDebuggerTrackCreator* Creator)
+    {
+		Creator->GetTrackTypes(TrackTypes);
+    });
+	
 	RecordingDuration.Set(0);
 
 	if (GEditor->bIsSimulatingInEditor || GEditor->PlayWorld)
@@ -1023,6 +1029,22 @@ void FRewindDebugger::Tick(float DeltaTime)
 															TRACE_CPUPROFILER_EVENT_SCOPE(ResetNodeVisitStates);
 															DebugData.ResetNodeVisitSites();
 														}
+														
+														// Anim node values can come from all phases
+														AnimationProvider->ReadAnimNodeValuesTimeline(ObjectId, [&Frame,AnimationProvider, &DebugData](const IAnimationProvider::AnimNodeValuesTimeline& InNodeValuesTimeline)
+														{
+															TRACE_CPUPROFILER_EVENT_SCOPE(AnimGraphNodeValues);
+															InNodeValuesTimeline.EnumerateEvents(Frame.StartTime, Frame.EndTime, [AnimationProvider, &DebugData](double InStartTime, double InEndTime, uint32 InDepth, const FAnimNodeValueMessage& InMessage)
+															{
+																// don't send "Name" Node value for display in the graph
+																if (FPlatformString::Strcmp(InMessage.Key, TEXT("Name")) != 0)
+																{
+																	FText Text = AnimationProvider->FormatNodeKeyValue(InMessage);
+																	DebugData.RecordNodeValue(InMessage.NodeId, Text.ToString());
+																}
+																return TraceServices::EEventEnumerate::Continue;
+															});
+														});
 
 														DebugData.DisableAllPoseWatches();
 							
@@ -1146,21 +1168,6 @@ void FRewindDebugger::Tick(float DeltaTime)
 
 																}
 							
-																// Anim node values can come from all phases
-																AnimationProvider->ReadAnimNodeValuesTimeline(Id, [InGraphStartTime, InGraphEndTime, AnimationProvider, &DebugData](const IAnimationProvider::AnimNodeValuesTimeline& InNodeValuesTimeline)
-																{
-																	TRACE_CPUPROFILER_EVENT_SCOPE(AnimGraphNodeValues);
-																	InNodeValuesTimeline.EnumerateEvents(InGraphStartTime, InGraphEndTime, [AnimationProvider, &DebugData](double InStartTime, double InEndTime, uint32 InDepth, const FAnimNodeValueMessage& InMessage)
-																	{
-																		// don't send "Name" Node value for display in the graph
-																		if (FPlatformString::Strcmp(InMessage.Key, TEXT("Name")) != 0)
-																		{
-																			FText Text = AnimationProvider->FormatNodeKeyValue(InMessage);
-																			DebugData.RecordNodeValue(InMessage.NodeId, Text.ToString());
-																		}
-																		return TraceServices::EEventEnumerate::Continue;
-																	});
-																});
 															}
 															return TraceServices::EEventEnumerate::Continue;
 														});
diff --git a/Engine/Plugins/Animation/GameplayInsights/Source/RewindDebugger/Private/RewindDebugger.h b/Engine/Plugins/Animation/GameplayInsights/Source/RewindDebugger/Private/RewindDebugger.h
--- a/Engine/Plugins/Animation/GameplayInsights/Source/RewindDebugger/Private/RewindDebugger.h
+++ b/Engine/Plugins/Animation/GameplayInsights/Source/RewindDebugger/Private/RewindDebugger.h
@@ -5,6 +5,7 @@
 #include "IRewindDebugger.h"
 #include "RewindDebuggerTrack.h"
 #include "BindableProperty.h"
+#include "IRewindDebuggerTrackCreator.h"
 #include "Containers/Ticker.h"
 #include "UObject/WeakObjectPtr.h"
 
@@ -125,6 +126,8 @@ public:
 	void SetIsDetailsPanelOpen(bool bIsOpen) { bIsDetailsPanelOpen = bIsOpen; }
 	bool IsDetailsPanelOpen(bool bIsOpen) { return bIsDetailsPanelOpen; }
 
+	TArrayView<RewindDebugger::FRewindDebuggerTrackType> GetTrackTypes() { return TrackTypes; };
+
 private:
 	void RefreshDebugComponents(TArray<TSharedPtr<RewindDebugger::FRewindDebuggerTrack>>& InTracks, TArray<TSharedPtr<FDebugObjectInfo>>& OutComponents);
 	
@@ -198,5 +201,7 @@ private:
 	bool bTargetActorPositionValid;
 	FVector TargetActorPosition;
 
+	TArray<RewindDebugger::FRewindDebuggerTrackType> TrackTypes;
+
 	bool bIsDetailsPanelOpen;
 };
diff --git a/Engine/Plugins/Animation/GameplayInsights/Source/RewindDebugger/Private/RewindDebuggerModule.cpp b/Engine/Plugins/Animation/GameplayInsights/Source/RewindDebugger/Private/RewindDebuggerModule.cpp
--- a/Engine/Plugins/Animation/GameplayInsights/Source/RewindDebugger/Private/RewindDebuggerModule.cpp
+++ b/Engine/Plugins/Animation/GameplayInsights/Source/RewindDebugger/Private/RewindDebuggerModule.cpp
@@ -171,6 +171,7 @@ TSharedRef<SDockTab> FRewindDebuggerModule::SpawnRewindDebuggerTab(const FSpawnT
 								.OnComponentDoubleClicked_Raw(DebuggerInstance, &FRewindDebugger::ComponentDoubleClicked)
 								.OnComponentSelectionChanged_Raw(DebuggerInstance, &FRewindDebugger::ComponentSelectionChanged)
 								.BuildComponentContextMenu_Raw(DebuggerInstance, &FRewindDebugger::BuildComponentContextMenu)
+								.TrackTypes_Lambda([]() { return FRewindDebugger::Instance()->GetTrackTypes(); })
 								.ScrubTime_Lambda([]() { return FRewindDebugger::Instance()->GetScrubTime(); });
 
 	DebuggerInstance->OnTrackCursor(FRewindDebugger::FOnTrackCursor::CreateSP(RewindDebuggerWidget.Get(), &SRewindDebugger::TrackCursor));
diff --git a/Engine/Plugins/Animation/GameplayInsights/Source/RewindDebugger/Private/RewindDebuggerSettings.h b/Engine/Plugins/Animation/GameplayInsights/Source/RewindDebugger/Private/RewindDebuggerSettings.h
--- a/Engine/Plugins/Animation/GameplayInsights/Source/RewindDebugger/Private/RewindDebuggerSettings.h
+++ b/Engine/Plugins/Animation/GameplayInsights/Source/RewindDebugger/Private/RewindDebuggerSettings.h
@@ -50,6 +50,10 @@ public:
 	UPROPERTY(EditAnywhere, Config, Category = Filters)
 	bool bShowEmptyObjectTracks;
 	
+	/** The track types listed here will be hidden from the track tree view */
+	UPROPERTY(EditAnywhere, Config, Category = Filters)
+	TArray<FName> HiddenTrackTypes;
+	
 	/** Get Mutable CDO of URewindDebuggerSettings */
 	static URewindDebuggerSettings & Get();
 };
diff --git a/Engine/Plugins/Animation/GameplayInsights/Source/RewindDebugger/Private/SRewindDebugger.cpp b/Engine/Plugins/Animation/GameplayInsights/Source/RewindDebugger/Private/SRewindDebugger.cpp
--- a/Engine/Plugins/Animation/GameplayInsights/Source/RewindDebugger/Private/SRewindDebugger.cpp
+++ b/Engine/Plugins/Animation/GameplayInsights/Source/RewindDebugger/Private/SRewindDebugger.cpp
@@ -4,6 +4,7 @@
 
 #include "ActorPickerMode.h"
 #include "Editor.h"
+#include "IRewindDebuggerTrackCreator.h"
 #include "Editor/EditorEngine.h"
 #include "Framework/Application/SlateApplication.h"
 #include "Framework/MultiBox/MultiBoxBuilder.h"
@@ -90,6 +91,27 @@ void SRewindDebugger::SetViewRange(TRange<double> NewRange)
 	OnViewRangeChanged.ExecuteIfBound(NewRange);
 }
 
+void SRewindDebugger::ToggleHideTrackType(const FName& TrackType)
+{
+	int32 Index = Settings.HiddenTrackTypes.Find(TrackType);
+
+	if (Index >=0)
+	{
+		Settings.HiddenTrackTypes.RemoveAtSwap(Index);
+	}
+	else
+	{
+		Settings.HiddenTrackTypes.Add(TrackType);
+	}
+	RefreshDebugComponents();
+}
+
+bool SRewindDebugger::ShouldHideTrackType(const FName& TrackType) const
+{
+	return Settings.HiddenTrackTypes.Contains(TrackType);
+}
+
+
 void SRewindDebugger::ToggleDisplayEmptyTracks()
 {
 	Settings.bShowEmptyObjectTracks = !Settings.bShowEmptyObjectTracks;
@@ -227,6 +249,7 @@ void SRewindDebugger::Construct(const FArguments& InArgs, TSharedRef<FUICommandL
 	OnViewRangeChanged = InArgs._OnViewRangeChanged;
 	OnComponentSelectionChanged = InArgs._OnComponentSelectionChanged;
 	BuildComponentContextMenu = InArgs._BuildComponentContextMenu;
+	TrackTypesAttribute = InArgs._TrackTypes;
 	ScrubTimeAttribute = InArgs._ScrubTime;
 	DebugComponents = InArgs._DebugComponents;
 	TraceTime.Initialize(InArgs._TraceTime);
@@ -457,17 +480,24 @@ void SRewindDebugger::Construct(const FArguments& InArgs, TSharedRef<FUICommandL
 	];
 }
 
-bool FilterTrack(TSharedPtr<RewindDebugger::FRewindDebuggerTrack>& Track, const FString& FilterString, bool bRemoveNoData, bool bParentFilterPassed = false)
+bool FilterTrack(TSharedPtr<RewindDebugger::FRewindDebuggerTrack>& Track, const FString& FilterString, bool bRemoveNoData, const TArray<FName>& FilteredTrackTypes, bool bParentFilterPassed = false)
 {
+	if (FilteredTrackTypes.Contains(Track->GetName()))
+	{
+		Track->SetIsVisible(false);
+		return false;
+	}
+	
 	const bool bStringFilterEmpty =  FilterString.IsEmpty();
 	const bool bStringFilterPassed = bParentFilterPassed || bStringFilterEmpty || Track->GetDisplayName().ToString().Contains(FilterString);
 
 	const bool bThisFilterPassed = (!bStringFilterEmpty && bStringFilterPassed);
 
+
 	bool bAnyChildVisible = false;
-	Track->IterateSubTracks([&bAnyChildVisible, bThisFilterPassed, FilterString, bRemoveNoData](TSharedPtr<RewindDebugger::FRewindDebuggerTrack> ChildTrack)
+	Track->IterateSubTracks([&bAnyChildVisible, bThisFilterPassed, FilterString, bRemoveNoData, FilteredTrackTypes](TSharedPtr<RewindDebugger::FRewindDebuggerTrack> ChildTrack)
 	{
-		const bool bChildIsVisible = FilterTrack(ChildTrack, FilterString, bRemoveNoData, bThisFilterPassed);
+		const bool bChildIsVisible = FilterTrack(ChildTrack, FilterString, bRemoveNoData, FilteredTrackTypes, bThisFilterPassed);
 		bAnyChildVisible |= bChildIsVisible;
 	});
 
@@ -483,7 +513,7 @@ void SRewindDebugger::RefreshDebugComponents()
 	{
 		for(TSharedPtr<RewindDebugger::FRewindDebuggerTrack>& DebugComponent : *DebugComponents)
 		{
-			FilterTrack(DebugComponent, TrackFilterBox->GetText().ToString(), !ShouldDisplayEmptyTracks());
+			FilterTrack(DebugComponent, TrackFilterBox->GetText().ToString(), !ShouldDisplayEmptyTracks(), Settings.HiddenTrackTypes);
 		}
 	}
 	
@@ -500,8 +530,27 @@ TSharedRef<SWidget> SRewindDebugger::MakeFilterMenu()
 {
 	FMenuBuilder Builder(true, nullptr);
 	Builder.AddWidget(TrackFilterBox.ToSharedRef(), FText(), true, false);
-	Builder.AddSeparator();
 	
+	Builder.BeginSection("TrackTypes", LOCTEXT("Track Types", "Track Types"));
+
+	const TArrayView<RewindDebugger::FRewindDebuggerTrackType> TrackTypes = TrackTypesAttribute.Get();
+
+	for(const RewindDebugger::FRewindDebuggerTrackType& TrackType : TrackTypes)
+	{
+		Builder.AddMenuEntry(
+			TrackType.DisplayName,
+			FText::Format(LOCTEXT("FilterTrackToolTip", "Show tracks of type {0}"), {TrackType.DisplayName}),
+			FSlateIcon(),
+			FUIAction(FExecuteAction::CreateLambda([this,TrackType](){ ToggleHideTrackType(TrackType.Name); }),
+				FCanExecuteAction(),
+				FIsActionChecked::CreateLambda([this,TrackType]() { return !ShouldHideTrackType(TrackType.Name); })),
+			NAME_None,
+			EUserInterfaceActionType::ToggleButton
+		);	
+	}
+
+	Builder.AddSeparator();
+
 	Builder.AddMenuEntry(
 		LOCTEXT("DisplayEmptyTracks", "Show Empty Object Tracks"),
 		LOCTEXT("DisplayEmptyTracksToolTip", "Show Object/Component tracks which have no sub tracks with any debug data"),
diff --git a/Engine/Plugins/Animation/GameplayInsights/Source/RewindDebugger/Private/SRewindDebugger.h b/Engine/Plugins/Animation/GameplayInsights/Source/RewindDebugger/Private/SRewindDebugger.h
--- a/Engine/Plugins/Animation/GameplayInsights/Source/RewindDebugger/Private/SRewindDebugger.h
+++ b/Engine/Plugins/Animation/GameplayInsights/Source/RewindDebugger/Private/SRewindDebugger.h
@@ -3,6 +3,7 @@
 
 #include "BindableProperty.h"
 #include "CoreMinimal.h"
+#include "IRewindDebuggerTrackCreator.h"
 #include "Framework/Commands/UICommandList.h"
 #include "IRewindDebuggerView.h"
 #include "RewindDebuggerCommands.h"
@@ -36,6 +37,7 @@ public:
 		SLATE_ARGUMENT(DebugTargetInitializer, DebugTargetActor);
 		SLATE_ARGUMENT(TBindablePropertyInitializer<double>, TraceTime);
 		SLATE_ARGUMENT(TBindablePropertyInitializer<float>, RecordingDuration);
+		SLATE_ATTRIBUTE(TArrayView<RewindDebugger::FRewindDebuggerTrackType>, TrackTypes);
 		SLATE_ATTRIBUTE(double, ScrubTime);
 		SLATE_ATTRIBUTE(bool, IsPIESimulating);
 		SLATE_EVENT( FOnScrubPositionChanged, OnScrubPositionChanged);
@@ -69,7 +71,9 @@ private:
 	void SetViewRange(TRange<double> NewRange);
 
 	void MainFrameCreationFinished(TSharedPtr<SWindow> InRootWindow, bool bIsNewProjectWindow);
-
+	
+	void ToggleHideTrackType(const FName& TrackType);
+	bool ShouldHideTrackType(const FName& TrackType) const;
 	void ToggleDisplayEmptyTracks();
 	bool ShouldDisplayEmptyTracks() const;
 	
@@ -80,6 +84,7 @@ private:
 	TAttribute<double> ScrubTimeAttribute;
 	TAttribute<bool> IsPIESimulating;
 	TAttribute<bool> TrackScrubbingAttribute;
+	TAttribute<TArrayView<RewindDebugger::FRewindDebuggerTrackType>> TrackTypesAttribute;
 	FOnScrubPositionChanged OnScrubPositionChanged;
 	FOnViewRangeChanged OnViewRangeChanged;
 	TRange<double> ViewRange;
diff --git a/Engine/Plugins/Animation/GameplayInsights/Source/RewindDebuggerVLog/Private/VisualLogTrack.cpp b/Engine/Plugins/Animation/GameplayInsights/Source/RewindDebuggerVLog/Private/VisualLogTrack.cpp
--- a/Engine/Plugins/Animation/GameplayInsights/Source/RewindDebuggerVLog/Private/VisualLogTrack.cpp
+++ b/Engine/Plugins/Animation/GameplayInsights/Source/RewindDebuggerVLog/Private/VisualLogTrack.cpp
@@ -176,6 +176,11 @@ FName FVisualLogTrackCreator::GetNameInternal() const
 {
 	return VisualLogName;
 }
+	
+void FVisualLogTrackCreator::GetTrackTypesInternal(TArray<FRewindDebuggerTrackType>& Types) const
+{
+	Types.Add({VisualLogName, LOCTEXT("Visual Logging", "Visual Logging")});
+}
 
 TSharedPtr<RewindDebugger::FRewindDebuggerTrack> FVisualLogTrackCreator::CreateTrackInternal(uint64 ObjectId) const
 {
diff --git a/Engine/Plugins/Animation/GameplayInsights/Source/RewindDebuggerVLog/Private/VisualLogTrack.h b/Engine/Plugins/Animation/GameplayInsights/Source/RewindDebuggerVLog/Private/VisualLogTrack.h
--- a/Engine/Plugins/Animation/GameplayInsights/Source/RewindDebuggerVLog/Private/VisualLogTrack.h
+++ b/Engine/Plugins/Animation/GameplayInsights/Source/RewindDebuggerVLog/Private/VisualLogTrack.h
@@ -81,6 +81,7 @@ class FVisualLogTrackCreator : public IRewindDebuggerTrackCreator
 private:
 	virtual FName GetTargetTypeNameInternal() const override;
 	virtual FName GetNameInternal() const override;
+	virtual void GetTrackTypesInternal(TArray<FRewindDebuggerTrackType>& Types) const override;
 	virtual TSharedPtr<FRewindDebuggerTrack> CreateTrackInternal(uint64 ObjectId) const override;
 	virtual bool HasDebugInfoInternal(uint64 ObjectId) const override;
 };
diff --git a/Engine/Plugins/Experimental/Animation/PoseSearch/Source/Editor/Private/PoseSearchDebugger.cpp b/Engine/Plugins/Experimental/Animation/PoseSearch/Source/Editor/Private/PoseSearchDebugger.cpp
--- a/Engine/Plugins/Experimental/Animation/PoseSearch/Source/Editor/Private/PoseSearchDebugger.cpp
+++ b/Engine/Plugins/Experimental/Animation/PoseSearch/Source/Editor/Private/PoseSearchDebugger.cpp
@@ -428,11 +428,17 @@ FName FDebuggerTrackCreator::GetTargetTypeNameInternal() const
 	static FName TargetTypeName = "AnimInstance";
 	return TargetTypeName;
 }
+	
+static const FName PoseSearchDebuggerName("PoseSearchDebugger");
 
 FName FDebuggerTrackCreator::GetNameInternal() const
 {
-	static const FName Name("PoseSearchDebugger");
-	return Name;
+	return PoseSearchDebuggerName;
+}
+
+void FDebuggerTrackCreator::GetTrackTypesInternal(TArray<RewindDebugger::FRewindDebuggerTrackType>& Types) const
+{
+	Types.Add({PoseSearchDebuggerName, LOCTEXT("Pose Search", "Pose Search")});
 }
 
 TSharedPtr<RewindDebugger::FRewindDebuggerTrack> FDebuggerTrackCreator::CreateTrackInternal(uint64 ObjectId) const
diff --git a/Engine/Plugins/Experimental/Animation/PoseSearch/Source/Editor/Private/PoseSearchDebugger.h b/Engine/Plugins/Experimental/Animation/PoseSearch/Source/Editor/Private/PoseSearchDebugger.h
--- a/Engine/Plugins/Experimental/Animation/PoseSearch/Source/Editor/Private/PoseSearchDebugger.h
+++ b/Engine/Plugins/Experimental/Animation/PoseSearch/Source/Editor/Private/PoseSearchDebugger.h
@@ -91,6 +91,7 @@ class FDebuggerTrackCreator : public RewindDebugger::IRewindDebuggerTrackCreator
 private:
 	virtual FName GetTargetTypeNameInternal() const override;
 	virtual FName GetNameInternal() const override;
+	virtual void GetTrackTypesInternal(TArray<RewindDebugger::FRewindDebuggerTrackType>& Types) const override;
 	virtual TSharedPtr<RewindDebugger::FRewindDebuggerTrack> CreateTrackInternal(uint64 ObjectId) const override;
 	virtual bool HasDebugInfoInternal(uint64 ObjectId) const override;
 	virtual int32 GetSortOrderPriorityInternal() const override { return 10; };
diff --git a/Engine/Plugins/Experimental/Chooser/Source/ChooserEditor/Private/ChooserTrack.cpp b/Engine/Plugins/Experimental/Chooser/Source/ChooserEditor/Private/ChooserTrack.cpp
--- a/Engine/Plugins/Experimental/Chooser/Source/ChooserEditor/Private/ChooserTrack.cpp
+++ b/Engine/Plugins/Experimental/Chooser/Source/ChooserEditor/Private/ChooserTrack.cpp
@@ -151,6 +151,11 @@ FName FChoosersTrackCreator::GetNameInternal() const
 	return ChoosersName;
 }
 
+void FChoosersTrackCreator::GetTrackTypesInternal(TArray<RewindDebugger::FRewindDebuggerTrackType>& Types) const
+{
+	Types.Add({ChoosersName, LOCTEXT("Chooser", "Choosers")});
+}
+
 TSharedPtr<RewindDebugger::FRewindDebuggerTrack> FChoosersTrackCreator::CreateTrackInternal(uint64 ObjectId) const
 {
 	return MakeShared<FChoosersTrack>(ObjectId);
diff --git a/Engine/Plugins/Experimental/Chooser/Source/ChooserEditor/Private/ChooserTrack.h b/Engine/Plugins/Experimental/Chooser/Source/ChooserEditor/Private/ChooserTrack.h
--- a/Engine/Plugins/Experimental/Chooser/Source/ChooserEditor/Private/ChooserTrack.h
+++ b/Engine/Plugins/Experimental/Chooser/Source/ChooserEditor/Private/ChooserTrack.h
@@ -68,6 +68,7 @@ class FChoosersTrackCreator : public RewindDebugger::IRewindDebuggerTrackCreator
 private:
 	virtual FName GetTargetTypeNameInternal() const;
 	virtual FName GetNameInternal() const override;
+	virtual void GetTrackTypesInternal(TArray<RewindDebugger::FRewindDebuggerTrackType>& Types) const override;
 	virtual TSharedPtr<RewindDebugger::FRewindDebuggerTrack> CreateTrackInternal(uint64 ObjectId) const override;
 	virtual bool HasDebugInfoInternal(uint64 ObjectId) const override;
 };
diff --git a/Engine/Plugins/Experimental/Chooser/Source/ChooserUncooked/Private/AnimGraphNode_ChooserPlayer.cpp b/Engine/Plugins/Experimental/Chooser/Source/ChooserUncooked/Private/AnimGraphNode_ChooserPlayer.cpp
--- a/Engine/Plugins/Experimental/Chooser/Source/ChooserUncooked/Private/AnimGraphNode_ChooserPlayer.cpp
+++ b/Engine/Plugins/Experimental/Chooser/Source/ChooserUncooked/Private/AnimGraphNode_ChooserPlayer.cpp
@@ -2,15 +2,12 @@
 
 #include "AnimGraphNode_ChooserPlayer.h"
 
-#include <chrono>
-
 #include "EditorCategoryUtils.h"
 #include "DetailLayoutBuilder.h"
 #include "Animation/AnimAttributes.h"
 #include "Animation/AnimPoseSearchProvider.h"
 #include "Animation/AnimRootMotionProvider.h"
 #include "ChooserPropertyAccess.h"
-#include "Elements/Interfaces/TypedElementDataStorageCompatibilityInterface.h"
 
 #define LOCTEXT_NAMESPACE "AnimGraphNode_ChooserPlayer"
 
diff --git a/Engine/Source/Editor/RewindDebuggerInterface/Public/IRewindDebuggerTrackCreator.h b/Engine/Source/Editor/RewindDebuggerInterface/Public/IRewindDebuggerTrackCreator.h
--- a/Engine/Source/Editor/RewindDebuggerInterface/Public/IRewindDebuggerTrackCreator.h
+++ b/Engine/Source/Editor/RewindDebuggerInterface/Public/IRewindDebuggerTrackCreator.h
@@ -18,6 +18,12 @@ namespace RewindDebugger
 	
 class FRewindDebuggerTrack;
 
+struct FRewindDebuggerTrackType
+{
+	FName Name;
+	FText DisplayName;
+};
+
 // Interface class which creates tracks
 class REWINDDEBUGGERINTERFACE_API IRewindDebuggerTrackCreator : public IModularFeature
 {
@@ -55,15 +61,27 @@ class REWINDDEBUGGERINTERFACE_API IRewindDebuggerTrackCreator : public IModularF
 	{
 		return CreateTrackInternal(ObjectId);
 	}
+	
+	void GetTrackTypes(TArray<FRewindDebuggerTrackType>& Types) const
+	{
+		return GetTrackTypesInternal(Types);
+	};
 
 	private:
 
+	// get the UObject type name this Creator will create child tracks for
 	virtual FName GetTargetTypeNameInternal() const { return "Object"; }
-	
+
+	// get the Name (unique identifier) for this Creator
 	virtual FName GetNameInternal() const { return FName(); }
-	
+
+	// An integer to override sort order for tracks created by this Creator (Higher priority will make tracks appear higher in the list)
 	virtual int32 GetSortOrderPriorityInternal() const { return 0; }
-	
+
+	// Add track types that this Creator will create to the track type list
+	virtual void GetTrackTypesInternal(TArray<FRewindDebuggerTrackType>& Types) const { };
+
+	// Returns true if this creator's track should be shown for this Object, false if there is no data and they should be hidden.
 	virtual bool HasDebugInfoInternal(uint64 ObjectId) const
 	{
 		return true;
