Auto eject on pause feature for rewind debugger

diff --git a/Engine/Plugins/Animation/GameplayInsights/Source/RewindDebugger/Private/RewindDebugger.cpp b/Engine/Plugins/Animation/GameplayInsights/Source/RewindDebugger/Private/RewindDebugger.cpp
--- a/Engine/Plugins/Animation/GameplayInsights/Source/RewindDebugger/Private/RewindDebugger.cpp
+++ b/Engine/Plugins/Animation/GameplayInsights/Source/RewindDebugger/Private/RewindDebugger.cpp
@@ -34,6 +34,9 @@
 #include "Widgets/Layout/SSpacer.h"
 #include "RewindDebuggerCommands.h"
 #include "RewindDebuggerTrackCreators.h"
+#include "UnrealEdGlobals.h"
+#include "Editor/UnrealEdEngine.h"
+#include "Kismet2/DebuggerCommands.h"
 
 #define LOCTEXT_NAMESPACE "RewindDebugger"
 
@@ -177,6 +180,20 @@ void FRewindDebugger::OnPIEPaused(bool bSimulating)
 #endif // OBJECT_TRACE_ENABLED
 		SetCurrentScrubTime(RecordingDuration.Get());
 	}
+	
+	if (ShouldAutoDetach() && FPlayWorldCommandCallbacks::IsInPIE())
+	{
+		bool CanEject= false;
+		for (auto It = GUnrealEd->SlatePlayInEditorMap.CreateIterator(); It; ++It)
+		{
+			CanEject = CanEject || It.Value().DestinationSlateViewport.IsValid();
+		}
+
+		if (CanEject)
+		{
+			GEditor->RequestToggleBetweenPIEandSIE();
+		}
+	}
 }
 
 void FRewindDebugger::OnPIEResumed(bool bSimulating)
@@ -193,6 +210,11 @@ void FRewindDebugger::OnPIEResumed(bool bSimulating)
 	}
 
 	MeshComponentsToReset.Empty();
+
+	if (ShouldAutoDetach() && FPlayWorldCommandCallbacks::IsInSIE())
+	{
+		GEditor->RequestToggleBetweenPIEandSIE();
+	}
 }
 
 void FRewindDebugger::OnPIESingleStepped(bool bSimulating)
@@ -455,6 +477,16 @@ void FRewindDebugger::SetShouldAutoRecordOnPIE(bool value)
 	URewindDebuggerSettings::Get().bShouldAutoRecordOnPIE = value;
 }
 
+bool FRewindDebugger::ShouldAutoDetach() const
+{
+	return URewindDebuggerSettings::Get().bShouldAutoDetach;
+}
+
+void FRewindDebugger::SetShouldAutoDetach(bool value)
+{
+	URewindDebuggerSettings::Get().bShouldAutoDetach = value;
+}
+
 void FRewindDebugger::StopRecording()
 {
 	if (bRecording)
diff --git a/Engine/Plugins/Animation/GameplayInsights/Source/RewindDebugger/Private/RewindDebugger.h b/Engine/Plugins/Animation/GameplayInsights/Source/RewindDebugger/Private/RewindDebugger.h
--- a/Engine/Plugins/Animation/GameplayInsights/Source/RewindDebugger/Private/RewindDebugger.h
+++ b/Engine/Plugins/Animation/GameplayInsights/Source/RewindDebugger/Private/RewindDebugger.h
@@ -70,6 +70,9 @@ public:
 
 	bool ShouldAutoRecordOnPIE() const;
 	void SetShouldAutoRecordOnPIE(bool value);
+	
+	bool ShouldAutoDetach() const;
+	void SetShouldAutoDetach(bool value);
 
 	// Stop recording: Stop tracing Object + Animation Data.
 	void StopRecording();
diff --git a/Engine/Plugins/Animation/GameplayInsights/Source/RewindDebugger/Private/RewindDebuggerCamera.cpp b/Engine/Plugins/Animation/GameplayInsights/Source/RewindDebugger/Private/RewindDebuggerCamera.cpp
--- a/Engine/Plugins/Animation/GameplayInsights/Source/RewindDebugger/Private/RewindDebuggerCamera.cpp
+++ b/Engine/Plugins/Animation/GameplayInsights/Source/RewindDebugger/Private/RewindDebuggerCamera.cpp
@@ -107,6 +107,8 @@ void FRewindDebuggerCamera::Update(float DeltaTime, IRewindDebugger* RewindDebug
 		static double LastCameraScrubTime = 0.0f;
 		if (CurrentTraceTime != LastCameraScrubTime)
 		{
+			bool bCameraTraceDataFound = false;
+			
 			FLevelEditorModule& LevelEditor = FModuleManager::GetModuleChecked<FLevelEditorModule>("LevelEditor");
 			SLevelViewport* LevelViewport = LevelEditor.GetFirstActiveLevelViewport().Get();
 			FLevelEditorViewportClient& LevelViewportClient = LevelViewport->GetLevelViewportClient();
@@ -114,8 +116,6 @@ void FRewindDebuggerCamera::Update(float DeltaTime, IRewindDebugger* RewindDebug
 			FVector TargetActorPosition;
 			bool bTargetActorPositionValid = RewindDebugger->GetTargetActorPosition(TargetActorPosition);
 
-			// only update camera in playback or scrubbing when the time has changed (allow free movement when paused)
-			LastCameraScrubTime = CurrentTraceTime;
 
 			if (CameraMode() == ERewindDebuggerCameraMode::FollowTargetActor)
 			{
@@ -132,21 +132,21 @@ void FRewindDebuggerCamera::Update(float DeltaTime, IRewindDebugger* RewindDebug
 			// always update the camera actor to the replay values even if it isn't locked
 			if (const IGameplayProvider* GameplayProvider = Session->ReadProvider<IGameplayProvider>("GameplayProvider"))
 			{
-				GameplayProvider->ReadViewTimeline([this, RewindDebugger, CurrentTraceTime, Session](const IGameplayProvider::ViewTimeline& TimelineData)
+				GameplayProvider->ReadViewTimeline([&bCameraTraceDataFound, this, RewindDebugger, CurrentTraceTime, Session](const IGameplayProvider::ViewTimeline& TimelineData)
 				{
 					const TraceServices::IFrameProvider& FrameProvider = TraceServices::ReadFrameProvider(*Session);
 					TraceServices::FFrame Frame;
 					if(FrameProvider.GetFrameFromTime(ETraceFrameType::TraceFrameType_Game, CurrentTraceTime, Frame))
 					{
 						TimelineData.EnumerateEvents(Frame.StartTime, Frame.EndTime,
-							[this, RewindDebugger](double InStartTime, double InEndTime, uint32 InDepth, const FViewMessage& ViewMessage)
+							[&bCameraTraceDataFound, this, RewindDebugger](double InStartTime, double InEndTime, uint32 InDepth, const FViewMessage& ViewMessage)
 							{
 								if (!CameraActor.IsValid())
 								{
 									FActorSpawnParameters SpawnParameters;
 									SpawnParameters.ObjectFlags |= RF_Transient;
 									CameraActor = RewindDebugger->GetWorldToVisualize()->SpawnActor<ACameraActor>(ViewMessage.Position, ViewMessage.Rotation, SpawnParameters);
-									CameraActor->SetActorLabel("RewindDebuggerCamera"); 
+									CameraActor->SetActorLabel("RewindDebuggerCamera");
 								}
 
 								UCameraComponent* Camera = CameraActor->GetCameraComponent();
@@ -154,6 +154,8 @@ void FRewindDebuggerCamera::Update(float DeltaTime, IRewindDebugger* RewindDebug
 								Camera->SetFieldOfView(ViewMessage.Fov);
 								Camera->SetAspectRatio(ViewMessage.AspectRatio);
 
+								bCameraTraceDataFound = true;
+
 								return TraceServices::EEventEnumerate::Stop;
 							});
 					}
@@ -170,6 +172,12 @@ void FRewindDebuggerCamera::Update(float DeltaTime, IRewindDebugger* RewindDebug
 
 			LastPosition = TargetActorPosition;
 			LastPositionValid = bTargetActorPositionValid;
+
+			if (bCameraTraceDataFound) // don't update this if there was no trace data found, because when first pausing, it can take a few frames for latest data to get processed
+			{
+				// only update camera in playback or scrubbing when the time has changed (allow free movement when paused)
+				LastCameraScrubTime = CurrentTraceTime;
+			}
 		}
 	}
 }
diff --git a/Engine/Plugins/Animation/GameplayInsights/Source/RewindDebugger/Private/RewindDebuggerSettings.cpp b/Engine/Plugins/Animation/GameplayInsights/Source/RewindDebugger/Private/RewindDebuggerSettings.cpp
--- a/Engine/Plugins/Animation/GameplayInsights/Source/RewindDebugger/Private/RewindDebuggerSettings.cpp
+++ b/Engine/Plugins/Animation/GameplayInsights/Source/RewindDebugger/Private/RewindDebuggerSettings.cpp
@@ -5,7 +5,7 @@
 
 #include UE_INLINE_GENERATED_CPP_BY_NAME(RewindDebuggerSettings)
 
-URewindDebuggerSettings::URewindDebuggerSettings() : CameraMode(ERewindDebuggerCameraMode::Replay), /*bShouldAutoDetach(false),*/ bShouldAutoRecordOnPIE(false)
+URewindDebuggerSettings::URewindDebuggerSettings() : CameraMode(ERewindDebuggerCameraMode::Replay), bShouldAutoDetach(true), bShouldAutoRecordOnPIE(false)
 {
 	FCoreDelegates::OnPreExit.AddLambda([]()
 	{
diff --git a/Engine/Plugins/Animation/GameplayInsights/Source/RewindDebugger/Private/RewindDebuggerSettings.h b/Engine/Plugins/Animation/GameplayInsights/Source/RewindDebugger/Private/RewindDebuggerSettings.h
--- a/Engine/Plugins/Animation/GameplayInsights/Source/RewindDebugger/Private/RewindDebuggerSettings.h
+++ b/Engine/Plugins/Animation/GameplayInsights/Source/RewindDebugger/Private/RewindDebuggerSettings.h
@@ -38,9 +38,9 @@ public:
 	UPROPERTY(EditAnywhere, Config, Category = Camera)
 	ERewindDebuggerCameraMode CameraMode;
 	
-	// /** If enabled, automatically detach player controller at the start of PIE */
-	// UPROPERTY(EditAnywhere, Config, Category = Other, meta = (DisplayName = "Auto Detach Player Controller on PIE"))
-	// bool bShouldAutoDetach;
+	/** If enabled, automatically detach player controller at the start of PIE */
+	UPROPERTY(EditAnywhere, Config, Category = Other, meta = (DisplayName = "Auto Detach Player Controller on PIE"))
+	bool bShouldAutoDetach;
 	
 	/** If enabled, start recording information at the start of PIE */
 	UPROPERTY(EditAnywhere, Config, Category = Other)
