Rewind Debugger track view improvements
- Striped background
- Track names in tooltips
- Synchronized hover state
- Object tracks now display lifetime with SEventTimelineView

diff --git a/Engine/Plugins/Animation/GameplayInsights/Source/GameplayInsights/Private/AnimCurvesTrack.cpp b/Engine/Plugins/Animation/GameplayInsights/Source/GameplayInsights/Private/AnimCurvesTrack.cpp
--- a/Engine/Plugins/Animation/GameplayInsights/Source/GameplayInsights/Private/AnimCurvesTrack.cpp
+++ b/Engine/Plugins/Animation/GameplayInsights/Source/GameplayInsights/Private/AnimCurvesTrack.cpp
@@ -207,6 +207,7 @@ TSharedPtr<SWidget> FAnimCurveTrack::GetTimelineViewInternal()
 	FLinearColor CurveColor(0.5,0.5,0.5);
 	
 	return SNew(SCurveTimelineView)
+		.TrackName(GetDisplayNameInternal())
 		.CurveColor(CurveColor)
 		.ViewRange_Lambda([]() { return IRewindDebugger::Instance()->GetCurrentViewRange(); })
 		.RenderFill(false)
diff --git a/Engine/Plugins/Animation/GameplayInsights/Source/GameplayInsights/Private/BlendWeightsTrack.cpp b/Engine/Plugins/Animation/GameplayInsights/Source/GameplayInsights/Private/BlendWeightsTrack.cpp
--- a/Engine/Plugins/Animation/GameplayInsights/Source/GameplayInsights/Private/BlendWeightsTrack.cpp
+++ b/Engine/Plugins/Animation/GameplayInsights/Source/GameplayInsights/Private/BlendWeightsTrack.cpp
@@ -283,44 +283,31 @@ static FLinearColor MakeBlendWeightCurveColor(uint32 InSeed, bool bInLine = fals
 
 TSharedPtr<SWidget> FBlendWeightTrack::GetTimelineViewInternal()
 {
-	FLinearColor Color;
+	FLinearColor FillColor;
+	FLinearColor CurveColor;
+	FLinearColor SelectedColor;
+	
 	switch(CurveType)
 	{
 	case ECurveType::BlendWeight:
-		Color = MakeBlendWeightCurveColor(CityHash32(reinterpret_cast<char*>(&AssetId), 8));
-		Color.A = 0.5f;
-		break; 
-	case ECurveType::PlaybackTime:
-		Color = FLinearColor::MakeFromHSV8(0, 50, 50);
-		break;
-	case ECurveType::RootMotionWeight:
-		Color = FLinearColor::MakeFromHSV8(60, 50, 50);
-		break;
-	case ECurveType::PlayRate:		
-		Color = FLinearColor::MakeFromHSV8(120, 50, 50);
-		break;
-	case ECurveType::BlendSpacePositionX:
-		Color = FLinearColor::MakeFromHSV8(180, 50, 50);
+		FillColor = MakeBlendWeightCurveColor(CityHash32(reinterpret_cast<char*>(&AssetId), 8));
+		FillColor.A = 0.5f;
+		CurveColor = FillColor;
+		CurveColor.R *= 0.5;
+		CurveColor.G *= 0.5;
+		CurveColor.B *= 0.5;
+		SelectedColor = CurveColor;
 		break;
-	case ECurveType::BlendSpacePositionY:
-		Color = FLinearColor::MakeFromHSV8(240, 50, 50);
-		break;
-	case ECurveType::BlendSpaceFilteredPositionX:
-		Color = FLinearColor::MakeFromHSV8(180, 50, 80);
-		break;
-	case ECurveType::BlendSpaceFilteredPositionY:
-		Color = FLinearColor::MakeFromHSV8(240, 50, 80);
+	default:
+		CurveColor = FLinearColor::MakeFromHSV8(100, 50, 65);
+		SelectedColor = FLinearColor::MakeFromHSV8(100, 50, 130);
 		break;
 	}
 
-	FLinearColor CurveColor = Color;
-	CurveColor.R *= 0.5;
-	CurveColor.G *= 0.5;
-	CurveColor.B *= 0.5;
-
 	TSharedPtr<SCurveTimelineView> CurveTimelineView = SNew(SCurveTimelineView)
-		.FillColor(Color)
-		.CurveColor(CurveColor)
+		.TrackName(GetDisplayNameInternal())
+		.FillColor(FillColor)
+		.CurveColor_Lambda([CurveColor, SelectedColor,  this]() { return GetIsSelected() ? SelectedColor : CurveColor; })
 		.ViewRange_Lambda([]() { return IRewindDebugger::Instance()->GetCurrentViewRange(); })
 		.RenderFill(CurveType == ECurveType::BlendWeight)
 		.CurveData_Raw(this, &FBlendWeightTrack::GetCurveData);
diff --git a/Engine/Plugins/Animation/GameplayInsights/Source/GameplayInsights/Private/InertializationsTrack.cpp b/Engine/Plugins/Animation/GameplayInsights/Source/GameplayInsights/Private/InertializationsTrack.cpp
--- a/Engine/Plugins/Animation/GameplayInsights/Source/GameplayInsights/Private/InertializationsTrack.cpp
+++ b/Engine/Plugins/Animation/GameplayInsights/Source/GameplayInsights/Private/InertializationsTrack.cpp
@@ -177,6 +177,7 @@ TSharedPtr<SWidget> FInertializationTrack::GetTimelineViewInternal()
 	CurveColor.B *= 0.5;
 
 	TSharedPtr<SCurveTimelineView> CurveTimelineView = SNew(SCurveTimelineView)
+		.TrackName(GetDisplayNameInternal())
 		.FillColor(Color)
 		.CurveColor(CurveColor)
 		.ViewRange_Lambda([]() { return IRewindDebugger::Instance()->GetCurrentViewRange(); })
diff --git a/Engine/Plugins/Animation/GameplayInsights/Source/GameplayInsights/Private/SCurveTimelineView.cpp b/Engine/Plugins/Animation/GameplayInsights/Source/GameplayInsights/Private/SCurveTimelineView.cpp
--- a/Engine/Plugins/Animation/GameplayInsights/Source/GameplayInsights/Private/SCurveTimelineView.cpp
+++ b/Engine/Plugins/Animation/GameplayInsights/Source/GameplayInsights/Private/SCurveTimelineView.cpp
@@ -18,6 +18,7 @@ void SCurveTimelineView::Construct( const SCurveTimelineView::FArguments& InArgs
 	RenderFill = InArgs._RenderFill;
 	DesiredSize = InArgs._DesiredSize;
 	CurveData = InArgs._CurveData;
+	TrackName = InArgs._TrackName;
 
 	// set clipping on by default, since the OnPaint function is drawing outside the bounds
 	Clipping = EWidgetClipping::ClipToBounds;
@@ -90,6 +91,13 @@ void SCurveTimelineView::UpdateCurveToolTip(const FGeometry& InMyGeometry, const
 								[
 									SNew(SVerticalBox)
 									+ SVerticalBox::Slot()
+									[
+										SNew(STextBlock)
+										.Text(this, &SCurveTimelineView::GetTrackName)
+										.Font(FCoreStyle::Get().GetFontStyle("ToolTip.LargerFont"))
+										.ColorAndOpacity(FLinearColor::Black)
+									]
+									+ SVerticalBox::Slot()
 									[
 										SNew(STextBlock)
 										.Text(this, &SCurveTimelineView::GetCurveToolTipInputText)
diff --git a/Engine/Plugins/Animation/GameplayInsights/Source/GameplayInsights/Public/SCurveTimelineView.h b/Engine/Plugins/Animation/GameplayInsights/Source/GameplayInsights/Public/SCurveTimelineView.h
--- a/Engine/Plugins/Animation/GameplayInsights/Source/GameplayInsights/Public/SCurveTimelineView.h
+++ b/Engine/Plugins/Animation/GameplayInsights/Source/GameplayInsights/Public/SCurveTimelineView.h
@@ -48,6 +48,9 @@ public:
 	
 		/** Enable Fill rendering */
 		SLATE_ATTRIBUTE(bool, RenderFill);
+
+		/** Track Name for tool tip */
+		SLATE_ATTRIBUTE(FText, TrackName);
      
 	SLATE_END_ARGS()
 
@@ -67,6 +70,7 @@ public:
 		FixedRangeMax = Max;
 	}
 
+	FText GetTrackName() const { return TrackName.Get(); }
 	FText GetCurveToolTipInputText() const { return CurveToolTipInputText; }
 	FText GetCurveToolTipOutputText() const { return CurveToolTipOutputText; }
 	
@@ -85,6 +89,7 @@ protected:
 	TAttribute<FLinearColor> CurveColor;
 	TAttribute<FLinearColor> FillColor;
 	TAttribute<bool> RenderFill;
+	TAttribute<FText> TrackName;
 	
 	TSharedPtr<SToolTip> CurveToolTip;	/** The tooltip control for the curve. */
 	FText CurveToolTipInputText;		
diff --git a/Engine/Plugins/Animation/GameplayInsights/Source/RewindDebugger/Private/RewindDebuggerObjectTrack.cpp b/Engine/Plugins/Animation/GameplayInsights/Source/RewindDebugger/Private/RewindDebuggerObjectTrack.cpp
--- a/Engine/Plugins/Animation/GameplayInsights/Source/RewindDebugger/Private/RewindDebuggerObjectTrack.cpp
+++ b/Engine/Plugins/Animation/GameplayInsights/Source/RewindDebugger/Private/RewindDebuggerObjectTrack.cpp
@@ -7,11 +7,14 @@
 #include "IAnimationProvider.h"
 #include "RewindDebuggerViewCreators.h"
 #include "RewindDebuggerTrackCreators.h"
-#include "SSegmentedTimelineView.h"
+#include "SEventTimelineView.h"
 #include "Styling/SlateIconFinder.h"
 #include "ObjectTrace.h"
 #include "IRewindDebuggerDoubleClickHandler.h"
 
+
+#define LOCTEXT_NAMESPACE "RewindDebuggerObjectTrack"
+
 namespace RewindDebugger
 {
 	
@@ -43,8 +46,7 @@ FRewindDebuggerObjectTrack::FRewindDebuggerObjectTrack(uint64 InObjectId, const
 	, bAddController(bInAddController)
 	, bDisplayNameValid(false)
 {
-	ExistenceRange = MakeShared<SSegmentedTimelineView::FSegmentData>();
-	ExistenceRange->Segments.SetNumUninitialized(1);
+	ExistenceRange = MakeShared<SEventTimelineView::FTimelineEventData>();
 	
 	IRewindDebugger* RewindDebugger = IRewindDebugger::Instance();
 	const TraceServices::IAnalysisSession* Session = RewindDebugger->GetAnalysisSession();
@@ -74,10 +76,9 @@ FRewindDebuggerObjectTrack::FRewindDebuggerObjectTrack(uint64 InObjectId, const
 	
 TSharedPtr<SWidget> FRewindDebuggerObjectTrack::GetTimelineViewInternal()
 {
-	return SNew(SSegmentedTimelineView)
+	return SNew(SEventTimelineView)
 		.ViewRange_Lambda([]() { return IRewindDebugger::Instance()->GetCurrentViewRange(); })
-		.FillColor(FLinearColor(0.02f, 0.02f, 0.02f, 0.5f))
-		.SegmentData_Raw(this, &FRewindDebuggerObjectTrack::GetExistenceRange);
+		.EventData_Raw(this, &FRewindDebuggerObjectTrack::GetExistenceRange);
 }
 
 bool FRewindDebuggerObjectTrack::HandleDoubleClickInternal()
@@ -188,8 +189,13 @@ bool FRewindDebuggerObjectTrack::UpdateInternal()
 
 	bool bChanged = false;
 
-	check(ExistenceRange->Segments.Num() == 1);
-	ExistenceRange->Segments[0] = GameplayProvider->GetObjectRecordingLifetime(ObjectId);
+	TRange<double> Existence = GameplayProvider->GetObjectRecordingLifetime(ObjectId);
+
+	ExistenceRange->Windows.SetNum(0,false);
+	if (Existence.HasLowerBound() && Existence.HasUpperBound())
+	{
+		ExistenceRange->Windows.Add({Existence.GetLowerBoundValue(), Existence.GetUpperBoundValue(), LOCTEXT("Object Existence","Object Existence"), LOCTEXT("Object Existence","Object Existence"), FLinearColor(0.1f,0.11f,0.1f)});
+	}
 
 	if (!Icon.IsSet())
 	{
@@ -359,4 +365,6 @@ bool FRewindDebuggerObjectTrack::UpdateInternal()
 	return bChanged;
 }
 
-}
\ No newline at end of file
+}
+
+#undef LOCTEXT_NAMESPACE
\ No newline at end of file
diff --git a/Engine/Plugins/Animation/GameplayInsights/Source/RewindDebugger/Private/RewindDebuggerObjectTrack.h b/Engine/Plugins/Animation/GameplayInsights/Source/RewindDebugger/Private/RewindDebuggerObjectTrack.h
--- a/Engine/Plugins/Animation/GameplayInsights/Source/RewindDebugger/Private/RewindDebuggerObjectTrack.h
+++ b/Engine/Plugins/Animation/GameplayInsights/Source/RewindDebugger/Private/RewindDebuggerObjectTrack.h
@@ -2,7 +2,7 @@
 #pragma once
 
 #include "RewindDebuggerTrack.h"
-#include "SSegmentedTimelineView.h"
+#include "SEventTimelineView.h"
 
 
 namespace RewindDebugger
@@ -22,7 +22,7 @@ public:
 
 	FRewindDebuggerObjectTrack(uint64 InObjectId, const FString& InObjectName, bool bInAddController = false);
 
-	TSharedPtr<SSegmentedTimelineView::FSegmentData> GetExistenceRange() const { return ExistenceRange; }
+	TSharedPtr<SEventTimelineView::FTimelineEventData> GetExistenceRange() const { return ExistenceRange; }
 
 private:
 	virtual TSharedPtr<SWidget> GetTimelineViewInternal() override;
@@ -39,7 +39,7 @@ private:
 	mutable FText DisplayName;
 	FString ObjectName;
 	FSlateIcon Icon;
-	TSharedPtr<SSegmentedTimelineView::FSegmentData> ExistenceRange;
+	TSharedPtr<SEventTimelineView::FTimelineEventData> ExistenceRange;
 	uint64 ObjectId;
 	TArray<FTrackCreatorAndTrack> TrackChildren;
 	TArray<TSharedPtr<FRewindDebuggerTrack>> Children;
diff --git a/Engine/Plugins/Animation/GameplayInsights/Source/RewindDebugger/Private/RewindDebuggerStyle.cpp b/Engine/Plugins/Animation/GameplayInsights/Source/RewindDebugger/Private/RewindDebuggerStyle.cpp
--- a/Engine/Plugins/Animation/GameplayInsights/Source/RewindDebugger/Private/RewindDebuggerStyle.cpp
+++ b/Engine/Plugins/Animation/GameplayInsights/Source/RewindDebugger/Private/RewindDebuggerStyle.cpp
@@ -40,6 +40,9 @@ FRewindDebuggerStyle::FRewindDebuggerStyle() :
 
 	// menu icon
 	Set("RewindDebugger.MenuIcon", new CORE_IMAGE_BRUSH_SVG("Slate/Starship/Common/menu", Icon16x16));
+
+	FTableRowStyle TableRowStyle = FTableRowStyle(FAppStyle::Get().GetWidgetStyle<FTableRowStyle>("SceneOutliner.TableViewRow"));
+	Set("RewindDebugger.TableRow", TableRowStyle);
 }
 
 void FRewindDebuggerStyle::Initialize()
diff --git a/Engine/Plugins/Animation/GameplayInsights/Source/RewindDebugger/Private/SRewindDebugger.cpp b/Engine/Plugins/Animation/GameplayInsights/Source/RewindDebugger/Private/SRewindDebugger.cpp
--- a/Engine/Plugins/Animation/GameplayInsights/Source/RewindDebugger/Private/SRewindDebugger.cpp
+++ b/Engine/Plugins/Animation/GameplayInsights/Source/RewindDebugger/Private/SRewindDebugger.cpp
@@ -568,7 +568,18 @@ TSharedRef<SWidget> SRewindDebugger::MakeFilterMenu()
 
 void SRewindDebugger::ComponentSelectionChanged(TSharedPtr<RewindDebugger::FRewindDebuggerTrack> SelectedItem, ESelectInfo::Type SelectInfo)
 {
+	if (SelectedComponent)
+	{
+		SelectedComponent->SetIsSelected(false);
+	}
+	
 	SelectedComponent = SelectedItem;
+	
+	if (SelectedComponent)
+	{
+		SelectedComponent->SetIsSelected(true);
+	}
+	
 
 	OnComponentSelectionChanged.ExecuteIfBound(SelectedItem);
 }
diff --git a/Engine/Plugins/Animation/GameplayInsights/Source/RewindDebugger/Private/SRewindDebuggerComponentTree.cpp b/Engine/Plugins/Animation/GameplayInsights/Source/RewindDebugger/Private/SRewindDebuggerComponentTree.cpp
--- a/Engine/Plugins/Animation/GameplayInsights/Source/RewindDebugger/Private/SRewindDebuggerComponentTree.cpp
+++ b/Engine/Plugins/Animation/GameplayInsights/Source/RewindDebugger/Private/SRewindDebuggerComponentTree.cpp
@@ -6,6 +6,7 @@
 #include "ObjectTrace.h"
 #include "Styling/SlateIconFinder.h"
 #include "Widgets/Images/SImage.h"
+#include "RewindDebuggerStyle.h"
 
 #include "RewindDebugger.h"
 #include "Widgets/Images/SLayeredImage.h"
@@ -23,6 +24,60 @@ SRewindDebuggerComponentTree::~SRewindDebuggerComponentTree()
 
 }
 
+class SRewindDebuggerComponentTreeTableRow : public STableRow<TSharedPtr<RewindDebugger::FRewindDebuggerTrack>>
+{
+public:
+	
+	SLATE_BEGIN_ARGS(SRewindDebuggerComponentTreeTableRow) {}
+	SLATE_END_ARGS()
+	
+	void Construct(const FArguments& InArgs, const TSharedRef<STableViewBase>& InOwnerTableView)
+	{
+		STableRow<TSharedPtr<RewindDebugger::FRewindDebuggerTrack>>::Construct( STableRow<TSharedPtr<RewindDebugger::FRewindDebuggerTrack>>::FArguments(), InOwnerTableView);
+		Style =  &FRewindDebuggerStyle::Get().GetWidgetStyle<FTableRowStyle>("RewindDebugger.TableRow");
+
+		SetHover(TAttribute<bool>::CreateLambda([this]()
+		{
+			if (auto TableView = OwnerTablePtr.Pin())
+			{
+				if (const TSharedPtr<RewindDebugger::FRewindDebuggerTrack>* Track = GetItemForThis(TableView.ToSharedRef()))
+				{
+					return (*Track)->GetIsHovered();
+				}
+			}
+
+			return false;
+		}));
+	}
+
+	virtual void OnMouseEnter(const FGeometry& MyGeometry, const FPointerEvent& MouseEvent) override
+	{
+		if (auto TableView = OwnerTablePtr.Pin())
+		{
+			if (const TSharedPtr<RewindDebugger::FRewindDebuggerTrack>* Track = GetItemForThis(TableView.ToSharedRef()))
+			{
+				(*Track)->SetIsTreeHovered(true);
+			}
+		}
+		STableRow<TSharedPtr<RewindDebugger::FRewindDebuggerTrack>>::OnMouseEnter(MyGeometry, MouseEvent);
+	}
+	
+	virtual void OnMouseLeave(const FPointerEvent& MouseEvent) override
+	{
+		STableRow<TSharedPtr<RewindDebugger::FRewindDebuggerTrack>>::OnMouseLeave(MouseEvent);
+
+		if (auto TableView = OwnerTablePtr.Pin())
+		{
+			if (const TSharedPtr<RewindDebugger::FRewindDebuggerTrack>* Track = GetItemForThis(TableView.ToSharedRef()))
+			{
+				(*Track)->SetIsTreeHovered(false);
+			}
+		}
+	}
+
+};
+
+
 TSharedRef<ITableRow> SRewindDebuggerComponentTree::ComponentTreeViewGenerateRow(TSharedPtr<RewindDebugger::FRewindDebuggerTrack> InItem, const TSharedRef<STableViewBase>& OwnerTable)
 {
 	const FSlateIcon ObjectIcon = InItem->GetIcon();
@@ -36,9 +91,10 @@ TSharedRef<ITableRow> SRewindDebuggerComponentTree::ComponentTreeViewGenerateRow
 		LayeredIcons->AddLayer(ObjectIcon.GetOverlayIcon());
 	}
 	
-	return 
-		SNew(STableRow<TSharedPtr<RewindDebugger::FRewindDebuggerTrack>>, OwnerTable)
-		[
+	TSharedRef<SRewindDebuggerComponentTreeTableRow> Row = SNew(SRewindDebuggerComponentTreeTableRow, OwnerTable);
+
+
+	Row->SetContent(
 			SNew(SHorizontalBox)
 			+SHorizontalBox::Slot().AutoWidth().Padding(2)
 			[
@@ -51,8 +107,11 @@ TSharedRef<ITableRow> SRewindDebuggerComponentTree::ComponentTreeViewGenerateRow
 			[
 				SNew(STextBlock)
 				.Text(InItem->GetDisplayName())
+				.Font_Lambda([this,  InItem]() { return InItem->GetIsHovered() ? FCoreStyle::GetDefaultFontStyle("Bold", 10) : FCoreStyle::GetDefaultFontStyle("Regular", 10); })
 			]
-		];
+		);
+
+	return Row;
 }
 
 void ComponentTreeViewGetChildren(TSharedPtr<RewindDebugger::FRewindDebuggerTrack> InItem, TArray<TSharedPtr<RewindDebugger::FRewindDebuggerTrack>>& OutChildren)
diff --git a/Engine/Plugins/Animation/GameplayInsights/Source/RewindDebugger/Private/SRewindDebuggerTimelines.cpp b/Engine/Plugins/Animation/GameplayInsights/Source/RewindDebugger/Private/SRewindDebuggerTimelines.cpp
--- a/Engine/Plugins/Animation/GameplayInsights/Source/RewindDebugger/Private/SRewindDebuggerTimelines.cpp
+++ b/Engine/Plugins/Animation/GameplayInsights/Source/RewindDebugger/Private/SRewindDebuggerTimelines.cpp
@@ -7,6 +7,7 @@
 #include "Widgets/Layout/SSpacer.h"
 
 #include "RewindDebugger.h"
+#include "RewindDebuggerStyle.h"
 #include "SSimpleTimeSlider.h"
 
 #define LOCTEXT_NAMESPACE "SAnimationInsights"
@@ -19,11 +20,24 @@ class SRewindDebuggerTimelineTableRow : public STableRow<TSharedPtr<RewindDebugg
 	SLATE_BEGIN_ARGS(SRewindDebuggerTimelineTableRow) {}
 	SLATE_END_ARGS()
 	
-	void Construct(const SRewindDebuggerTimelineTableRow::FArguments& InArgs, const TSharedRef<STableViewBase>& InOwnerTableView)
+	void Construct(const FArguments& InArgs, const TSharedRef<STableViewBase>& InOwnerTableView)
 	{
 		STableRow<TSharedPtr<RewindDebugger::FRewindDebuggerTrack>>::Construct( STableRow<TSharedPtr<RewindDebugger::FRewindDebuggerTrack>>::FArguments(), InOwnerTableView);
-
+		Style =  &FRewindDebuggerStyle::Get().GetWidgetStyle<FTableRowStyle>("RewindDebugger.TableRow");
+	
 		SetExpanderArrowVisibility(EVisibility::Collapsed);
+
+		SetHover(TAttribute<bool>::CreateLambda([this]()
+		{
+			if (auto TableView = OwnerTablePtr.Pin())
+			{
+				if (const TSharedPtr<RewindDebugger::FRewindDebuggerTrack>* Track = GetItemForThis(TableView.ToSharedRef()))
+				{
+					return (*Track)->GetIsHovered();
+				}
+			}
+			return false;
+		}));
 	}
 
 	virtual int32 GetIndentLevel() const override
@@ -31,6 +45,31 @@ class SRewindDebuggerTimelineTableRow : public STableRow<TSharedPtr<RewindDebugg
 		// don't indent timeline tracks
 		return 0;
 	}
+	
+	virtual void OnMouseEnter(const FGeometry& MyGeometry, const FPointerEvent& MouseEvent) override
+	{
+		if (auto TableView = OwnerTablePtr.Pin())
+		{
+			if (const TSharedPtr<RewindDebugger::FRewindDebuggerTrack>* Track = GetItemForThis(TableView.ToSharedRef()))
+			{
+				(*Track)->SetIsTrackHovered(true);
+			}
+		}
+		STableRow<TSharedPtr<RewindDebugger::FRewindDebuggerTrack>>::OnMouseEnter(MyGeometry, MouseEvent);
+	}
+	
+	virtual void OnMouseLeave(const FPointerEvent& MouseEvent) override
+	{
+		STableRow<TSharedPtr<RewindDebugger::FRewindDebuggerTrack>>::OnMouseLeave(MouseEvent);
+
+		if (auto TableView = OwnerTablePtr.Pin())
+		{
+			if (const TSharedPtr<RewindDebugger::FRewindDebuggerTrack>* Track = GetItemForThis(TableView.ToSharedRef()))
+			{
+				(*Track)->SetIsTrackHovered(false);
+			}
+		}
+	}
 };
 
 
diff --git a/Engine/Source/Editor/RewindDebuggerInterface/Public/RewindDebuggerTrack.h b/Engine/Source/Editor/RewindDebuggerInterface/Public/RewindDebuggerTrack.h
--- a/Engine/Source/Editor/RewindDebuggerInterface/Public/RewindDebuggerTrack.h
+++ b/Engine/Source/Editor/RewindDebuggerInterface/Public/RewindDebuggerTrack.h
@@ -36,6 +36,41 @@ public:
 	{
 		bExpanded = bIsExpanded;
 	}
+	
+	bool GetIsSelected()
+	{
+		return bSelected;
+	}
+
+	void SetIsSelected(bool bIsSelected)
+	{
+		bSelected = bIsSelected;
+	}
+    		
+	bool GetIsTreeHovered()
+	{
+		return bTreeHovered;
+	}
+	
+	void SetIsTreeHovered(bool bIsHovered)
+   	{
+   		bTreeHovered = bIsHovered;
+   	}	
+		
+	bool GetIsTrackHovered()
+	{
+		return bTrackHovered;
+	}
+	
+	void SetIsTrackHovered(bool bIsHovered)
+   	{
+   		bTrackHovered = bIsHovered;
+   	}
+	
+	bool GetIsHovered()
+	{
+		return bTrackHovered || bTreeHovered;
+	}
 
 	// Update should do work to compute children etc for the current time range.  Return true if children have changed.
 	bool Update()
@@ -121,8 +156,11 @@ private:
 		return true;
 	};
 
-	bool bExpanded;
-	bool bVisible;
+	bool bSelected : 1 = false;
+	bool bTrackHovered : 1 = false;
+	bool bTreeHovered : 1 = false;
+	bool bExpanded : 1 = false;
+	bool bVisible : 1 = false;
 };
 	
 }
\ No newline at end of file
