Fixing up some instances of auto for coding standard compliance

diff --git a/Engine/Plugins/Animation/GameplayInsights/Source/RewindDebugger/Private/SRewindDebuggerComponentTree.cpp b/Engine/Plugins/Animation/GameplayInsights/Source/RewindDebugger/Private/SRewindDebuggerComponentTree.cpp
--- a/Engine/Plugins/Animation/GameplayInsights/Source/RewindDebugger/Private/SRewindDebuggerComponentTree.cpp
+++ b/Engine/Plugins/Animation/GameplayInsights/Source/RewindDebugger/Private/SRewindDebuggerComponentTree.cpp
@@ -38,7 +38,7 @@ public:
 
 		SetHover(TAttribute<bool>::CreateLambda([this]()
 		{
-			if (auto TableView = OwnerTablePtr.Pin())
+			if (TSharedPtr<ITypedTableView<TSharedPtr<RewindDebugger::FRewindDebuggerTrack>>> TableView = OwnerTablePtr.Pin())
 			{
 				if (const TSharedPtr<RewindDebugger::FRewindDebuggerTrack>* Track = GetItemForThis(TableView.ToSharedRef()))
 				{
@@ -52,7 +52,7 @@ public:
 
 	virtual void OnMouseEnter(const FGeometry& MyGeometry, const FPointerEvent& MouseEvent) override
 	{
-		if (auto TableView = OwnerTablePtr.Pin())
+		if (TSharedPtr<ITypedTableView<TSharedPtr<RewindDebugger::FRewindDebuggerTrack>>> TableView = OwnerTablePtr.Pin())
 		{
 			if (const TSharedPtr<RewindDebugger::FRewindDebuggerTrack>* Track = GetItemForThis(TableView.ToSharedRef()))
 			{
@@ -66,7 +66,7 @@ public:
 	{
 		STableRow<TSharedPtr<RewindDebugger::FRewindDebuggerTrack>>::OnMouseLeave(MouseEvent);
 
-		if (auto TableView = OwnerTablePtr.Pin())
+		if (TSharedPtr<ITypedTableView<TSharedPtr<RewindDebugger::FRewindDebuggerTrack>>> TableView = OwnerTablePtr.Pin())
 		{
 			if (const TSharedPtr<RewindDebugger::FRewindDebuggerTrack>* Track = GetItemForThis(TableView.ToSharedRef()))
 			{
@@ -179,7 +179,7 @@ static void RestoreExpansion(TSharedPtr<RewindDebugger::FRewindDebuggerTrack> Tr
 
 void SRewindDebuggerComponentTree::RestoreExpansion()
 {
-	for (auto& Track : *DebugComponents)
+	for (TSharedPtr<RewindDebugger::FRewindDebuggerTrack>& Track : *DebugComponents)
 	{
 		::RestoreExpansion(Track, ComponentTreeView);
 	}
diff --git a/Engine/Source/Editor/RewindDebuggerInterface/Public/RewindDebuggerTrack.h b/Engine/Source/Editor/RewindDebuggerInterface/Public/RewindDebuggerTrack.h
--- a/Engine/Source/Editor/RewindDebuggerInterface/Public/RewindDebuggerTrack.h
+++ b/Engine/Source/Editor/RewindDebuggerInterface/Public/RewindDebuggerTrack.h
@@ -19,7 +19,7 @@ namespace RewindDebugger
 class FRewindDebuggerTrack
 {
 public:
-	FRewindDebuggerTrack(): bExpanded(true), bVisible(true)
+	FRewindDebuggerTrack()
 	{
 	}
 
@@ -159,8 +159,8 @@ private:
 	bool bSelected : 1 = false;
 	bool bTrackHovered : 1 = false;
 	bool bTreeHovered : 1 = false;
-	bool bExpanded : 1 = false;
-	bool bVisible : 1 = false;
+	bool bExpanded : 1 = true;
+	bool bVisible : 1 = true;
 };
 	
 }
\ No newline at end of file
