add UI for auto eject/record features in Rewind Debugger
- change default setting for Auto Eject to disabled

diff --git a/Engine/Plugins/Animation/GameplayInsights/Source/RewindDebugger/Private/RewindDebugger.cpp b/Engine/Plugins/Animation/GameplayInsights/Source/RewindDebugger/Private/RewindDebugger.cpp
--- a/Engine/Plugins/Animation/GameplayInsights/Source/RewindDebugger/Private/RewindDebugger.cpp
+++ b/Engine/Plugins/Animation/GameplayInsights/Source/RewindDebugger/Private/RewindDebugger.cpp
@@ -181,7 +181,7 @@ void FRewindDebugger::OnPIEPaused(bool bSimulating)
 		SetCurrentScrubTime(RecordingDuration.Get());
 	}
 	
-	if (ShouldAutoDetach() && FPlayWorldCommandCallbacks::IsInPIE())
+	if (ShouldAutoEject() && FPlayWorldCommandCallbacks::IsInPIE())
 	{
 		bool CanEject= false;
 		for (auto It = GUnrealEd->SlatePlayInEditorMap.CreateIterator(); It; ++It)
@@ -211,7 +211,7 @@ void FRewindDebugger::OnPIEResumed(bool bSimulating)
 
 	MeshComponentsToReset.Empty();
 
-	if (ShouldAutoDetach() && FPlayWorldCommandCallbacks::IsInSIE())
+	if (ShouldAutoEject() && FPlayWorldCommandCallbacks::IsInSIE())
 	{
 		GEditor->RequestToggleBetweenPIEandSIE();
 	}
@@ -477,14 +477,14 @@ void FRewindDebugger::SetShouldAutoRecordOnPIE(bool value)
 	URewindDebuggerSettings::Get().bShouldAutoRecordOnPIE = value;
 }
 
-bool FRewindDebugger::ShouldAutoDetach() const
+bool FRewindDebugger::ShouldAutoEject() const
 {
-	return URewindDebuggerSettings::Get().bShouldAutoDetach;
+	return URewindDebuggerSettings::Get().bShouldAutoEject;
 }
 
-void FRewindDebugger::SetShouldAutoDetach(bool value)
+void FRewindDebugger::SetShouldAutoEject(bool value)
 {
-	URewindDebuggerSettings::Get().bShouldAutoDetach = value;
+	URewindDebuggerSettings::Get().bShouldAutoEject = value;
 }
 
 void FRewindDebugger::StopRecording()
@@ -1357,6 +1357,22 @@ void FRewindDebugger::RegisterToolBar()
 				LOCTEXT("Blank",""),
 				TAttribute<FText>(),
 				FSlateIcon("RewindDebuggerStyle", "RewindDebugger.StopRecording.small")));
+
+	Section.AddSeparator(NAME_None);
+	
+	Section.AddEntry(FToolMenuEntry::InitToolBarButton(
+				Commands.AutoEject,
+				LOCTEXT("Blank",""),
+				TAttribute<FText>(),
+				FSlateIcon("RewindDebuggerStyle", "RewindDebugger.AutoEject")));
+					Section.AddSeparator(NAME_None);
+                				
+	Section.AddEntry(FToolMenuEntry::InitToolBarButton(
+				Commands.AutoRecord,
+				LOCTEXT("Blank",""),
+				TAttribute<FText>(),
+				FSlateIcon("RewindDebuggerStyle", "RewindDebugger.AutoRecord")));
+
 	
 	Menu->SetStyleSet(&FAppStyle::Get());
 	Menu->StyleName = "PaletteToolBar";
diff --git a/Engine/Plugins/Animation/GameplayInsights/Source/RewindDebugger/Private/RewindDebugger.h b/Engine/Plugins/Animation/GameplayInsights/Source/RewindDebugger/Private/RewindDebugger.h
--- a/Engine/Plugins/Animation/GameplayInsights/Source/RewindDebugger/Private/RewindDebugger.h
+++ b/Engine/Plugins/Animation/GameplayInsights/Source/RewindDebugger/Private/RewindDebugger.h
@@ -71,8 +71,8 @@ public:
 	bool ShouldAutoRecordOnPIE() const;
 	void SetShouldAutoRecordOnPIE(bool value);
 	
-	bool ShouldAutoDetach() const;
-	void SetShouldAutoDetach(bool value);
+	bool ShouldAutoEject() const;
+	void SetShouldAutoEject(bool value);
 
 	// Stop recording: Stop tracing Object + Animation Data.
 	void StopRecording();
diff --git a/Engine/Plugins/Animation/GameplayInsights/Source/RewindDebugger/Private/RewindDebuggerCommands.h b/Engine/Plugins/Animation/GameplayInsights/Source/RewindDebugger/Private/RewindDebuggerCommands.h
--- a/Engine/Plugins/Animation/GameplayInsights/Source/RewindDebugger/Private/RewindDebuggerCommands.h
+++ b/Engine/Plugins/Animation/GameplayInsights/Source/RewindDebugger/Private/RewindDebuggerCommands.h
@@ -31,6 +31,8 @@ public:
 		UI_COMMAND(Play, "Play", "Playback recorded data", EUserInterfaceActionType::Button, FInputChord(EKeys::Down));
 		UI_COMMAND(NextFrame, "Next Frame", "Step one frame forward", EUserInterfaceActionType::Button, FInputChord(EKeys::Right));
 		UI_COMMAND(LastFrame, "Last Frame", "Jump to last recorded frame", EUserInterfaceActionType::Button, FInputChord(EModifierKey::Control, EKeys::Up));
+		UI_COMMAND(AutoEject, "Auto Eject", "Automatically detach player control when PIE is paused", EUserInterfaceActionType::ToggleButton, FInputChord());
+		UI_COMMAND(AutoRecord, "Auto Record", "Automatically start recording when PIE is started", EUserInterfaceActionType::ToggleButton, FInputChord());
 	}
 
 
@@ -44,6 +46,8 @@ public:
 	TSharedPtr<FUICommandInfo> NextFrame;
 	TSharedPtr<FUICommandInfo> LastFrame;
 	TSharedPtr<FUICommandInfo> PauseOrPlay;
+	TSharedPtr<FUICommandInfo> AutoEject;
+	TSharedPtr<FUICommandInfo> AutoRecord;
 };
 
 
diff --git a/Engine/Plugins/Animation/GameplayInsights/Source/RewindDebugger/Private/RewindDebuggerModule.cpp b/Engine/Plugins/Animation/GameplayInsights/Source/RewindDebugger/Private/RewindDebuggerModule.cpp
--- a/Engine/Plugins/Animation/GameplayInsights/Source/RewindDebugger/Private/RewindDebuggerModule.cpp
+++ b/Engine/Plugins/Animation/GameplayInsights/Source/RewindDebugger/Private/RewindDebuggerModule.cpp
@@ -143,6 +143,19 @@ TSharedRef<SDockTab> FRewindDebuggerModule::SpawnRewindDebuggerTab(const FSpawnT
 							 FIsActionChecked(),
 							 FIsActionButtonVisible::CreateRaw(DebuggerInstance, &FRewindDebugger::CanStopRecording));
 
+	CommandList->MapAction(Commands.AutoEject,
+							 FExecuteAction::CreateLambda([DebuggerInstance]() { DebuggerInstance->SetShouldAutoEject(!DebuggerInstance->ShouldAutoEject()); }),
+							 FCanExecuteAction(),
+							 FCanExecuteAction::CreateRaw(DebuggerInstance, &FRewindDebugger::ShouldAutoEject),
+							 FIsActionButtonVisible());
+							 
+	CommandList->MapAction(Commands.AutoRecord,
+							 FExecuteAction::CreateLambda([DebuggerInstance]() { DebuggerInstance->SetShouldAutoRecordOnPIE(!DebuggerInstance->ShouldAutoRecordOnPIE()); }),
+							 FCanExecuteAction(),
+							 FCanExecuteAction::CreateRaw(DebuggerInstance, &FRewindDebugger::ShouldAutoRecordOnPIE),
+							 FIsActionButtonVisible());
+		 
+
 	// Register PIE Rewind Debugger Commands
 	if (GEditor != nullptr)
 	{
diff --git a/Engine/Plugins/Animation/GameplayInsights/Source/RewindDebugger/Private/RewindDebuggerSettings.cpp b/Engine/Plugins/Animation/GameplayInsights/Source/RewindDebugger/Private/RewindDebuggerSettings.cpp
--- a/Engine/Plugins/Animation/GameplayInsights/Source/RewindDebugger/Private/RewindDebuggerSettings.cpp
+++ b/Engine/Plugins/Animation/GameplayInsights/Source/RewindDebugger/Private/RewindDebuggerSettings.cpp
@@ -5,7 +5,7 @@
 
 #include UE_INLINE_GENERATED_CPP_BY_NAME(RewindDebuggerSettings)
 
-URewindDebuggerSettings::URewindDebuggerSettings() : CameraMode(ERewindDebuggerCameraMode::Replay), bShouldAutoDetach(true), bShouldAutoRecordOnPIE(false)
+URewindDebuggerSettings::URewindDebuggerSettings() : CameraMode(ERewindDebuggerCameraMode::Replay), bShouldAutoEject(false), bShouldAutoRecordOnPIE(false)
 {
 	FCoreDelegates::OnPreExit.AddLambda([]()
 	{
diff --git a/Engine/Plugins/Animation/GameplayInsights/Source/RewindDebugger/Private/RewindDebuggerSettings.h b/Engine/Plugins/Animation/GameplayInsights/Source/RewindDebugger/Private/RewindDebuggerSettings.h
--- a/Engine/Plugins/Animation/GameplayInsights/Source/RewindDebugger/Private/RewindDebuggerSettings.h
+++ b/Engine/Plugins/Animation/GameplayInsights/Source/RewindDebugger/Private/RewindDebuggerSettings.h
@@ -38,9 +38,9 @@ public:
 	UPROPERTY(EditAnywhere, Config, Category = Camera)
 	ERewindDebuggerCameraMode CameraMode;
 	
-	/** If enabled, automatically detach player controller at the start of PIE */
-	UPROPERTY(EditAnywhere, Config, Category = Other, meta = (DisplayName = "Auto Detach Player Controller on PIE"))
-	bool bShouldAutoDetach;
+	/** If enabled, automatically detach player control when PIE is paused */
+	UPROPERTY(EditAnywhere, Config, Category = Other)
+	bool bShouldAutoEject;
 	
 	/** If enabled, start recording information at the start of PIE */
 	UPROPERTY(EditAnywhere, Config, Category = Other)
diff --git a/Engine/Plugins/Animation/GameplayInsights/Source/RewindDebugger/Private/RewindDebuggerStyle.cpp b/Engine/Plugins/Animation/GameplayInsights/Source/RewindDebugger/Private/RewindDebuggerStyle.cpp
--- a/Engine/Plugins/Animation/GameplayInsights/Source/RewindDebugger/Private/RewindDebuggerStyle.cpp
+++ b/Engine/Plugins/Animation/GameplayInsights/Source/RewindDebugger/Private/RewindDebuggerStyle.cpp
@@ -30,6 +30,9 @@ FRewindDebuggerStyle::FRewindDebuggerStyle() :
 	Set("RewindDebugger.Play.small", new CORE_IMAGE_BRUSH("Editor/Slate/Sequencer/Transport_Bar/Play_24x", Icon24x24));
 	Set("RewindDebugger.NextFrame.small", new CORE_IMAGE_BRUSH("Editor/Slate/Sequencer/Transport_Bar/Step_Forward_24x", Icon24x24));
 	Set("RewindDebugger.LastFrame.small", new CORE_IMAGE_BRUSH("Editor/Slate/Sequencer/Transport_Bar/Go_To_End_24x", Icon24x24));
+	
+	Set("RewindDebugger.AutoEject", new IMAGE_BRUSH("autoeject", {30.f, 24.f}));
+	Set("RewindDebugger.AutoRecord", new IMAGE_BRUSH("autorecord", {30.f, 24.f}));
 
 	// Actor Picker
 	Set("RewindDebugger.SelectActor", new CORE_IMAGE_BRUSH("Editor/Slate/Icons/eyedropper_16px", Icon16x16));
