Fix rewind debugger recording time to work correctly when loading trace files

diff --git a/Engine/Plugins/Animation/GameplayInsights/Source/GameplayInsights/Private/GameplayProvider.cpp b/Engine/Plugins/Animation/GameplayInsights/Source/GameplayInsights/Private/GameplayProvider.cpp
--- a/Engine/Plugins/Animation/GameplayInsights/Source/GameplayInsights/Private/GameplayProvider.cpp
+++ b/Engine/Plugins/Animation/GameplayInsights/Source/GameplayInsights/Private/GameplayProvider.cpp
@@ -616,6 +616,11 @@ void FGameplayProvider::ReadViewTimeline(TFunctionRef<void(const IGameplayProvid
 	}
 }
 
+double FGameplayProvider::GetRecordingDuration() const
+{
+	return RecordingDuration;
+}
+
 void FGameplayProvider::AppendView(uint64 InPlayerId, double InTime, const FVector& InPosition, const FRotator& InRotation, float InFov, float InAspectRatio)
 {
 	Session.WriteAccessCheck();
@@ -677,6 +682,7 @@ void FGameplayProvider::AppendRecordingInfo(uint64 InWorldId, double InProfileTi
 				NewRecordingInfo.RecordingIndex = InRecordingIndex;
 				NewRecordingInfo.FrameIndex = InFrameIndex;
 				NewRecordingInfo.ElapsedTime = InElapsedTime;
+				RecordingDuration = FMath::Max(RecordingDuration,InElapsedTime);
 
 				if (TSharedRef<TraceServices::TPointTimeline<FRecordingInfoMessage>>* ExistingRecording = Recordings.Find(InRecordingIndex))
 				{
diff --git a/Engine/Plugins/Animation/GameplayInsights/Source/GameplayInsights/Private/GameplayProvider.h b/Engine/Plugins/Animation/GameplayInsights/Source/GameplayInsights/Private/GameplayProvider.h
--- a/Engine/Plugins/Animation/GameplayInsights/Source/GameplayInsights/Private/GameplayProvider.h
+++ b/Engine/Plugins/Animation/GameplayInsights/Source/GameplayInsights/Private/GameplayProvider.h
@@ -42,6 +42,7 @@ public:
 	virtual const TCHAR* GetPropertyName(uint32 InPropertyStringId) const override;
 	virtual const RecordingInfoTimeline* GetRecordingInfo(uint32 RecordingId) const override; 
 	virtual void ReadViewTimeline(TFunctionRef<void(const ViewTimeline&)> Callback) const override;
+	virtual double GetRecordingDuration() const override;
 
 
 	/** Add a class message */
@@ -185,6 +186,9 @@ private:
 
 	/** Whether we have any data */
 	bool bHasAnyData;
+	
+	/** Whether we have any data */
+	double RecordingDuration = 0.0;
 
 	/** Whether we have any object properties */
 	bool bHasObjectProperties;
diff --git a/Engine/Plugins/Animation/GameplayInsights/Source/GameplayInsights/Public/IGameplayProvider.h b/Engine/Plugins/Animation/GameplayInsights/Source/GameplayInsights/Public/IGameplayProvider.h
--- a/Engine/Plugins/Animation/GameplayInsights/Source/GameplayInsights/Public/IGameplayProvider.h
+++ b/Engine/Plugins/Animation/GameplayInsights/Source/GameplayInsights/Public/IGameplayProvider.h
@@ -135,4 +135,5 @@ public:
 	virtual uint64 FindPossessingController(uint64 Pawn, double Time) const = 0;
 	virtual TRange<double> GetObjectTraceLifetime(uint64 ObjectId) const = 0;
 	virtual TRange<double> GetObjectRecordingLifetime(uint64 ObjectId) const = 0;
+	virtual double GetRecordingDuration() const = 0;
 };
\ No newline at end of file
diff --git a/Engine/Plugins/Animation/GameplayInsights/Source/RewindDebugger/Private/RewindDebugger.cpp b/Engine/Plugins/Animation/GameplayInsights/Source/RewindDebugger/Private/RewindDebugger.cpp
--- a/Engine/Plugins/Animation/GameplayInsights/Source/RewindDebugger/Private/RewindDebugger.cpp
+++ b/Engine/Plugins/Animation/GameplayInsights/Source/RewindDebugger/Private/RewindDebugger.cpp
@@ -171,14 +171,13 @@ void FRewindDebugger::OnPIEPaused(bool bSimulating)
 {
 	bPIESimulating = false;
 	ControlState = EControlState::Pause;
-
+	
 	if (bRecording)
 	{
+    #if OBJECT_TRACE_ENABLED
 		UWorld* World = GetWorldToVisualize();
-#if OBJECT_TRACE_ENABLED
-		RecordingDuration.Set(FObjectTrace::GetWorldElapsedTime(World));
-#endif // OBJECT_TRACE_ENABLED
-		SetCurrentScrubTime(RecordingDuration.Get());
+		SetCurrentScrubTime(FObjectTrace::GetWorldElapsedTime(World));
+    #endif // OBJECT_TRACE_ENABLED
 	}
 	
 	if (ShouldAutoEject() && FPlayWorldCommandCallbacks::IsInPIE())
@@ -232,11 +231,10 @@ void FRewindDebugger::OnPIESingleStepped(bool bSimulating)
 
 	if (bRecording)
 	{
+    #if OBJECT_TRACE_ENABLED
 		UWorld* World = GetWorldToVisualize();
-#if OBJECT_TRACE_ENABLED
-		RecordingDuration.Set(FObjectTrace::GetWorldElapsedTime(World));
-#endif // OBJECT_TRACE_ENABLED
-		SetCurrentScrubTime(RecordingDuration.Get());
+		SetCurrentScrubTime(FObjectTrace::GetWorldElapsedTime(World));
+    #endif // OBJECT_TRACE_ENABLED
 	}
 }
 
@@ -248,9 +246,6 @@ void FRewindDebugger::OnPIEStopped(bool bSimulating)
 	MeshComponentsToReset.Empty();
 
 	StopRecording();
-	// clear the current recording (until we support playback in the Editor world on spawned actors)
-	RecordingDuration.Set(0);
-	SetCurrentScrubTime(0);
 }
 
 bool FRewindDebugger::GetTargetActorPosition(FVector& OutPosition) const
@@ -883,6 +878,9 @@ void FRewindDebugger::Tick(float DeltaTime)
 		if (AnimationProvider && GameplayProvider)
 		{
 			TraceServices::FAnalysisSessionReadScope SessionReadScope(*Session);
+
+			RecordingDuration.Set(GameplayProvider->GetRecordingDuration());
+			
 			UWorld* World = GetWorldToVisualize();
 
 			if (bPIESimulating)
@@ -890,16 +888,13 @@ void FRewindDebugger::Tick(float DeltaTime)
 				if (bRecording)
 				{
 					TRACE_CPUPROFILER_EVENT_SCOPE(FRewindDebugger::Tick_UpdateSimulating);
-#if OBJECT_TRACE_ENABLED
-					RecordingDuration.Set(FObjectTrace::GetWorldElapsedTime(World));
-#endif // OBJECT_TRACE_ENABLED
 					SetCurrentScrubTime(RecordingDuration.Get());
 					TrackCursorDelegate.ExecuteIfBound(false);
 				}
 			}
 			else
 			{
-				if (RecordingDuration.Get() > 0)
+				if (RecordingDuration.Get() > 0 && CurrentScrubTime <= RecordingDuration.Get()) 
 				{
 					if (ControlState == EControlState::Play || ControlState == EControlState::PlayReverse)
 					{
@@ -915,6 +910,8 @@ void FRewindDebugger::Tick(float DeltaTime)
 						}
 					}
 
+					SetCurrentScrubTime(CurrentScrubTime);// update trace time
+					
 					const double CurrentTraceTime = TraceTime.Get();
 					if (CurrentTraceTime != PreviousTraceTime)
 					{
