Fix for crash when missing a bone index in Foot Placement

diff --git a/Engine/Plugins/Animation/GameplayInsights/Source/RewindDebugger/Private/RewindDebugger.cpp b/Engine/Plugins/Animation/GameplayInsights/Source/RewindDebugger/Private/RewindDebugger.cpp
--- a/Engine/Plugins/Animation/GameplayInsights/Source/RewindDebugger/Private/RewindDebugger.cpp
+++ b/Engine/Plugins/Animation/GameplayInsights/Source/RewindDebugger/Private/RewindDebugger.cpp
@@ -469,7 +469,9 @@ bool FRewindDebugger::ShouldAutoRecordOnPIE() const
 
 void FRewindDebugger::SetShouldAutoRecordOnPIE(bool value)
 {
-	URewindDebuggerSettings::Get().bShouldAutoRecordOnPIE = value;
+	URewindDebuggerSettings& RewindDebuggerSettings = URewindDebuggerSettings::Get();
+	RewindDebuggerSettings.Modify();
+	RewindDebuggerSettings.bShouldAutoRecordOnPIE = value;
 }
 
 bool FRewindDebugger::ShouldAutoEject() const
@@ -479,7 +481,9 @@ bool FRewindDebugger::ShouldAutoEject() const
 
 void FRewindDebugger::SetShouldAutoEject(bool value)
 {
-	URewindDebuggerSettings::Get().bShouldAutoEject = value;
+	URewindDebuggerSettings& RewindDebuggerSettings = URewindDebuggerSettings::Get();
+	RewindDebuggerSettings.Modify();
+	RewindDebuggerSettings.bShouldAutoEject = value;
 }
 
 void FRewindDebugger::StopRecording()
diff --git a/Engine/Plugins/Animation/GameplayInsights/Source/RewindDebugger/Private/RewindDebuggerSettings.h b/Engine/Plugins/Animation/GameplayInsights/Source/RewindDebugger/Private/RewindDebuggerSettings.h
--- a/Engine/Plugins/Animation/GameplayInsights/Source/RewindDebugger/Private/RewindDebuggerSettings.h
+++ b/Engine/Plugins/Animation/GameplayInsights/Source/RewindDebugger/Private/RewindDebuggerSettings.h
@@ -19,7 +19,7 @@ enum class ERewindDebuggerCameraMode
 /**
  * Implements the settings for the Rewind Debugger.
  */
-UCLASS(Config=Editor, meta=(DisplayName="Rewind Debugger"))
+UCLASS(config = EditorPerProjectUserSettings, meta=(DisplayName="Rewind Debugger"))
 class URewindDebuggerSettings : public UDeveloperSettings
 {
 	GENERATED_BODY()
