[Backout] - CL31270597
Original CL Desc
-----------------------------------------------------------------
Fix for crash when missing a bone index in Foot Placement

diff --git a/Engine/Plugins/Animation/GameplayInsights/Source/RewindDebugger/Private/RewindDebugger.cpp b/Engine/Plugins/Animation/GameplayInsights/Source/RewindDebugger/Private/RewindDebugger.cpp
--- a/Engine/Plugins/Animation/GameplayInsights/Source/RewindDebugger/Private/RewindDebugger.cpp
+++ b/Engine/Plugins/Animation/GameplayInsights/Source/RewindDebugger/Private/RewindDebugger.cpp
@@ -469,9 +469,7 @@ bool FRewindDebugger::ShouldAutoRecordOnPIE() const
 
 void FRewindDebugger::SetShouldAutoRecordOnPIE(bool value)
 {
-	URewindDebuggerSettings& RewindDebuggerSettings = URewindDebuggerSettings::Get();
-	RewindDebuggerSettings.Modify();
-	RewindDebuggerSettings.bShouldAutoRecordOnPIE = value;
+	URewindDebuggerSettings::Get().bShouldAutoRecordOnPIE = value;
 }
 
 bool FRewindDebugger::ShouldAutoEject() const
@@ -481,9 +479,7 @@ bool FRewindDebugger::ShouldAutoEject() const
 
 void FRewindDebugger::SetShouldAutoEject(bool value)
 {
-	URewindDebuggerSettings& RewindDebuggerSettings = URewindDebuggerSettings::Get();
-	RewindDebuggerSettings.Modify();
-	RewindDebuggerSettings.bShouldAutoEject = value;
+	URewindDebuggerSettings::Get().bShouldAutoEject = value;
 }
 
 void FRewindDebugger::StopRecording()
diff --git a/Engine/Plugins/Animation/GameplayInsights/Source/RewindDebugger/Private/RewindDebuggerSettings.h b/Engine/Plugins/Animation/GameplayInsights/Source/RewindDebugger/Private/RewindDebuggerSettings.h
--- a/Engine/Plugins/Animation/GameplayInsights/Source/RewindDebugger/Private/RewindDebuggerSettings.h
+++ b/Engine/Plugins/Animation/GameplayInsights/Source/RewindDebugger/Private/RewindDebuggerSettings.h
@@ -19,7 +19,7 @@ enum class ERewindDebuggerCameraMode
 /**
  * Implements the settings for the Rewind Debugger.
  */
-UCLASS(config = EditorPerProjectUserSettings, meta=(DisplayName="Rewind Debugger"))
+UCLASS(Config=Editor, meta=(DisplayName="Rewind Debugger"))
 class URewindDebuggerSettings : public UDeveloperSettings
 {
 	GENERATED_BODY()
