build fix for release editor config

diff --git a/Engine/Plugins/Animation/GameplayInsights/Source/RewindDebuggerVLog/Private/RewindDebuggerVLog.cpp b/Engine/Plugins/Animation/GameplayInsights/Source/RewindDebuggerVLog/Private/RewindDebuggerVLog.cpp
--- a/Engine/Plugins/Animation/GameplayInsights/Source/RewindDebuggerVLog/Private/RewindDebuggerVLog.cpp
+++ b/Engine/Plugins/Animation/GameplayInsights/Source/RewindDebuggerVLog/Private/RewindDebuggerVLog.cpp
@@ -98,6 +98,7 @@ void FRewindDebuggerVLog::RenderLogEntry(const FVisualLogEntry& Entry)
 
 void FRewindDebuggerVLog::ImmediateRender(const UObject* Object, const FVisualLogEntry& Entry)
 {
+#if OBJECT_TRACE_ENABLED
 	if (IRewindDebugger* RewindDebugger = IRewindDebugger::Instance())
 	{
 		uint64 ObjectId = FObjectTrace::GetObjectId(Object);
@@ -106,6 +107,7 @@ void FRewindDebuggerVLog::ImmediateRender(const UObject* Object, const FVisualLo
 			RenderLogEntry(Entry);
 		}
 	}
+#endif
 }
 
 bool FRewindDebuggerVLog::IsCategoryActive(const FName& Category)
