Fix settings in rewind debugger not saving

diff --git a/Engine/Plugins/Animation/GameplayInsights/Source/RewindDebugger/Private/BindableProperty.h b/Engine/Plugins/Animation/GameplayInsights/Source/RewindDebugger/Private/BindableProperty.h
--- a/Engine/Plugins/Animation/GameplayInsights/Source/RewindDebugger/Private/BindableProperty.h
+++ b/Engine/Plugins/Animation/GameplayInsights/Source/RewindDebugger/Private/BindableProperty.h
@@ -56,6 +56,12 @@ struct TBindablePropertyInitializer
         Source = InSource;
         Value = T();
     }
+
+	TBindablePropertyInitializer(TBindableProperty<T, BindingType>* InSource, const T& InitialValue)
+    {
+		Source = InSource;
+		Value = InitialValue;
+    }
     
     TBindablePropertyInitializer(const T& InValue)
     {
@@ -145,6 +151,7 @@ class TBindableProperty
         {
             if (Initializer.Source)
             {
+            	InternalValue = Initializer.Value;
                 Bind(Initializer.Source);
             }
             else
diff --git a/Engine/Plugins/Animation/GameplayInsights/Source/RewindDebugger/Private/RewindDebugger.cpp b/Engine/Plugins/Animation/GameplayInsights/Source/RewindDebugger/Private/RewindDebugger.cpp
--- a/Engine/Plugins/Animation/GameplayInsights/Source/RewindDebugger/Private/RewindDebugger.cpp
+++ b/Engine/Plugins/Animation/GameplayInsights/Source/RewindDebugger/Private/RewindDebugger.cpp
@@ -95,8 +95,17 @@ FRewindDebugger::FRewindDebugger()  :
 	FEditorDelegates::EndPIE.AddRaw(this, &FRewindDebugger::OnPIEStopped);
 	FEditorDelegates::SingleStepPIE.AddRaw(this, &FRewindDebugger::OnPIESingleStepped);
 
+	
 	DebugTargetActor.OnPropertyChanged = DebugTargetActor.OnPropertyChanged.CreateLambda([this](FString Target)
 		{
+			URewindDebuggerSettings& Settings = URewindDebuggerSettings::Get();
+			if (Settings.DebugTargetActor != Target)
+			{
+				Settings.DebugTargetActor = Target;
+				Settings.Modify();
+				Settings.SaveConfig();
+			}
+		
 			TargetObjectIds.SetNum(0);
 			GetTargetObjectIds(TargetObjectIds);
 			// make sure all the SubObjects of the target actor have been traced
@@ -112,7 +121,7 @@ FRewindDebugger::FRewindDebugger()  :
 
 			RefreshDebugTracks();
 		});
-
+	
 	UnrealInsightsModule = &FModuleManager::LoadModuleChecked<IUnrealInsightsModule>("TraceInsights");
 
 	TickerHandle = FTSTicker::GetCoreTicker().AddTicker(TEXT("RewindDebugger"), 0.0f, [this](float DeltaTime)
@@ -469,7 +478,10 @@ bool FRewindDebugger::ShouldAutoRecordOnPIE() const
 
 void FRewindDebugger::SetShouldAutoRecordOnPIE(bool value)
 {
-	URewindDebuggerSettings::Get().bShouldAutoRecordOnPIE = value;
+	URewindDebuggerSettings& RewindDebuggerSettings = URewindDebuggerSettings::Get();
+	RewindDebuggerSettings.Modify();
+	RewindDebuggerSettings.bShouldAutoRecordOnPIE = value;
+	RewindDebuggerSettings.SaveConfig();
 }
 
 bool FRewindDebugger::ShouldAutoEject() const
@@ -479,7 +491,10 @@ bool FRewindDebugger::ShouldAutoEject() const
 
 void FRewindDebugger::SetShouldAutoEject(bool value)
 {
-	URewindDebuggerSettings::Get().bShouldAutoEject = value;
+	URewindDebuggerSettings& RewindDebuggerSettings = URewindDebuggerSettings::Get();
+	RewindDebuggerSettings.Modify();
+	RewindDebuggerSettings.bShouldAutoEject = value;
+	RewindDebuggerSettings.SaveConfig();
 }
 
 void FRewindDebugger::StopRecording()
@@ -818,7 +833,7 @@ const TraceServices::IAnalysisSession* FRewindDebugger::GetAnalysisSession() con
 		UnrealInsightsModule = &FModuleManager::LoadModuleChecked<IUnrealInsightsModule>("TraceInsights");
 	}
 
-	return UnrealInsightsModule->GetAnalysisSession().Get();
+	return UnrealInsightsModule ? UnrealInsightsModule->GetAnalysisSession().Get() : nullptr;
 }
 
 void FRewindDebugger::Tick(float DeltaTime)
diff --git a/Engine/Plugins/Animation/GameplayInsights/Source/RewindDebugger/Private/RewindDebuggerCamera.cpp b/Engine/Plugins/Animation/GameplayInsights/Source/RewindDebugger/Private/RewindDebuggerCamera.cpp
--- a/Engine/Plugins/Animation/GameplayInsights/Source/RewindDebugger/Private/RewindDebuggerCamera.cpp
+++ b/Engine/Plugins/Animation/GameplayInsights/Source/RewindDebugger/Private/RewindDebuggerCamera.cpp
@@ -90,6 +90,8 @@ void FRewindDebuggerCamera::SetCameraMode(ERewindDebuggerCameraMode InMode)
 	}
 	
 	RewindDebuggerSettings.CameraMode = InMode;
+	RewindDebuggerSettings.Modify();
+	RewindDebuggerSettings.SaveConfig();
 }
 
 void FRewindDebuggerCamera::Update(float DeltaTime, IRewindDebugger* RewindDebugger)
@@ -110,73 +112,75 @@ void FRewindDebuggerCamera::Update(float DeltaTime, IRewindDebugger* RewindDebug
 			bool bCameraTraceDataFound = false;
 			
 			FLevelEditorModule& LevelEditor = FModuleManager::GetModuleChecked<FLevelEditorModule>("LevelEditor");
-			SLevelViewport* LevelViewport = LevelEditor.GetFirstActiveLevelViewport().Get();
-			FLevelEditorViewportClient& LevelViewportClient = LevelViewport->GetLevelViewportClient();
+			if (SLevelViewport* LevelViewport = LevelEditor.GetFirstActiveLevelViewport().Get())
+			{
+				FLevelEditorViewportClient& LevelViewportClient = LevelViewport->GetLevelViewportClient();
 
-			FVector TargetActorPosition;
-			bool bTargetActorPositionValid = RewindDebugger->GetTargetActorPosition(TargetActorPosition);
+				FVector TargetActorPosition;
+				bool bTargetActorPositionValid = RewindDebugger->GetTargetActorPosition(TargetActorPosition);
 
 
-			if (CameraMode() == ERewindDebuggerCameraMode::FollowTargetActor)
-			{
-				// Follow Actor mode: apply position changes from the target actor to the camera
-				if (bTargetActorPositionValid)
+				if (CameraMode() == ERewindDebuggerCameraMode::FollowTargetActor)
 				{
-					if(LastPositionValid)
+					// Follow Actor mode: apply position changes from the target actor to the camera
+					if (bTargetActorPositionValid)
 					{
-						LevelViewportClient.SetViewLocation(LevelViewportClient.GetViewLocation() + TargetActorPosition - LastPosition);
+						if(LastPositionValid)
+						{
+							LevelViewportClient.SetViewLocation(LevelViewportClient.GetViewLocation() + TargetActorPosition - LastPosition);
+						}
 					}
 				}
-			}
 
-			// always update the camera actor to the replay values even if it isn't locked
-			if (const IGameplayProvider* GameplayProvider = Session->ReadProvider<IGameplayProvider>("GameplayProvider"))
-			{
-				GameplayProvider->ReadViewTimeline([&bCameraTraceDataFound, this, RewindDebugger, CurrentTraceTime, Session](const IGameplayProvider::ViewTimeline& TimelineData)
+				// always update the camera actor to the replay values even if it isn't locked
+				if (const IGameplayProvider* GameplayProvider = Session->ReadProvider<IGameplayProvider>("GameplayProvider"))
 				{
-					const TraceServices::IFrameProvider& FrameProvider = TraceServices::ReadFrameProvider(*Session);
-					TraceServices::FFrame Frame;
-					if(FrameProvider.GetFrameFromTime(ETraceFrameType::TraceFrameType_Game, CurrentTraceTime, Frame))
+					GameplayProvider->ReadViewTimeline([&bCameraTraceDataFound, this, RewindDebugger, CurrentTraceTime, Session](const IGameplayProvider::ViewTimeline& TimelineData)
 					{
-						TimelineData.EnumerateEvents(Frame.StartTime, Frame.EndTime,
-							[&bCameraTraceDataFound, this, RewindDebugger](double InStartTime, double InEndTime, uint32 InDepth, const FViewMessage& ViewMessage)
-							{
-								if (!CameraActor.IsValid())
+						const TraceServices::IFrameProvider& FrameProvider = TraceServices::ReadFrameProvider(*Session);
+						TraceServices::FFrame Frame;
+						if(FrameProvider.GetFrameFromTime(ETraceFrameType::TraceFrameType_Game, CurrentTraceTime, Frame))
+						{
+							TimelineData.EnumerateEvents(Frame.StartTime, Frame.EndTime,
+								[&bCameraTraceDataFound, this, RewindDebugger](double InStartTime, double InEndTime, uint32 InDepth, const FViewMessage& ViewMessage)
 								{
-									FActorSpawnParameters SpawnParameters;
-									SpawnParameters.ObjectFlags |= RF_Transient;
-									CameraActor = RewindDebugger->GetWorldToVisualize()->SpawnActor<ACameraActor>(ViewMessage.Position, ViewMessage.Rotation, SpawnParameters);
-									CameraActor->SetActorLabel("RewindDebuggerCamera");
-								}
-
-								UCameraComponent* Camera = CameraActor->GetCameraComponent();
-								Camera->SetWorldLocationAndRotation(ViewMessage.Position, ViewMessage.Rotation);
-								Camera->SetFieldOfView(ViewMessage.Fov);
-								Camera->SetAspectRatio(ViewMessage.AspectRatio);
-
-								bCameraTraceDataFound = true;
-
-								return TraceServices::EEventEnumerate::Stop;
-							});
-					}
-				});
-			}
+									if (!CameraActor.IsValid())
+									{
+										FActorSpawnParameters SpawnParameters;
+										SpawnParameters.ObjectFlags |= RF_Transient;
+										CameraActor = RewindDebugger->GetWorldToVisualize()->SpawnActor<ACameraActor>(ViewMessage.Position, ViewMessage.Rotation, SpawnParameters);
+										CameraActor->SetActorLabel("RewindDebuggerCamera");
+									}
+
+									UCameraComponent* Camera = CameraActor->GetCameraComponent();
+									Camera->SetWorldLocationAndRotation(ViewMessage.Position, ViewMessage.Rotation);
+									Camera->SetFieldOfView(ViewMessage.Fov);
+									Camera->SetAspectRatio(ViewMessage.AspectRatio);
+
+									bCameraTraceDataFound = true;
+
+									return TraceServices::EEventEnumerate::Stop;
+								});
+						}
+					});
+				}
 
-			if (CameraMode() == ERewindDebuggerCameraMode::Replay)
-			{
-				if (CameraActor.IsValid())
+				if (CameraMode() == ERewindDebuggerCameraMode::Replay)
 				{
-					LevelViewportClient.SetActorLock(CameraActor.Get());
+					if (CameraActor.IsValid())
+					{
+						LevelViewportClient.SetActorLock(CameraActor.Get());
+					}
 				}
-			}
 
-			LastPosition = TargetActorPosition;
-			LastPositionValid = bTargetActorPositionValid;
+				LastPosition = TargetActorPosition;
+				LastPositionValid = bTargetActorPositionValid;
 
-			if (bCameraTraceDataFound) // don't update this if there was no trace data found, because when first pausing, it can take a few frames for latest data to get processed
-			{
-				// only update camera in playback or scrubbing when the time has changed (allow free movement when paused)
-				LastCameraScrubTime = CurrentTraceTime;
+				if (bCameraTraceDataFound) // don't update this if there was no trace data found, because when first pausing, it can take a few frames for latest data to get processed
+				{
+					// only update camera in playback or scrubbing when the time has changed (allow free movement when paused)
+					LastCameraScrubTime = CurrentTraceTime;
+				}
 			}
 		}
 	}
diff --git a/Engine/Plugins/Animation/GameplayInsights/Source/RewindDebugger/Private/RewindDebuggerModule.cpp b/Engine/Plugins/Animation/GameplayInsights/Source/RewindDebugger/Private/RewindDebuggerModule.cpp
--- a/Engine/Plugins/Animation/GameplayInsights/Source/RewindDebugger/Private/RewindDebuggerModule.cpp
+++ b/Engine/Plugins/Animation/GameplayInsights/Source/RewindDebugger/Private/RewindDebuggerModule.cpp
@@ -175,7 +175,7 @@ TSharedRef<SDockTab> FRewindDebuggerModule::SpawnRewindDebuggerTab(const FSpawnT
 	}
 	
 	RewindDebuggerWidget = SNew(SRewindDebugger, CommandList.ToSharedRef(), MajorTab, SpawnTabArgs.GetOwnerWindow())
-								.DebugTargetActor(DebuggerInstance->GetDebugTargetActorProperty())
+								.DebugTargetActor({ DebuggerInstance->GetDebugTargetActorProperty(), URewindDebuggerSettings::Get().DebugTargetActor})
 								.RecordingDuration(DebuggerInstance->GetRecordingDurationProperty())
 								.DebugComponents(&DebuggerInstance->GetDebugTracks())
 								.TraceTime(DebuggerInstance->GetTraceTimeProperty())
diff --git a/Engine/Plugins/Animation/GameplayInsights/Source/RewindDebugger/Private/RewindDebuggerSettings.h b/Engine/Plugins/Animation/GameplayInsights/Source/RewindDebugger/Private/RewindDebuggerSettings.h
--- a/Engine/Plugins/Animation/GameplayInsights/Source/RewindDebugger/Private/RewindDebuggerSettings.h
+++ b/Engine/Plugins/Animation/GameplayInsights/Source/RewindDebugger/Private/RewindDebuggerSettings.h
@@ -19,7 +19,7 @@ enum class ERewindDebuggerCameraMode
 /**
  * Implements the settings for the Rewind Debugger.
  */
-UCLASS(Config=Editor, meta=(DisplayName="Rewind Debugger"))
+UCLASS(config = EditorPerProjectUserSettings, meta=(DisplayName="Rewind Debugger"))
 class URewindDebuggerSettings : public UDeveloperSettings
 {
 	GENERATED_BODY()
@@ -53,6 +53,10 @@ public:
 	/** The track types listed here will be hidden from the track tree view */
 	UPROPERTY(EditAnywhere, Config, Category = Filters)
 	TArray<FName> HiddenTrackTypes;
+
+	/** The track types listed here will be hidden from the track tree view */
+	UPROPERTY(Config)
+	FString DebugTargetActor;
 	
 	/** Get Mutable CDO of URewindDebuggerSettings */
 	static URewindDebuggerSettings & Get();
