Fix warnings when opening assets from rewind debugger, and also make assets scrub to the correct frame when opening via double click

diff --git a/Engine/Plugins/Animation/GameplayInsights/Source/GameplayInsights/Private/BlendWeightsTrack.cpp b/Engine/Plugins/Animation/GameplayInsights/Source/GameplayInsights/Private/BlendWeightsTrack.cpp
--- a/Engine/Plugins/Animation/GameplayInsights/Source/GameplayInsights/Private/BlendWeightsTrack.cpp
+++ b/Engine/Plugins/Animation/GameplayInsights/Source/GameplayInsights/Private/BlendWeightsTrack.cpp
@@ -10,9 +10,13 @@
 
 #if WITH_EDITOR
 #include "Animation/AnimInstance.h"
+#include "AnimPreviewInstance.h"
 #include "Animation/AnimSequence.h"
 #include "Animation/BlendSpace.h"
+#include "Animation/DebugSkelMeshComponent.h"
 #include "Editor.h"
+#include "IAnimationEditor.h"
+#include "IPersonaToolkit.h"
 #include "Subsystems/AssetEditorSubsystem.h"
 #endif
 
@@ -341,8 +345,79 @@ bool FBlendWeightTrack::HandleDoubleClickInternal()
 		const FGameplayProvider* GameplayProvider = AnalysisSession->ReadProvider<FGameplayProvider>(FGameplayProvider::ProviderName);
 
 		const FObjectInfo& AssetInfo = GameplayProvider->GetObjectInfo(GetAssetId());
+		
+		UObject* Asset = nullptr;
+		bool bMessageFound = false;
+		float PlaybackTime = 0;
+		float BlendSpaceX = 0;
+		float BlendSpaceY = 0;
+
+		float CurrentTraceTime = IRewindDebugger::Instance()->CurrentTraceTime();
+		
+		const FAnimationProvider* AnimationProvider = AnalysisSession->ReadProvider<FAnimationProvider>(FAnimationProvider::ProviderName);
+		const TraceServices::IFrameProvider& FrameProvider = TraceServices::ReadFrameProvider(*AnalysisSession);
+		TraceServices::FFrame Frame;
+		if (FrameProvider.GetFrameFromTime(ETraceFrameType::TraceFrameType_Game, CurrentTraceTime, Frame))
+		{
+			AnimationProvider->ReadTickRecordTimeline(ObjectId, [this, &bMessageFound, &PlaybackTime, &BlendSpaceX, &BlendSpaceY, &GameplayProvider, &Frame](const FAnimationProvider::TickRecordTimeline& InTimeline)
+			{
+				InTimeline.EnumerateEvents(Frame.StartTime, Frame.EndTime, [this, &bMessageFound, &PlaybackTime, &BlendSpaceX, &BlendSpaceY, &GameplayProvider, &Frame](double InStartTime, double InEndTime, uint32 InDepth, const FTickRecordMessage& InMessage)
+				{
+					if(InStartTime >= Frame.StartTime && InEndTime <= Frame.EndTime)
+					{
+						if (InMessage.NodeId == NodeId && InMessage.AssetId == AssetId)
+						{
+							bMessageFound = true;
+							PlaybackTime = InMessage.PlaybackTime;
+							BlendSpaceX = InMessage.BlendSpacePositionX;
+							BlendSpaceY = InMessage.BlendSpacePositionY;
+							return TraceServices::EEventEnumerate::Stop;
+						}
+					}
+					return TraceServices::EEventEnumerate::Continue;
+				});
+			});
+		}
+		
+		FString PackagePathString = FPackageName::ObjectPathToPackageName(FString(AssetInfo.PathName));
+
+		UPackage* Package = LoadPackage(NULL, ToCStr(PackagePathString), LOAD_NoRedirects);
+		if (Package)
+		{
+			Package->FullyLoad();
+                
+			FString AssetName = FPaths::GetBaseFilename(AssetInfo.PathName);
+			Asset = FindObject<UObject>(Package, *AssetName);
+		}
+		else
+		{
+			// fallback for unsaved assets
+			Asset = FindObject<UObject>(nullptr, AssetInfo.PathName);
+		}
+                    	
+		if (Asset != nullptr)
+		{
+			if (UAssetEditorSubsystem* AssetEditorSS = GEditor->GetEditorSubsystem<UAssetEditorSubsystem>())
+			{
+				AssetEditorSS->OpenEditorForAsset(Asset);
 
-		GEditor->GetEditorSubsystem<UAssetEditorSubsystem>()->OpenEditorForAsset(AssetInfo.PathName);
+				if (bMessageFound)
+				{
+					// if the asset is playing on the current frame, scrub to the appropriate time
+					if (IAssetEditorInstance* Editor = AssetEditorSS->FindEditorForAsset(Asset, true))
+					{
+						if (Editor->GetEditorName()=="AnimationEditor")
+						{
+							IAnimationEditor* AnimationEditor = static_cast<IAnimationEditor*>(Editor);
+							UDebugSkelMeshComponent* PreviewComponent = AnimationEditor->GetPersonaToolkit()->GetPreviewMeshComponent();
+							PreviewComponent->PreviewInstance->SetPosition(PlaybackTime);
+							PreviewComponent->PreviewInstance->SetPlaying(false);
+							PreviewComponent->PreviewInstance->SetBlendSpacePosition(FVector(BlendSpaceX, BlendSpaceY, 0.0f));
+						}
+					}
+				}
+			}
+		}
 
 		return true;
 	}
diff --git a/Engine/Plugins/Animation/GameplayInsights/Source/GameplayInsights/Private/MontagesTrack.cpp b/Engine/Plugins/Animation/GameplayInsights/Source/GameplayInsights/Private/MontagesTrack.cpp
--- a/Engine/Plugins/Animation/GameplayInsights/Source/GameplayInsights/Private/MontagesTrack.cpp
+++ b/Engine/Plugins/Animation/GameplayInsights/Source/GameplayInsights/Private/MontagesTrack.cpp
@@ -8,8 +8,12 @@
 #include "Styling/SlateIconFinder.h"
 #include "SMontageView.h"
 #if WITH_EDITOR
+#include "AnimPreviewInstance.h"
 #include "Animation/AnimMontage.h"
+#include "Animation/DebugSkelMeshComponent.h"
 #include "Editor.h"
+#include "IAnimationEditor.h"
+#include "IPersonaToolkit.h"
 #include "Subsystems/AssetEditorSubsystem.h"
 #endif
 
@@ -302,8 +306,76 @@ bool FMontageTrack::HandleDoubleClickInternal()
 		const FGameplayProvider* GameplayProvider = AnalysisSession->ReadProvider<FGameplayProvider>(FGameplayProvider::ProviderName);
 
 		const FObjectInfo& AssetInfo = GameplayProvider->GetObjectInfo(GetAssetId());
+	
+		UObject* Asset = nullptr;
+		bool bMessageFound = false;
+		float PlaybackTime = 0;
+
+		float CurrentTraceTime = IRewindDebugger::Instance()->CurrentTraceTime();
+		
+		const FAnimationProvider* AnimationProvider = AnalysisSession->ReadProvider<FAnimationProvider>(FAnimationProvider::ProviderName);
+		const TraceServices::IFrameProvider& FrameProvider = TraceServices::ReadFrameProvider(*AnalysisSession);
+		TraceServices::FFrame Frame;
+		if (FrameProvider.GetFrameFromTime(ETraceFrameType::TraceFrameType_Game, CurrentTraceTime, Frame))
+		{
+			AnimationProvider->ReadMontageTimeline(ObjectId, [this, &bMessageFound, &PlaybackTime, &GameplayProvider, &Frame](const FAnimationProvider::AnimMontageTimeline& InTimeline)
+			{
+				InTimeline.EnumerateEvents(Frame.StartTime, Frame.EndTime, [this, &bMessageFound, &PlaybackTime, &GameplayProvider, &Frame](double InStartTime, double InEndTime, uint32 InDepth, const FAnimMontageMessage& InMessage)
+				{
+					if(InStartTime >= Frame.StartTime && InEndTime <= Frame.EndTime)
+					{
+						if (InMessage.MontageId == AssetId)
+						{
+							bMessageFound = true;
+							PlaybackTime = InMessage.Position;
+							return TraceServices::EEventEnumerate::Stop;
+						}
+					}
+					return TraceServices::EEventEnumerate::Continue;
+				});
+			});
+		}
+		
+		FString PackagePathString = FPackageName::ObjectPathToPackageName(FString(AssetInfo.PathName));
+
+		UPackage* Package = LoadPackage(NULL, ToCStr(PackagePathString), LOAD_NoRedirects);
+		if (Package)
+		{
+			Package->FullyLoad();
+                
+			FString AssetName = FPaths::GetBaseFilename(AssetInfo.PathName);
+			Asset = FindObject<UObject>(Package, *AssetName);
+		}
+		else
+		{
+			// fallback for unsaved assets
+			Asset = FindObject<UObject>(nullptr, AssetInfo.PathName);
+		}
+                    	
+		if (Asset != nullptr)
+		{
+			if (UAssetEditorSubsystem* AssetEditorSS = GEditor->GetEditorSubsystem<UAssetEditorSubsystem>())
+			{
+				AssetEditorSS->OpenEditorForAsset(Asset);
 
-		GEditor->GetEditorSubsystem<UAssetEditorSubsystem>()->OpenEditorForAsset(AssetInfo.PathName);
+				if (bMessageFound)
+				{
+					// if the asset is playing on the current frame, scrub to the appropriate time
+					if (IAssetEditorInstance* Editor = AssetEditorSS->FindEditorForAsset(Asset, true))
+					{
+						if (Editor->GetEditorName()=="AnimationEditor")
+						{
+							IAnimationEditor* AnimationEditor = static_cast<IAnimationEditor*>(Editor);
+							UDebugSkelMeshComponent* PreviewComponent = AnimationEditor->GetPersonaToolkit()->GetPreviewMeshComponent();
+							PreviewComponent->PreviewInstance->SetPosition(PlaybackTime);
+							PreviewComponent->PreviewInstance->SetPlaying(false);
+						}
+					}
+				}
+			}
+		}
+
+		
 
 		return true;
 	}
diff --git a/Engine/Plugins/Animation/GameplayInsights/Source/GameplayInsights/Private/SVariantValueView.cpp b/Engine/Plugins/Animation/GameplayInsights/Source/GameplayInsights/Private/SVariantValueView.cpp
--- a/Engine/Plugins/Animation/GameplayInsights/Source/GameplayInsights/Private/SVariantValueView.cpp
+++ b/Engine/Plugins/Animation/GameplayInsights/Source/GameplayInsights/Private/SVariantValueView.cpp
@@ -168,8 +168,10 @@ TSharedRef<SWidget> SVariantValueView::MakeVariantValueWidget(const TraceService
 				.OnNavigate_Lambda([ObjectInfo, InValue]()
 				{
 					UObject* Asset = nullptr;
-						
-					UPackage* Package = LoadPackage(NULL, ObjectInfo.PathName, LOAD_NoRedirects);
+					
+					FString PackagePathString = FPackageName::ObjectPathToPackageName(FString(ObjectInfo.PathName));
+
+					UPackage* Package = LoadPackage(NULL, ToCStr(PackagePathString), LOAD_NoRedirects);
 					if (Package)
 					{
 						Package->FullyLoad();
