fix for stack overflows in rewind debugger

diff --git a/Engine/Plugins/Animation/GameplayInsights/Source/RewindDebugger/Private/SRewindDebugger.cpp b/Engine/Plugins/Animation/GameplayInsights/Source/RewindDebugger/Private/SRewindDebugger.cpp
--- a/Engine/Plugins/Animation/GameplayInsights/Source/RewindDebugger/Private/SRewindDebugger.cpp
+++ b/Engine/Plugins/Animation/GameplayInsights/Source/RewindDebugger/Private/SRewindDebugger.cpp
@@ -268,7 +268,15 @@ void SRewindDebugger::Construct(const FArguments& InArgs, TSharedRef<FUICommandL
 
 	ComponentTreeView =	SNew(SRewindDebuggerComponentTree)
 		.ExternalScrollBar(ScrollBar)
-		.OnExpansionChanged_Lambda([this]() { TimelinesView->RestoreExpansion(); })
+		.OnExpansionChanged_Lambda([this]()
+		{
+			if (!bInExpansionChanged)
+			{
+				bInExpansionChanged = true;
+				TimelinesView->RestoreExpansion();
+				bInExpansionChanged = false;
+			}
+		})
 		.OnScrolled_Lambda([this](double ScrollOffset){ TimelinesView->ScrollTo(ScrollOffset); })
 		.DebugComponents(InArgs._DebugComponents)
 		.OnMouseButtonDoubleClick(InArgs._OnComponentDoubleClicked)
@@ -287,7 +295,16 @@ void SRewindDebugger::Construct(const FArguments& InArgs, TSharedRef<FUICommandL
 
 	 TimelinesView = SNew(SRewindDebuggerTimelines)
 		.ExternalScrollbar(ScrollBar)
-		.OnExpansionChanged_Lambda([this]() { ComponentTreeView->RestoreExpansion(); })
+		.OnExpansionChanged_Lambda(
+				[this]()
+				{
+					if (!bInExpansionChanged)
+					{
+						bInExpansionChanged = true;
+						ComponentTreeView->RestoreExpansion();
+						bInExpansionChanged = false;
+					}
+				})
 		.OnScrolled_Lambda([this](double ScrollOffset){ ComponentTreeView->ScrollTo(ScrollOffset); })
 		.DebugComponents(InArgs._DebugComponents)
 		.ViewRange_Lambda([this](){return ViewRange;})
diff --git a/Engine/Plugins/Animation/GameplayInsights/Source/RewindDebugger/Private/SRewindDebugger.h b/Engine/Plugins/Animation/GameplayInsights/Source/RewindDebugger/Private/SRewindDebugger.h
--- a/Engine/Plugins/Animation/GameplayInsights/Source/RewindDebugger/Private/SRewindDebugger.h
+++ b/Engine/Plugins/Animation/GameplayInsights/Source/RewindDebugger/Private/SRewindDebugger.h
@@ -119,7 +119,8 @@ private:
 	TSharedPtr<SWidget> OnContextMenuOpening();
 	
 	TSharedPtr<SSearchBox> TrackFilterBox;
-	
+
+	bool bInExpansionChanged = false;
 	bool bInSelectionChanged = false;
 	bool bDisplayEmptyTracks = false;
 
