fix rewind debugger crash with autorecord and opening rewind debugger after pie is started

diff --git a/Engine/Plugins/Animation/GameplayInsights/Source/RewindDebugger/Private/RewindDebugger.cpp b/Engine/Plugins/Animation/GameplayInsights/Source/RewindDebugger/Private/RewindDebugger.cpp
--- a/Engine/Plugins/Animation/GameplayInsights/Source/RewindDebugger/Private/RewindDebugger.cpp
+++ b/Engine/Plugins/Animation/GameplayInsights/Source/RewindDebugger/Private/RewindDebugger.cpp
@@ -83,6 +83,8 @@ FRewindDebugger::FRewindDebugger()  :
     });
 	
 	RecordingDuration.Set(0);
+	
+	UnrealInsightsModule = &FModuleManager::LoadModuleChecked<IUnrealInsightsModule>("TraceInsights");
 
 	if (GEditor->bIsSimulatingInEditor || GEditor->PlayWorld)
 	{
@@ -121,8 +123,6 @@ FRewindDebugger::FRewindDebugger()  :
 
 			RefreshDebugTracks();
 		});
-	
-	UnrealInsightsModule = &FModuleManager::LoadModuleChecked<IUnrealInsightsModule>("TraceInsights");
 
 	TickerHandle = FTSTicker::GetCoreTicker().AddTicker(TEXT("RewindDebugger"), 0.0f, [this](float DeltaTime)
 	{
@@ -461,14 +461,10 @@ void FRewindDebugger::StartRecording()
 
 	FTraceAuxiliary::OnConnection.AddRaw(this, &FRewindDebugger::OnConnection); 
 
-	
 	FTraceAuxiliary::Start(FTraceAuxiliary::EConnectionType::Network, TEXT("127.0.0.1"), TEXT(""), &Options, LogRewindDebugger);
-	UnrealInsightsModule->StartAnalysisForLastLiveSession();
-
-
+	UnrealInsightsModule->StartAnalysisForLastLiveSession(5.0);
 
 	TargetObjectIds.Empty(2);
-
 }
 
 bool FRewindDebugger::ShouldAutoRecordOnPIE() const
