Rewind Debugger - Move animation specific code out of main update loop into an extension

diff --git a/Engine/Plugins/Animation/GameplayInsights/Source/RewindDebugger/Private/RewindDebugger.cpp b/Engine/Plugins/Animation/GameplayInsights/Source/RewindDebugger/Private/RewindDebugger.cpp
--- a/Engine/Plugins/Animation/GameplayInsights/Source/RewindDebugger/Private/RewindDebugger.cpp
+++ b/Engine/Plugins/Animation/GameplayInsights/Source/RewindDebugger/Private/RewindDebugger.cpp
@@ -2,12 +2,10 @@
 
 #include "RewindDebugger.h"
 
-#include "Animation/AnimBlueprint.h"
 #include "Animation/AnimBlueprintGeneratedClass.h"
 #include "Animation/AnimTrace.h"
 #include "Components/SkeletalMeshComponent.h"
 #include "Editor.h"
-#include "Engine/PoseWatch.h"
 #include "Engine/World.h"
 #include "EngineUtils.h"
 #include "GameFramework/Controller.h"
@@ -208,17 +206,6 @@ void FRewindDebugger::OnPIEResumed(bool bSimulating)
 {
 	bPIESimulating = true;
 
-	// restore all relative transforms of any meshes that may have been moved while scrubbing
-	for (TTuple<uint64, FMeshComponentResetData>& MeshData : MeshComponentsToReset)
-	{
-		if (USkeletalMeshComponent* MeshComponent = MeshData.Value.Component.Get())
-		{
-			MeshComponent->SetRelativeTransform(MeshData.Value.RelativeTransform, false, nullptr, ETeleportType::TeleportPhysics);
-		}
-	}
-
-	MeshComponentsToReset.Empty();
-
 	if (ShouldAutoEject() && FPlayWorldCommandCallbacks::IsInSIE())
 	{
 		GEditor->RequestToggleBetweenPIEandSIE();
@@ -227,17 +214,6 @@ void FRewindDebugger::OnPIEResumed(bool bSimulating)
 
 void FRewindDebugger::OnPIESingleStepped(bool bSimulating)
 {
-	// restore all relative transforms of any meshes that may have been moved while scrubbing
-	for (TTuple<uint64, FMeshComponentResetData>& MeshData : MeshComponentsToReset)
-	{
-		if (USkeletalMeshComponent* MeshComponent = MeshData.Value.Component.Get())
-		{
-			MeshComponent->SetRelativeTransform(MeshData.Value.RelativeTransform, false, nullptr, ETeleportType::TeleportPhysics);
-		}
-	}
-
-	MeshComponentsToReset.Empty();
-
 	if (bRecording)
 	{
     #if OBJECT_TRACE_ENABLED
@@ -252,7 +228,6 @@ void FRewindDebugger::OnPIEStopped(bool bSimulating)
 {
 	bPIEStarted = false;
 	bPIESimulating = false;
-	MeshComponentsToReset.Empty();
 
 	StopRecording();
 }
@@ -842,7 +817,6 @@ void FRewindDebugger::Tick(float DeltaTime)
 		
 		UE::Trace::ToggleChannel(TEXT("Object"), true);
 		UE::Trace::ToggleChannel(TEXT("ObjectProperties"), true);
-		UE::Trace::ToggleChannel(TEXT("Animation"), true);
 		UE::Trace::ToggleChannel(TEXT("Frame"), true);
 
 		// update extensions
@@ -975,251 +949,6 @@ void FRewindDebugger::Tick(float DeltaTime)
 #endif // OBJECT_TRACE_ENABLED
 								}
 							}
-							
-							// update pose on all SkeletalMeshComponents:
-							// - enumerate all skeletal mesh pose timelines
-							// - check if the corresponding mesh component still exists
-							// - apply the recorded pose for the current Frame
-							{
-								TRACE_CPUPROFILER_EVENT_SCOPE(FRewindDebugger::Tick_UpdatePoses);
-								AnimationProvider->EnumerateSkeletalMeshPoseTimelines([this, &Frame, AnimationProvider, GameplayProvider](uint64 ObjectId, const IAnimationProvider::SkeletalMeshPoseTimeline& TimelineData)
-								{
-#if OBJECT_TRACE_ENABLED
-									if(UObject* ObjectInstance = FObjectTrace::GetObjectFromId(ObjectId))
-									{
-										if(USkeletalMeshComponent* MeshComponent = Cast<USkeletalMeshComponent>(ObjectInstance))
-										{
-											AnimationProvider->ReadSkeletalMeshPoseTimeline(ObjectId, [this, &Frame, ObjectId, MeshComponent, AnimationProvider](const IAnimationProvider::SkeletalMeshPoseTimeline& TimelineData, bool bHasCurves)
-											{
-												const FSkeletalMeshPoseMessage * PoseMessage = nullptr;
-
-												// Get last pose in frame
-												TimelineData.EnumerateEvents(Frame.StartTime, Frame.EndTime,
-													[&PoseMessage](double InStartTime, double InEndTime, uint32 InDepth, const FSkeletalMeshPoseMessage& InPoseMessage)
-													{
-														PoseMessage = &InPoseMessage;
-														return TraceServices::EEventEnumerate::Continue;
-													});
-
-												// Update mesh based on pose
-												if (PoseMessage)
-												{
-													FTransform ComponentWorldTransform;
-													if (const FSkeletalMeshInfo* SkeletalMeshInfo = AnimationProvider->FindSkeletalMeshInfo(PoseMessage->MeshId))
-													{
-														AnimationProvider->GetSkeletalMeshComponentSpacePose(*PoseMessage, *SkeletalMeshInfo, ComponentWorldTransform, MeshComponent->GetEditableComponentSpaceTransforms());
-														MeshComponent->ApplyEditedComponentSpaceTransforms();
-
-														if (MeshComponentsToReset.Find(ObjectId) == nullptr)
-														{
-															FMeshComponentResetData ResetData;
-															ResetData.Component = MeshComponent;
-															ResetData.RelativeTransform = MeshComponent->GetRelativeTransform();
-															MeshComponentsToReset.Add(ObjectId, ResetData);
-														}
-
-														MeshComponent->SetWorldTransform(ComponentWorldTransform, false, nullptr, ETeleportType::TeleportPhysics);
-														MeshComponent->SetForcedLOD(PoseMessage->LodIndex + 1);
-														MeshComponent->UpdateChildTransforms(EUpdateTransformFlags::None, ETeleportType::TeleportPhysics);
-													}
-												}
-											});
-										}
-									}
-#endif // OBJECT_TRACE_ENABLED
-								});
-							}
-
-							{
-								TRACE_CPUPROFILER_EVENT_SCOPE(FRewindDebugger::Tick_AnimBlueprintsDebug);
-								// Apply Animation Blueprint Debugging Data:
-								// - enumerate over all anim graph timelines
-								// - check if their instance class still exists and is the debugging target for the Animation Blueprint Editor
-								// - if it is copy that debug data into the class debug data for the blueprint debugger
-								AnimationProvider->EnumerateAnimGraphTimelines([&Frame, AnimationProvider, GameplayProvider](uint64 ObjectId, const IAnimationProvider::AnimGraphTimeline& AnimGraphTimeline)
-								{
-#if OBJECT_TRACE_ENABLED
-									if(UObject* ObjectInstance = FObjectTrace::GetObjectFromId(ObjectId))
-									{
-										if(UAnimInstance* AnimInstance = Cast<UAnimInstance>(ObjectInstance))
-										{
-											if(UAnimBlueprintGeneratedClass* InstanceClass = Cast<UAnimBlueprintGeneratedClass>(AnimInstance->GetClass()))
-											{
-												if(UAnimBlueprint* AnimBlueprint = Cast<UAnimBlueprint>(InstanceClass->ClassGeneratedBy))
-												{
-													// for child Animation Blueprints, we actually want to debug the root blueprint (since the child doesn't contain any anim graphs)
-													if (UAnimBlueprint* RootAnimBP = UAnimBlueprint::FindRootAnimBlueprint(AnimBlueprint))
-													{
-														if (UAnimBlueprintGeneratedClass* RootInstanceClass = Cast<UAnimBlueprintGeneratedClass>(RootAnimBP->GeneratedClass))
-														{
-															AnimBlueprint = RootAnimBP;
-															InstanceClass = RootInstanceClass;
-														}
-													}
-
-													if(AnimBlueprint->IsObjectBeingDebugged(AnimInstance))
-													{
-														TRACE_CPUPROFILER_EVENT_SCOPE(FRewindDebugger::Tick_UpdateBlueprintDebug);
-														// update debug info for attached Animation Blueprint editors
-														uint64 Id = FObjectTrace::GetObjectId(AnimInstance);
-														const int32 NodeCount = InstanceClass->GetAnimNodeProperties().Num();
-								
-														FAnimBlueprintDebugData& DebugData = InstanceClass->GetAnimBlueprintDebugData();
-														{
-															TRACE_CPUPROFILER_EVENT_SCOPE(ResetNodeVisitStates);
-															DebugData.ResetNodeVisitSites();
-														}
-														
-														// Anim node values can come from all phases
-														AnimationProvider->ReadAnimNodeValuesTimeline(ObjectId, [&Frame,AnimationProvider, &DebugData](const IAnimationProvider::AnimNodeValuesTimeline& InNodeValuesTimeline)
-														{
-															TRACE_CPUPROFILER_EVENT_SCOPE(AnimGraphNodeValues);
-															InNodeValuesTimeline.EnumerateEvents(Frame.StartTime, Frame.EndTime, [AnimationProvider, &DebugData](double InStartTime, double InEndTime, uint32 InDepth, const FAnimNodeValueMessage& InMessage)
-															{
-																// don't send "Name" Node value for display in the graph
-																if (FPlatformString::Strcmp(InMessage.Key, TEXT("Name")) != 0)
-																{
-																	FText Text = AnimationProvider->FormatNodeKeyValue(InMessage);
-																	DebugData.RecordNodeValue(InMessage.NodeId, Text.ToString());
-																}
-																return TraceServices::EEventEnumerate::Continue;
-															});
-														});
-
-														DebugData.DisableAllPoseWatches();
-							
-														AnimGraphTimeline.EnumerateEvents(Frame.StartTime, Frame.EndTime, [Id, AnimationProvider, GameplayProvider, &DebugData, NodeCount](double InGraphStartTime, double InGraphEndTime, uint32 InDepth, const FAnimGraphMessage& InMessage)
-														{
-															TRACE_CPUPROFILER_EVENT_SCOPE(AnimGraphTimelineEvent);
-																	
-															// Basic verification - check node count is the same
-															// @TODO: could add some form of node hash/CRC to the class to improve this
-															if(InMessage.NodeCount == NodeCount)
-															{
-																// Check for an update phase (which contains weights)
-																if(InMessage.Phase == EAnimGraphPhase::Update)
-																{
-																	AnimationProvider->ReadAnimNodesTimeline(Id, [InGraphStartTime, InGraphEndTime, &DebugData](const IAnimationProvider::AnimNodesTimeline& InNodesTimeline)
-																	{
-																		TRACE_CPUPROFILER_EVENT_SCOPE(AnimGraphDebugNodeVisits);
-																		InNodesTimeline.EnumerateEvents(InGraphStartTime, InGraphEndTime, [&DebugData](double InStartTime, double InEndTime, uint32 InDepth, const FAnimNodeMessage& InMessage)
-																		{
-																			DebugData.RecordNodeVisit(InMessage.NodeId, InMessage.PreviousNodeId, InMessage.Weight);
-																			return TraceServices::EEventEnumerate::Continue;
-																		});
-																	});
-							
-																	AnimationProvider->ReadStateMachinesTimeline(Id, [InGraphStartTime, InGraphEndTime, &DebugData](const IAnimationProvider::StateMachinesTimeline& InStateMachinesTimeline)
-																	{
-																		TRACE_CPUPROFILER_EVENT_SCOPE(AnimGraphDebugStateMachine);
-																		InStateMachinesTimeline.EnumerateEvents(InGraphStartTime, InGraphEndTime, [&DebugData](double InStartTime, double InEndTime, uint32 InDepth, const FAnimStateMachineMessage& InMessage)
-																		{
-																			DebugData.RecordStateData(InMessage.StateMachineIndex, InMessage.StateIndex, InMessage.StateWeight, InMessage.ElapsedTime);
-																			return TraceServices::EEventEnumerate::Continue;
-																		});
-																	});
-							
-																	AnimationProvider->ReadAnimSequencePlayersTimeline(Id, [InGraphStartTime, InGraphEndTime, GameplayProvider, &DebugData](const IAnimationProvider::AnimSequencePlayersTimeline& InSequencePlayersTimeline)
-																	{
-																		TRACE_CPUPROFILER_EVENT_SCOPE(AnimGraphDebugSequencePlayers);
-																		InSequencePlayersTimeline.EnumerateEvents(InGraphStartTime, InGraphEndTime, [&DebugData](double InStartTime, double InEndTime, uint32 InDepth, const FAnimSequencePlayerMessage& InMessage)
-																		{
-																			DebugData.RecordSequencePlayer(InMessage.NodeId, InMessage.Position, InMessage.Length, InMessage.FrameCounter);
-																			return TraceServices::EEventEnumerate::Continue;
-																		});
-																	});
-							
-																	AnimationProvider->ReadAnimBlendSpacePlayersTimeline(Id, [InGraphStartTime, InGraphEndTime, GameplayProvider, &DebugData](const IAnimationProvider::BlendSpacePlayersTimeline& InBlendSpacePlayersTimeline)
-																	{
-																		TRACE_CPUPROFILER_EVENT_SCOPE(AnimGraphBlendSpaces);
-																		InBlendSpacePlayersTimeline.EnumerateEvents(InGraphStartTime, InGraphEndTime, [GameplayProvider, &DebugData](double InStartTime, double InEndTime, uint32 InDepth, const FBlendSpacePlayerMessage& InMessage)
-																		{
-																			UBlendSpace* BlendSpace = nullptr;
-																			const FObjectInfo* BlendSpaceInfo = GameplayProvider->FindObjectInfo(InMessage.BlendSpaceId);
-																			if(BlendSpaceInfo)
-																			{
-																				BlendSpace = TSoftObjectPtr<UBlendSpace>(FSoftObjectPath(BlendSpaceInfo->PathName)).LoadSynchronous();
-																			}
-							
-																			DebugData.RecordBlendSpacePlayer(InMessage.NodeId, BlendSpace, FVector(InMessage.PositionX, InMessage.PositionY, InMessage.PositionZ), FVector(InMessage.FilteredPositionX, InMessage.FilteredPositionY, InMessage.FilteredPositionZ));
-																			return TraceServices::EEventEnumerate::Continue;
-																		});
-																	});
-							
-																	AnimationProvider->ReadAnimSyncTimeline(Id, [InGraphStartTime, InGraphEndTime, AnimationProvider, &DebugData](const IAnimationProvider::AnimSyncTimeline& InAnimSyncTimeline)
-																	{
-																		TRACE_CPUPROFILER_EVENT_SCOPE(AnimGraphAnimSync);
-																		InAnimSyncTimeline.EnumerateEvents(InGraphStartTime, InGraphEndTime, [AnimationProvider, &DebugData](double InStartTime, double InEndTime, uint32 InDepth, const FAnimSyncMessage& InMessage)
-																		{
-																			const TCHAR* GroupName = AnimationProvider->GetName(InMessage.GroupNameId);
-																			if(GroupName)
-																			{
-																				DebugData.RecordNodeSync(InMessage.SourceNodeId, FName(GroupName));
-																			}
-																
-																			return TraceServices::EEventEnumerate::Continue;
-																		});
-																	});
-																}
-							
-																// Some traces come from both update and evaluate phases
-																if(InMessage.Phase == EAnimGraphPhase::Update || InMessage.Phase == EAnimGraphPhase::Evaluate)
-																{
-																	AnimationProvider->ReadAnimAttributesTimeline(Id, [InGraphStartTime, InGraphEndTime, AnimationProvider, &DebugData](const IAnimationProvider::AnimAttributeTimeline& InAnimAttributeTimeline)
-																	{
-																		TRACE_CPUPROFILER_EVENT_SCOPE(AnimGraphAttributes);
-																		InAnimAttributeTimeline.EnumerateEvents(InGraphStartTime, InGraphEndTime, [AnimationProvider, &DebugData](double InStartTime, double InEndTime, uint32 InDepth, const FAnimAttributeMessage& InMessage)
-																		{
-																			const TCHAR* AttributeName = AnimationProvider->GetName(InMessage.AttributeNameId);
-																			if(AttributeName)
-																			{
-																				DebugData.RecordNodeAttribute(InMessage.TargetNodeId, InMessage.SourceNodeId, FName(AttributeName));
-																			}
-																
-																			return TraceServices::EEventEnumerate::Continue;
-																		});
-																	});
-
-																	
-																	AnimationProvider->ReadPoseWatchTimeline(Id, [InGraphStartTime, InGraphEndTime, AnimationProvider, &DebugData](const IAnimationProvider::PoseWatchTimeline& InPoseWatchTimeline)
-																		{
-																			InPoseWatchTimeline.EnumerateEvents(InGraphStartTime, InGraphEndTime, [AnimationProvider, &DebugData](double InStartTime, double InEndTime, uint32 InDepth, const FPoseWatchMessage& InMessage)
-																				{
-#if WITH_EDITOR
-																					for (FAnimNodePoseWatch& PoseWatch : DebugData.AnimNodePoseWatch)
-																					{
-																						if (PoseWatch.NodeID == InMessage.PoseWatchId)
-																						{
-																							TArray<FBoneIndexType> RequiredBones;
-																							TArray<FTransform> BoneTransforms;
-																							AnimationProvider->GetPoseWatchData(InMessage, BoneTransforms, RequiredBones);
-
-																							PoseWatch.SetPose(RequiredBones, BoneTransforms);
-																							PoseWatch.SetWorldTransform(InMessage.WorldTransform);
-
-																							PoseWatch.PoseWatch->SetIsNodeEnabled(true);
-																							break;
-																						}
-																					}
-#endif //WITH_EDITOR
-																					return TraceServices::EEventEnumerate::Continue;
-																				});
-																		});
-
-																}
-							
-															}
-															return TraceServices::EEventEnumerate::Continue;
-														});
-													}
-												}
-											}
-										}
-									}
-#endif // OBJECT_TRACE_ENABLED
-									return TraceServices::EEventEnumerate::Continue;
-								});
-							}
 						}
 					}
 				}
diff --git a/Engine/Plugins/Animation/GameplayInsights/Source/RewindDebugger/Private/RewindDebugger.h b/Engine/Plugins/Animation/GameplayInsights/Source/RewindDebugger/Private/RewindDebugger.h
--- a/Engine/Plugins/Animation/GameplayInsights/Source/RewindDebugger/Private/RewindDebugger.h
+++ b/Engine/Plugins/Animation/GameplayInsights/Source/RewindDebugger/Private/RewindDebugger.h
@@ -189,14 +189,7 @@ private:
 	TArray<TSharedPtr<RewindDebugger::FRewindDebuggerTrack>> DebugTracks;
 	TSharedPtr<RewindDebugger::FRewindDebuggerTrack> SelectedTrack;
 
-	struct FMeshComponentResetData
-	{
-		TWeakObjectPtr<USkeletalMeshComponent> Component;
-		FTransform RelativeTransform;
-	};
-
 	TArray<uint64> TargetObjectIds;
-	TMap<uint64, FMeshComponentResetData> MeshComponentsToReset;
 
 	mutable class IUnrealInsightsModule *UnrealInsightsModule;
 	FTSTicker::FDelegateHandle TickerHandle;
diff --git a/Engine/Plugins/Animation/GameplayInsights/Source/RewindDebugger/Private/RewindDebuggerAnimation.cpp b/Engine/Plugins/Animation/GameplayInsights/Source/RewindDebugger/Private/RewindDebuggerAnimation.cpp
new file mode 100644
--- /dev/null
+++ b/Engine/Plugins/Animation/GameplayInsights/Source/RewindDebugger/Private/RewindDebuggerAnimation.cpp
@@ -0,0 +1,360 @@
+// Copyright Epic Games, Inc. All Rights Reserved.
+
+#include "RewindDebuggerAnimation.h"
+#include "Animation/AnimBlueprint.h"
+#include "Animation/AnimBlueprintGeneratedClass.h"
+#include "Editor.h"
+#include "Engine/PoseWatch.h"
+#include "IAnimationProvider.h"
+#include "IGameplayProvider.h"
+#include "Insights/IUnrealInsightsModule.h"
+#include "IRewindDebugger.h"
+#include "LevelEditor.h"
+#include "Modules/ModuleManager.h"
+#include "SLevelViewport.h"
+#include "ToolMenu.h"
+#include "ToolMenus.h"
+#include "ObjectTrace.h"
+#include "TraceServices/Model/Frames.h"
+
+#define LOCTEXT_NAMESPACE "RewindDebuggerAnimation"
+
+FRewindDebuggerAnimation::FRewindDebuggerAnimation()
+{
+}
+
+void FRewindDebuggerAnimation::RecordingStarted(IRewindDebugger*)
+{
+	UE::Trace::ToggleChannel(TEXT("Animation"), true);
+}
+
+void FRewindDebuggerAnimation::RecordingStopped(IRewindDebugger*)
+{
+	UE::Trace::ToggleChannel(TEXT("Animation"), false);
+}
+
+void FRewindDebuggerAnimation::Initialize()
+{
+	FEditorDelegates::ResumePIE.AddRaw(this, &FRewindDebuggerAnimation::OnPIEResumed);
+	FEditorDelegates::EndPIE.AddRaw(this, &FRewindDebuggerAnimation::OnPIEStopped);
+	FEditorDelegates::SingleStepPIE.AddRaw(this, &FRewindDebuggerAnimation::OnPIESingleStepped);
+}
+
+void FRewindDebuggerAnimation::Shutdown()
+{
+	FEditorDelegates::ResumePIE.RemoveAll(this);
+	FEditorDelegates::EndPIE.RemoveAll(this);
+	FEditorDelegates::SingleStepPIE.RemoveAll(this);
+}
+
+void FRewindDebuggerAnimation::OnPIEResumed(bool bSimulating)
+{
+	// restore all relative transforms of any meshes that may have been moved while scrubbing
+	for (TTuple<uint64, FMeshComponentResetData>& MeshData : MeshComponentsToReset)
+	{
+		if (USkeletalMeshComponent* MeshComponent = MeshData.Value.Component.Get())
+		{
+			MeshComponent->SetRelativeTransform(MeshData.Value.RelativeTransform, false, nullptr, ETeleportType::TeleportPhysics);
+		}
+	}
+
+	MeshComponentsToReset.Empty();
+}
+
+void FRewindDebuggerAnimation::OnPIESingleStepped(bool bSimulating)
+{
+	// restore all relative transforms of any meshes that may have been moved while scrubbing
+	for (TTuple<uint64, FMeshComponentResetData>& MeshData : MeshComponentsToReset)
+	{
+		if (USkeletalMeshComponent* MeshComponent = MeshData.Value.Component.Get())
+		{
+			MeshComponent->SetRelativeTransform(MeshData.Value.RelativeTransform, false, nullptr, ETeleportType::TeleportPhysics);
+		}
+	}
+
+	MeshComponentsToReset.Empty();
+}
+
+
+void FRewindDebuggerAnimation::OnPIEStopped(bool bSimulating)
+{
+	MeshComponentsToReset.Empty();
+}
+
+
+void FRewindDebuggerAnimation::Update(float DeltaTime, IRewindDebugger* RewindDebugger)
+{
+	check(RewindDebugger);
+	
+	if (RewindDebugger->IsPIESimulating() || RewindDebugger->GetRecordingDuration() == 0.0)
+	{
+		return;
+	}
+
+	if (const TraceServices::IAnalysisSession* Session = RewindDebugger->GetAnalysisSession())
+	{
+		TraceServices::FAnalysisSessionReadScope SessionReadScope(*Session);
+		double CurrentTraceTime = RewindDebugger->CurrentTraceTime();
+		
+		if (CurrentTraceTime != LastScrubTime)
+		{
+			const TraceServices::IFrameProvider& FrameProvider = TraceServices::ReadFrameProvider(*Session);
+			TraceServices::FFrame Frame;
+			if (FrameProvider.GetFrameFromTime(ETraceFrameType::TraceFrameType_Game, CurrentTraceTime, Frame))
+			{
+				const IAnimationProvider* AnimationProvider = Session->ReadProvider<IAnimationProvider>("AnimationProvider");
+				const IGameplayProvider* GameplayProvider = Session->ReadProvider<IGameplayProvider>("GameplayProvider");
+				
+				if (AnimationProvider && GameplayProvider)
+				{
+					// UWorld* World = GetWorldToVisualize();
+
+					// update pose on all SkeletalMeshComponents:
+					// - enumerate all skeletal mesh pose timelines
+					// - check if the corresponding mesh component still exists
+					// - apply the recorded pose for the current Frame
+					{
+						TRACE_CPUPROFILER_EVENT_SCOPE(FRewindDebugger::Tick_UpdatePoses);
+						AnimationProvider->EnumerateSkeletalMeshPoseTimelines([this, &Frame, AnimationProvider, GameplayProvider](uint64 ObjectId, const IAnimationProvider::SkeletalMeshPoseTimeline& TimelineData)
+						{
+#if OBJECT_TRACE_ENABLED
+							if(UObject* ObjectInstance = FObjectTrace::GetObjectFromId(ObjectId))
+							{
+								if(USkeletalMeshComponent* MeshComponent = Cast<USkeletalMeshComponent>(ObjectInstance))
+								{
+									AnimationProvider->ReadSkeletalMeshPoseTimeline(ObjectId, [this, &Frame, ObjectId, MeshComponent, AnimationProvider](const IAnimationProvider::SkeletalMeshPoseTimeline& TimelineData, bool bHasCurves)
+									{
+										const FSkeletalMeshPoseMessage * PoseMessage = nullptr;
+
+										// Get last pose in frame
+										TimelineData.EnumerateEvents(Frame.StartTime, Frame.EndTime,
+											[&PoseMessage](double InStartTime, double InEndTime, uint32 InDepth, const FSkeletalMeshPoseMessage& InPoseMessage)
+											{
+												PoseMessage = &InPoseMessage;
+												return TraceServices::EEventEnumerate::Continue;
+											});
+
+										// Update mesh based on pose
+										if (PoseMessage)
+										{
+											FTransform ComponentWorldTransform;
+											if (const FSkeletalMeshInfo* SkeletalMeshInfo = AnimationProvider->FindSkeletalMeshInfo(PoseMessage->MeshId))
+											{
+												AnimationProvider->GetSkeletalMeshComponentSpacePose(*PoseMessage, *SkeletalMeshInfo, ComponentWorldTransform, MeshComponent->GetEditableComponentSpaceTransforms());
+												MeshComponent->ApplyEditedComponentSpaceTransforms();
+
+												if (MeshComponentsToReset.Find(ObjectId) == nullptr)
+												{
+													FMeshComponentResetData ResetData;
+													ResetData.Component = MeshComponent;
+													ResetData.RelativeTransform = MeshComponent->GetRelativeTransform();
+													MeshComponentsToReset.Add(ObjectId, ResetData);
+												}
+
+												MeshComponent->SetWorldTransform(ComponentWorldTransform, false, nullptr, ETeleportType::TeleportPhysics);
+												MeshComponent->SetForcedLOD(PoseMessage->LodIndex + 1);
+												MeshComponent->UpdateChildTransforms(EUpdateTransformFlags::None, ETeleportType::TeleportPhysics);
+											}
+										}
+									});
+								}
+							}
+#endif // OBJECT_TRACE_ENABLED
+							});
+						}
+        
+						{
+							TRACE_CPUPROFILER_EVENT_SCOPE(FRewindDebugger::Tick_AnimBlueprintsDebug);
+							// Apply Animation Blueprint Debugging Data:
+							// - enumerate over all anim graph timelines
+							// - check if their instance class still exists and is the debugging target for the Animation Blueprint Editor
+							// - if it is copy that debug data into the class debug data for the blueprint debugger
+							AnimationProvider->EnumerateAnimGraphTimelines([&Frame, AnimationProvider, GameplayProvider](uint64 ObjectId, const IAnimationProvider::AnimGraphTimeline& AnimGraphTimeline)
+							{
+#if OBJECT_TRACE_ENABLED
+								if(UObject* ObjectInstance = FObjectTrace::GetObjectFromId(ObjectId))
+								{
+									if(UAnimInstance* AnimInstance = Cast<UAnimInstance>(ObjectInstance))
+									{
+										if(UAnimBlueprintGeneratedClass* InstanceClass = Cast<UAnimBlueprintGeneratedClass>(AnimInstance->GetClass()))
+										{
+											if(UAnimBlueprint* AnimBlueprint = Cast<UAnimBlueprint>(InstanceClass->ClassGeneratedBy))
+											{
+												// for child Animation Blueprints, we actually want to debug the root blueprint (since the child doesn't contain any anim graphs)
+												if (UAnimBlueprint* RootAnimBP = UAnimBlueprint::FindRootAnimBlueprint(AnimBlueprint))
+												{
+													if (UAnimBlueprintGeneratedClass* RootInstanceClass = Cast<UAnimBlueprintGeneratedClass>(RootAnimBP->GeneratedClass))
+													{
+														AnimBlueprint = RootAnimBP;
+														InstanceClass = RootInstanceClass;
+													}
+												}
+
+												if(AnimBlueprint->IsObjectBeingDebugged(AnimInstance))
+												{
+													TRACE_CPUPROFILER_EVENT_SCOPE(FRewindDebugger::Tick_UpdateBlueprintDebug);
+													// update debug info for attached Animation Blueprint editors
+													uint64 Id = FObjectTrace::GetObjectId(AnimInstance);
+													const int32 NodeCount = InstanceClass->GetAnimNodeProperties().Num();
+							
+													FAnimBlueprintDebugData& DebugData = InstanceClass->GetAnimBlueprintDebugData();
+													{
+														TRACE_CPUPROFILER_EVENT_SCOPE(ResetNodeVisitStates);
+														DebugData.ResetNodeVisitSites();
+													}
+													
+													// Anim node values can come from all phases
+													AnimationProvider->ReadAnimNodeValuesTimeline(ObjectId, [&Frame,AnimationProvider, &DebugData](const IAnimationProvider::AnimNodeValuesTimeline& InNodeValuesTimeline)
+													{
+														TRACE_CPUPROFILER_EVENT_SCOPE(AnimGraphNodeValues);
+														InNodeValuesTimeline.EnumerateEvents(Frame.StartTime, Frame.EndTime, [AnimationProvider, &DebugData](double InStartTime, double InEndTime, uint32 InDepth, const FAnimNodeValueMessage& InMessage)
+														{
+															// don't send "Name" Node value for display in the graph
+															if (FPlatformString::Strcmp(InMessage.Key, TEXT("Name")) != 0)
+															{
+																FText Text = AnimationProvider->FormatNodeKeyValue(InMessage);
+																DebugData.RecordNodeValue(InMessage.NodeId, Text.ToString());
+															}
+															return TraceServices::EEventEnumerate::Continue;
+														});
+													});
+
+													DebugData.DisableAllPoseWatches();
+						
+													AnimGraphTimeline.EnumerateEvents(Frame.StartTime, Frame.EndTime, [Id, AnimationProvider, GameplayProvider, &DebugData, NodeCount](double InGraphStartTime, double InGraphEndTime, uint32 InDepth, const FAnimGraphMessage& InMessage)
+													{
+														TRACE_CPUPROFILER_EVENT_SCOPE(AnimGraphTimelineEvent);
+																
+														// Basic verification - check node count is the same
+														// @TODO: could add some form of node hash/CRC to the class to improve this
+														if(InMessage.NodeCount == NodeCount)
+														{
+															// Check for an update phase (which contains weights)
+															if(InMessage.Phase == EAnimGraphPhase::Update)
+															{
+																AnimationProvider->ReadAnimNodesTimeline(Id, [InGraphStartTime, InGraphEndTime, &DebugData](const IAnimationProvider::AnimNodesTimeline& InNodesTimeline)
+																{
+																	TRACE_CPUPROFILER_EVENT_SCOPE(AnimGraphDebugNodeVisits);
+																	InNodesTimeline.EnumerateEvents(InGraphStartTime, InGraphEndTime, [&DebugData](double InStartTime, double InEndTime, uint32 InDepth, const FAnimNodeMessage& InMessage)
+																	{
+																		DebugData.RecordNodeVisit(InMessage.NodeId, InMessage.PreviousNodeId, InMessage.Weight);
+																		return TraceServices::EEventEnumerate::Continue;
+																	});
+																});
+						
+																AnimationProvider->ReadStateMachinesTimeline(Id, [InGraphStartTime, InGraphEndTime, &DebugData](const IAnimationProvider::StateMachinesTimeline& InStateMachinesTimeline)
+																{
+																	TRACE_CPUPROFILER_EVENT_SCOPE(AnimGraphDebugStateMachine);
+																	InStateMachinesTimeline.EnumerateEvents(InGraphStartTime, InGraphEndTime, [&DebugData](double InStartTime, double InEndTime, uint32 InDepth, const FAnimStateMachineMessage& InMessage)
+																	{
+																		DebugData.RecordStateData(InMessage.StateMachineIndex, InMessage.StateIndex, InMessage.StateWeight, InMessage.ElapsedTime);
+																		return TraceServices::EEventEnumerate::Continue;
+																	});
+																});
+						
+																AnimationProvider->ReadAnimSequencePlayersTimeline(Id, [InGraphStartTime, InGraphEndTime, GameplayProvider, &DebugData](const IAnimationProvider::AnimSequencePlayersTimeline& InSequencePlayersTimeline)
+																{
+																	TRACE_CPUPROFILER_EVENT_SCOPE(AnimGraphDebugSequencePlayers);
+																	InSequencePlayersTimeline.EnumerateEvents(InGraphStartTime, InGraphEndTime, [&DebugData](double InStartTime, double InEndTime, uint32 InDepth, const FAnimSequencePlayerMessage& InMessage)
+																	{
+																		DebugData.RecordSequencePlayer(InMessage.NodeId, InMessage.Position, InMessage.Length, InMessage.FrameCounter);
+																		return TraceServices::EEventEnumerate::Continue;
+																	});
+																});
+						
+																AnimationProvider->ReadAnimBlendSpacePlayersTimeline(Id, [InGraphStartTime, InGraphEndTime, GameplayProvider, &DebugData](const IAnimationProvider::BlendSpacePlayersTimeline& InBlendSpacePlayersTimeline)
+																{
+																	TRACE_CPUPROFILER_EVENT_SCOPE(AnimGraphBlendSpaces);
+																	InBlendSpacePlayersTimeline.EnumerateEvents(InGraphStartTime, InGraphEndTime, [GameplayProvider, &DebugData](double InStartTime, double InEndTime, uint32 InDepth, const FBlendSpacePlayerMessage& InMessage)
+																	{
+																		UBlendSpace* BlendSpace = nullptr;
+																		const FObjectInfo* BlendSpaceInfo = GameplayProvider->FindObjectInfo(InMessage.BlendSpaceId);
+																		if(BlendSpaceInfo)
+																		{
+																			BlendSpace = TSoftObjectPtr<UBlendSpace>(FSoftObjectPath(BlendSpaceInfo->PathName)).LoadSynchronous();
+																		}
+						
+																		DebugData.RecordBlendSpacePlayer(InMessage.NodeId, BlendSpace, FVector(InMessage.PositionX, InMessage.PositionY, InMessage.PositionZ), FVector(InMessage.FilteredPositionX, InMessage.FilteredPositionY, InMessage.FilteredPositionZ));
+																		return TraceServices::EEventEnumerate::Continue;
+																	});
+																});
+						
+																AnimationProvider->ReadAnimSyncTimeline(Id, [InGraphStartTime, InGraphEndTime, AnimationProvider, &DebugData](const IAnimationProvider::AnimSyncTimeline& InAnimSyncTimeline)
+																{
+																	TRACE_CPUPROFILER_EVENT_SCOPE(AnimGraphAnimSync);
+																	InAnimSyncTimeline.EnumerateEvents(InGraphStartTime, InGraphEndTime, [AnimationProvider, &DebugData](double InStartTime, double InEndTime, uint32 InDepth, const FAnimSyncMessage& InMessage)
+																	{
+																		const TCHAR* GroupName = AnimationProvider->GetName(InMessage.GroupNameId);
+																		if(GroupName)
+																		{
+																			DebugData.RecordNodeSync(InMessage.SourceNodeId, FName(GroupName));
+																		}
+															
+																		return TraceServices::EEventEnumerate::Continue;
+																	});
+																});
+															}
+						
+															// Some traces come from both update and evaluate phases
+															if(InMessage.Phase == EAnimGraphPhase::Update || InMessage.Phase == EAnimGraphPhase::Evaluate)
+															{
+																AnimationProvider->ReadAnimAttributesTimeline(Id, [InGraphStartTime, InGraphEndTime, AnimationProvider, &DebugData](const IAnimationProvider::AnimAttributeTimeline& InAnimAttributeTimeline)
+																{
+																	TRACE_CPUPROFILER_EVENT_SCOPE(AnimGraphAttributes);
+																	InAnimAttributeTimeline.EnumerateEvents(InGraphStartTime, InGraphEndTime, [AnimationProvider, &DebugData](double InStartTime, double InEndTime, uint32 InDepth, const FAnimAttributeMessage& InMessage)
+																	{
+																		const TCHAR* AttributeName = AnimationProvider->GetName(InMessage.AttributeNameId);
+																		if(AttributeName)
+																		{
+																			DebugData.RecordNodeAttribute(InMessage.TargetNodeId, InMessage.SourceNodeId, FName(AttributeName));
+																		}
+															
+																		return TraceServices::EEventEnumerate::Continue;
+																	});
+																});
+
+																
+																AnimationProvider->ReadPoseWatchTimeline(Id, [InGraphStartTime, InGraphEndTime, AnimationProvider, &DebugData](const IAnimationProvider::PoseWatchTimeline& InPoseWatchTimeline)
+																	{
+																		InPoseWatchTimeline.EnumerateEvents(InGraphStartTime, InGraphEndTime, [AnimationProvider, &DebugData](double InStartTime, double InEndTime, uint32 InDepth, const FPoseWatchMessage& InMessage)
+																			{
+																				for (FAnimNodePoseWatch& PoseWatch : DebugData.AnimNodePoseWatch)
+																				{
+																					if (PoseWatch.NodeID == InMessage.PoseWatchId)
+																					{
+																						TArray<FBoneIndexType> RequiredBones;
+																						TArray<FTransform> BoneTransforms;
+																						AnimationProvider->GetPoseWatchData(InMessage, BoneTransforms, RequiredBones);
+
+																						PoseWatch.SetPose(RequiredBones, BoneTransforms);
+																						PoseWatch.SetWorldTransform(InMessage.WorldTransform);
+
+																						PoseWatch.PoseWatch->SetIsNodeEnabled(true);
+																						break;
+																					}
+																				}
+																				return TraceServices::EEventEnumerate::Continue;
+																			});
+																	});
+
+															}
+						
+														}
+														return TraceServices::EEventEnumerate::Continue;
+													});
+												}
+											}
+										}
+									}
+								}
+#endif // OBJECT_TRACE_ENABLED
+								return TraceServices::EEventEnumerate::Continue;
+							});
+						}
+					}
+				}
+			}
+	}
+}
+
+#undef LOCTEXT_NAMESPACE
\ No newline at end of file
diff --git a/Engine/Plugins/Animation/GameplayInsights/Source/RewindDebugger/Private/RewindDebuggerAnimation.h b/Engine/Plugins/Animation/GameplayInsights/Source/RewindDebugger/Private/RewindDebuggerAnimation.h
new file mode 100644
--- /dev/null
+++ b/Engine/Plugins/Animation/GameplayInsights/Source/RewindDebugger/Private/RewindDebuggerAnimation.h
@@ -0,0 +1,40 @@
+// Copyright Epic Games, Inc. All Rights Reserved.
+#pragma once
+
+#include "CoreMinimal.h"
+#include "IRewindDebuggerExtension.h"
+#include "RewindDebuggerSettings.h"
+#include "UObject/WeakObjectPtr.h"
+
+// Rewind debugger extension for animation support
+//  replay of animated pose data
+//  updating animation blueprint debugger
+
+class FRewindDebuggerAnimation : public IRewindDebuggerExtension
+{
+public:
+
+
+	FRewindDebuggerAnimation();
+	virtual ~FRewindDebuggerAnimation() {};
+	void Initialize();
+	void Shutdown();
+
+	virtual void Update(float DeltaTime, IRewindDebugger* RewindDebugger) override;
+	virtual void RecordingStarted(IRewindDebugger* RewindDebugger) override;
+	virtual void RecordingStopped(IRewindDebugger* RewindDebugger) override;
+	
+	void OnPIEResumed(bool bSimulating);
+	void OnPIEStopped(bool bSimulating);
+	void OnPIESingleStepped(bool bSimulating);
+
+private:
+	struct FMeshComponentResetData
+	{
+		TWeakObjectPtr<USkeletalMeshComponent> Component;
+		FTransform RelativeTransform;
+	};
+
+	TMap<uint64, FMeshComponentResetData> MeshComponentsToReset;
+	double LastScrubTime = 0.0;
+};
diff --git a/Engine/Plugins/Animation/GameplayInsights/Source/RewindDebugger/Private/RewindDebuggerModule.cpp b/Engine/Plugins/Animation/GameplayInsights/Source/RewindDebugger/Private/RewindDebuggerModule.cpp
--- a/Engine/Plugins/Animation/GameplayInsights/Source/RewindDebugger/Private/RewindDebuggerModule.cpp
+++ b/Engine/Plugins/Animation/GameplayInsights/Source/RewindDebugger/Private/RewindDebuggerModule.cpp
@@ -238,7 +238,9 @@ void FRewindDebuggerModule::StartupModule()
 	);
 	
 	RewindDebuggerCameraExtension.Initialize();
+	RewindDebuggerAnimationExtension.Initialize();
 	IModularFeatures::Get().RegisterModularFeature(IRewindDebuggerExtension::ModularFeatureName, &RewindDebuggerCameraExtension);
+	IModularFeatures::Get().RegisterModularFeature(IRewindDebuggerExtension::ModularFeatureName, &RewindDebuggerAnimationExtension);
 	IModularFeatures::Get().RegisterModularFeature(IRewindDebuggerDoubleClickHandler::ModularFeatureName, &AnimInstanceDoubleClickHandler);
 	IModularFeatures::Get().RegisterModularFeature(IRewindDebuggerDoubleClickHandler::ModularFeatureName, &BlueprintDoubleClickHandler);
 
@@ -250,13 +252,17 @@ void FRewindDebuggerModule::StartupModule()
 
 void FRewindDebuggerModule::ShutdownModule()
 {
+	RewindDebuggerAnimationExtension.Shutdown();
+	
 	IModularFeatures::Get().UnregisterModularFeature(IRewindDebuggerExtension::ModularFeatureName, &RewindDebuggerCameraExtension);
+	IModularFeatures::Get().UnregisterModularFeature(IRewindDebuggerExtension::ModularFeatureName, &RewindDebuggerAnimationExtension);
 	IModularFeatures::Get().UnregisterModularFeature(IRewindDebuggerDoubleClickHandler::ModularFeatureName, &AnimInstanceDoubleClickHandler);
 	IModularFeatures::Get().UnregisterModularFeature(IRewindDebuggerDoubleClickHandler::ModularFeatureName, &BlueprintDoubleClickHandler);
 
 	FRewindDebuggerCommands::Unregister();
 	FRewindDebuggerStyle::Shutdown();
 	FRewindDebugger::Shutdown();
+
 }
 
 IMPLEMENT_MODULE(FRewindDebuggerModule, RewindDebugger);
diff --git a/Engine/Plugins/Animation/GameplayInsights/Source/RewindDebugger/Private/RewindDebuggerModule.h b/Engine/Plugins/Animation/GameplayInsights/Source/RewindDebugger/Private/RewindDebuggerModule.h
--- a/Engine/Plugins/Animation/GameplayInsights/Source/RewindDebugger/Private/RewindDebuggerModule.h
+++ b/Engine/Plugins/Animation/GameplayInsights/Source/RewindDebugger/Private/RewindDebuggerModule.h
@@ -6,6 +6,7 @@
 #include "Modules/ModuleInterface.h"
 #include "Framework/Docking/TabManager.h"
 #include "RewindDebuggerCamera.h"
+#include "RewindDebuggerAnimation.h"
 
 class SRewindDebugger;
 class SRewindDebuggerDetails;
@@ -28,4 +29,5 @@ private:
 	TSharedPtr<SRewindDebuggerDetails> RewindDebuggerDetailsWidget;
 
 	FRewindDebuggerCamera RewindDebuggerCameraExtension;
+	FRewindDebuggerAnimation RewindDebuggerAnimationExtension;
 };
