Rewind debugger playback rate setting

diff --git a/Engine/Plugins/Animation/GameplayInsights/Source/RewindDebugger/Private/RewindDebugger.cpp b/Engine/Plugins/Animation/GameplayInsights/Source/RewindDebugger/Private/RewindDebugger.cpp
--- a/Engine/Plugins/Animation/GameplayInsights/Source/RewindDebugger/Private/RewindDebugger.cpp
+++ b/Engine/Plugins/Animation/GameplayInsights/Source/RewindDebugger/Private/RewindDebugger.cpp
@@ -35,6 +35,7 @@
 #include "UnrealEdGlobals.h"
 #include "Editor/UnrealEdEngine.h"
 #include "Kismet2/DebuggerCommands.h"
+#include "Widgets/Input/SNumericEntryBox.h"
 
 #define LOCTEXT_NAMESPACE "RewindDebugger"
 
@@ -66,7 +67,6 @@ FRewindDebugger::FRewindDebugger()  :
 	bPIEStarted(false),
 	bPIESimulating(false),
 	bRecording(false),
-	PlaybackRate(1),
 	PreviousTraceTime(-1),
 	CurrentScrubTime(0),
 	CurrentViewRange(0,0),
@@ -883,6 +883,7 @@ void FRewindDebugger::Tick(float DeltaTime)
 				{
 					if (ControlState == EControlState::Play || ControlState == EControlState::PlayReverse)
 					{
+						float PlaybackRate = URewindDebuggerSettings::Get().PlaybackRate;
 						TRACE_CPUPROFILER_EVENT_SCOPE(FRewindDebugger::Tick_UpdatePlayback);
 						float Rate = PlaybackRate * (ControlState == EControlState::Play ? 1 : -1);
 						SetCurrentScrubTime(FMath::Clamp(CurrentScrubTime + Rate * DeltaTime, 0.0f, RecordingDuration.Get()));
@@ -1070,6 +1071,127 @@ void FRewindDebugger::RegisterToolBar()
 				FText::Format(LOCTEXT("PlayButtonTooltip", "{0} ({1})"), Commands.Play->GetDescription(), Commands.PauseOrPlay->GetInputText()),
 				FSlateIcon("RewindDebuggerStyle", "RewindDebugger.Play.small")));
 
+	Section.AddEntry(
+    		FToolMenuEntry::InitComboButton(
+    			"PlaybackRate",
+    			FToolUIActionChoice(),
+    			FNewToolMenuChoice(
+    				FNewToolMenuDelegate::CreateLambda([](UToolMenu* InNewToolMenu)
+    				{
+    					FToolMenuSection& Section = InNewToolMenu->AddSection("PlaybackSpeed", LOCTEXT("Playback Speed", "Playback Speed"));
+    					
+						Section.AddEntry(
+							FToolMenuEntry::InitMenuEntry(
+								"001",LOCTEXT("0.1","0.1"), LOCTEXT("Set playback speed to 0.1", "Set playback speed to 0.1"), FSlateIcon(),
+								FUIAction(
+								FExecuteAction::CreateLambda([]()
+									{ 
+										URewindDebuggerSettings::Get().PlaybackRate = 0.1;
+									}),
+									FCanExecuteAction(),
+									FIsActionChecked::CreateLambda([]
+									{
+										return FMath::IsNearlyEqual(URewindDebuggerSettings::Get().PlaybackRate, 0.1);
+									})
+									)
+									, EUserInterfaceActionType::RadioButton
+								)
+    					);
+						Section.AddEntry(
+							FToolMenuEntry::InitMenuEntry(
+								"025",LOCTEXT("0.25","0.25"), LOCTEXT("Set playback speed to 0.25", "Set playback speed to 0.25"), FSlateIcon(),
+								FUIAction(
+								FExecuteAction::CreateLambda([]()
+									{ 
+										URewindDebuggerSettings::Get().PlaybackRate = 0.25;
+									}),
+									FCanExecuteAction(),
+									FIsActionChecked::CreateLambda([]
+									{
+										return FMath::IsNearlyEqual(URewindDebuggerSettings::Get().PlaybackRate, 0.25);
+									})
+									)
+									, EUserInterfaceActionType::RadioButton
+							)
+    					);
+    					Section.AddEntry(
+							FToolMenuEntry::InitMenuEntry(
+								"05",LOCTEXT("0.5","0.5"), LOCTEXT("Set playback speed to 0.5", "Set playback speed to 0.5"), FSlateIcon(),
+								FUIAction(
+								FExecuteAction::CreateLambda([]()
+									{ 
+										URewindDebuggerSettings::Get().PlaybackRate = 0.5;
+									}),
+									FCanExecuteAction(),
+									FIsActionChecked::CreateLambda([]
+									{
+										return FMath::IsNearlyEqual(URewindDebuggerSettings::Get().PlaybackRate, 0.5);
+									})
+									)
+									, EUserInterfaceActionType::RadioButton
+							)
+						);
+
+						Section.AddEntry(
+							FToolMenuEntry::InitMenuEntry(
+								"1",LOCTEXT("1","1"), LOCTEXT("Set playback speed to 1", "Set playback speed to 1"), FSlateIcon(),
+								FUIAction(
+								FExecuteAction::CreateLambda([]()
+									{ 
+										URewindDebuggerSettings::Get().PlaybackRate = 1;
+									}),
+									FCanExecuteAction(),
+									FIsActionChecked::CreateLambda([]
+									{
+										return FMath::IsNearlyEqual(URewindDebuggerSettings::Get().PlaybackRate, 1);
+									})
+									)
+									, EUserInterfaceActionType::RadioButton
+							)
+						);
+    					
+    					Section.AddEntry(
+							FToolMenuEntry::InitMenuEntry(
+								"2",LOCTEXT("2","2"), LOCTEXT("Set playback speed to 2", "Set playback speed to 2"), FSlateIcon(),
+								FUIAction(
+								FExecuteAction::CreateLambda([]()
+									{ 
+										URewindDebuggerSettings::Get().PlaybackRate = 2;
+									}),
+									FCanExecuteAction(),
+									FIsActionChecked::CreateLambda([]
+									{
+										return FMath::IsNearlyEqual(URewindDebuggerSettings::Get().PlaybackRate, 2);
+									})
+									)
+									, EUserInterfaceActionType::RadioButton
+							)
+						);
+
+						Section.AddEntry(
+							FToolMenuEntry::InitWidget(
+								"EditInSequencerMenu", 
+								SNew(SNumericEntryBox<float>)
+									.Value_Lambda([]()
+									{
+										return URewindDebuggerSettings::Get().PlaybackRate;
+									})
+									.OnValueChanged_Lambda([](float Value)
+									{
+										URewindDebuggerSettings::Get().PlaybackRate = Value;
+									}),
+								FText::GetEmpty(),
+								true, false, true
+							)
+						);
+    				})
+    			),
+				LOCTEXT("Blank",""),
+    			LOCTEXT("PlaybackRate_Tooltip", "Playback Options"),
+				FSlateIcon(FAppStyle::GetAppStyleSetName(), "Sequencer.PlaybackOptions")
+    		)
+    	);
+
 	Section.AddEntry(FToolMenuEntry::InitToolBarButton(
 				Commands.NextFrame,
 				LOCTEXT("Blank",""),
diff --git a/Engine/Plugins/Animation/GameplayInsights/Source/RewindDebugger/Private/RewindDebugger.h b/Engine/Plugins/Animation/GameplayInsights/Source/RewindDebugger/Private/RewindDebugger.h
--- a/Engine/Plugins/Animation/GameplayInsights/Source/RewindDebugger/Private/RewindDebugger.h
+++ b/Engine/Plugins/Animation/GameplayInsights/Source/RewindDebugger/Private/RewindDebugger.h
@@ -164,7 +164,6 @@ private:
 	
 	bool bRecording;
 
-	float PlaybackRate;
 	double PreviousTraceTime;
 	double CurrentScrubTime;
 	TRange<double> CurrentViewRange;
diff --git a/Engine/Plugins/Animation/GameplayInsights/Source/RewindDebugger/Private/RewindDebuggerSettings.h b/Engine/Plugins/Animation/GameplayInsights/Source/RewindDebugger/Private/RewindDebuggerSettings.h
--- a/Engine/Plugins/Animation/GameplayInsights/Source/RewindDebugger/Private/RewindDebuggerSettings.h
+++ b/Engine/Plugins/Animation/GameplayInsights/Source/RewindDebugger/Private/RewindDebuggerSettings.h
@@ -45,6 +45,10 @@ public:
 	/** If enabled, start recording information at the start of PIE */
 	UPROPERTY(EditAnywhere, Config, Category = Other)
 	bool bShouldAutoRecordOnPIE;
+	
+	/** Playback speed multiplier */
+	UPROPERTY(EditAnywhere, Config, Category = Other)
+	float PlaybackRate = 1.0;
 
 	/** If enabled, show empty tracks on Rewind Debugger Timeline*/
 	UPROPERTY(EditAnywhere, Config, Category = Filters)
@@ -54,7 +58,8 @@ public:
 	UPROPERTY(EditAnywhere, Config, Category = Filters)
 	TArray<FName> HiddenTrackTypes;
 
-	/** The track types listed here will be hidden from the track tree view */
+	
+	//** Currently selected target actor's name */
 	UPROPERTY(Config)
 	FString DebugTargetActor;
 	
