Fix crash in Anim Insights pose scrubbing

diff --git a/Engine/Plugins/Animation/GameplayInsights/Source/GameplayInsightsEditor/Private/InsightsSkeletalMeshComponent.cpp b/Engine/Plugins/Animation/GameplayInsights/Source/GameplayInsightsEditor/Private/InsightsSkeletalMeshComponent.cpp
--- a/Engine/Plugins/Animation/GameplayInsights/Source/GameplayInsightsEditor/Private/InsightsSkeletalMeshComponent.cpp
+++ b/Engine/Plugins/Animation/GameplayInsights/Source/GameplayInsightsEditor/Private/InsightsSkeletalMeshComponent.cpp
@@ -48,4 +48,6 @@ void UInsightsSkeletalMeshComponent::InitAnim(bool bForceReInit)
 
 		MorphTargetWeights.SetNumZeroed(GetSkeletalMeshAsset()->GetMorphTargets().Num());
 	}
+
+	USkeletalMeshComponent::InitAnim(bForceReInit);
 }
