Rewind debugger trace save/load support

diff --git a/Engine/Plugins/Animation/GameplayInsights/Source/RewindDebugger/Private/RewindDebugger.cpp b/Engine/Plugins/Animation/GameplayInsights/Source/RewindDebugger/Private/RewindDebugger.cpp
--- a/Engine/Plugins/Animation/GameplayInsights/Source/RewindDebugger/Private/RewindDebugger.cpp
+++ b/Engine/Plugins/Animation/GameplayInsights/Source/RewindDebugger/Private/RewindDebugger.cpp
@@ -2,6 +2,7 @@
 
 #include "RewindDebugger.h"
 
+#include "DesktopPlatformModule.h"
 #include "Animation/AnimBlueprintGeneratedClass.h"
 #include "Animation/AnimTrace.h"
 #include "Components/SkeletalMeshComponent.h"
@@ -10,7 +11,9 @@
 #include "EngineUtils.h"
 #include "GameFramework/Controller.h"
 #include "GameFramework/Pawn.h"
+#include "HAL/PlatformFileManager.h"
 #include "IAnimationProvider.h"
+#include "IDesktopPlatform.h"
 #include "IGameplayProvider.h"
 #include "IRewindDebuggerDoubleClickHandler.h"
 #include "IRewindDebuggerExtension.h"
@@ -26,6 +29,7 @@
 #include "RewindDebuggerSettings.h"
 #include "SLevelViewport.h"
 #include "ToolMenus.h"
+#include "Trace/StoreClient.h"
 #include "TraceServices/Model/Frames.h"
 #include "UObject/UObjectIterator.h"
 #include "Widgets/Docking/SDockTab.h"
@@ -35,6 +39,8 @@
 #include "UnrealEdGlobals.h"
 #include "Editor/UnrealEdEngine.h"
 #include "Kismet2/DebuggerCommands.h"
+#include "TraceServices/AnalysisService.h"
+#include "TraceServices/ITraceServicesModule.h"
 #include "Widgets/Input/SNumericEntryBox.h"
 
 #define LOCTEXT_NAMESPACE "RewindDebugger"
@@ -62,18 +68,7 @@ static void TraceSubobjects(UObject* OuterObject)
 	}
 }
 
-FRewindDebugger::FRewindDebugger()  :
-	ControlState(FRewindDebugger::EControlState::Pause),
-	bPIEStarted(false),
-	bPIESimulating(false),
-	bRecording(false),
-	PreviousTraceTime(-1),
-	CurrentScrubTime(0),
-	CurrentViewRange(0,0),
-	CurrentTraceRange(0,0),
-	RecordingIndex(0),
-	bTargetActorPositionValid(false),
-	bIsDetailsPanelOpen(true)
+FRewindDebugger::FRewindDebugger()
 {
 	RewindDebugger::FRewindDebuggerTrackCreators::EnumerateCreators([this](const RewindDebugger::IRewindDebuggerTrackCreator* Creator)
     {
@@ -292,11 +287,14 @@ void FRewindDebugger::GetTargetObjectIds(TArray<uint64>& OutTargetObjectIds) con
 
 	// make sure all the SubObjects of the target actor have been traced
 #if OBJECT_TRACE_ENABLED
-	for (uint64 OutTargetObjectId : TargetObjectIds)
+	if (IsRecording())
 	{
-		if (UObject* TargetObject = FObjectTrace::GetObjectFromId(OutTargetObjectId))
+		for (uint64 OutTargetObjectId : TargetObjectIds)
 		{
-			TraceSubobjects(TargetObject);
+			if (UObject* TargetObject = FObjectTrace::GetObjectFromId(OutTargetObjectId))
+			{
+				TraceSubobjects(TargetObject);
+			}
 		}
 	}
 #endif
@@ -440,6 +438,142 @@ void FRewindDebugger::StartRecording()
 	UnrealInsightsModule->StartAnalysisForLastLiveSession(5.0);
 
 	TargetObjectIds.Empty(2);
+
+	bTraceFileLoaded = false;
+}
+
+bool FRewindDebugger::CanOpenTrace() const
+{
+	return !bPIEStarted;
+}
+
+
+void FRewindDebugger::OpenTrace(const FString& FilePath)
+{ 
+	ClearTrace();
+
+	IUnrealInsightsModule& TraceInsightsModule = FModuleManager::LoadModuleChecked<IUnrealInsightsModule>("TraceInsights");
+	TraceInsightsModule.StartAnalysisForTraceFile(*FilePath);
+	bTraceFileLoaded = GetAnalysisSession() != nullptr;
+
+	// todo: optionally open the map the trace file was recorded in
+}
+
+void FRewindDebugger::OpenTrace()
+{
+	FString FolderPath = "";
+	
+	TArray<FString> OutOpenFilenames;
+	if (IDesktopPlatform* DesktopPlatform = FDesktopPlatformModule::Get())
+	{
+		FString ExtensionStr;
+		ExtensionStr += TEXT("Unreal Trace|*.utrace|");
+	
+		DesktopPlatform->OpenFileDialog(
+			FSlateApplication::Get().FindBestParentWindowHandleForDialogs(nullptr),
+			LOCTEXT("OpenDialogTitle", "Open Rewind Debugger Recording").ToString(),
+			FolderPath,
+			TEXT(""),
+			*ExtensionStr,
+			EFileDialogFlags::None,
+			OutOpenFilenames
+		);
+	}
+
+	if (OutOpenFilenames.Num() > 0)
+	{
+		if (OutOpenFilenames[0].EndsWith(TEXT("utrace")))
+		{
+			OpenTrace(OutOpenFilenames[0]);
+		}
+	}
+}
+
+bool FRewindDebugger::CanClearTrace() const
+{
+	return GetAnalysisSession() != nullptr;
+}
+
+void FRewindDebugger::ClearTrace()
+{
+	StopRecording();
+	RecordingDuration.Set(0);
+	
+	bTraceFileLoaded = false;
+	TargetObjectIds.Empty();
+	CurrentViewRange.SetLowerBoundValue(0);
+	CurrentViewRange.SetUpperBoundValue(0);
+	CurrentTraceRange.SetLowerBoundValue(0);
+	CurrentTraceRange.SetUpperBoundValue(0);
+	RecordingDuration.Set(0.0);
+	SetCurrentScrubTime(0.0);
+	
+	// update extensions
+	IterateExtensions([this](IRewindDebuggerExtension* Extension)
+		{
+			Extension->Clear(this);
+		}
+	);
+	
+	IUnrealInsightsModule& TraceInsightsModule = FModuleManager::LoadModuleChecked<IUnrealInsightsModule>("TraceInsights");
+	// only way I can find to clear the session is trying to load a name that doesn't exist.
+	TraceInsightsModule.StartAnalysisForTraceFile(TEXT("0"));
+
+	RefreshDebugTracks();
+}
+
+bool FRewindDebugger::CanSaveTrace() const
+{
+	const TraceServices::IAnalysisSession* Session = GetAnalysisSession();
+	return Session != nullptr && Session->IsAnalysisComplete();
+}
+
+void FRewindDebugger::SaveTrace(FString FileName)
+{
+
+	if (const TraceServices::IAnalysisSession* Session = GetAnalysisSession())
+	{
+		if (Session->IsAnalysisComplete())
+		{
+			FString SourceFileName = Session->GetName();
+
+			FPlatformFileManager& FileManager = FPlatformFileManager::Get();
+			IPlatformFile& PlatformFile = FileManager.GetPlatformFile();
+
+			PlatformFile.CopyFile(*FileName, *SourceFileName);
+		}
+		
+	}
+}
+
+void FRewindDebugger::SaveTrace()
+{
+	FString FolderPath = "";
+	
+	TArray<FString> OutFilenames;
+	if (IDesktopPlatform* DesktopPlatform = FDesktopPlatformModule::Get())
+	{
+		FString ExtensionStr;
+		ExtensionStr += TEXT("Rewind Debugger Recording |*.utrace|");
+	
+		DesktopPlatform->SaveFileDialog(
+			FSlateApplication::Get().FindBestParentWindowHandleForDialogs(nullptr),
+			LOCTEXT("SaveDialogTitle", "Save Rewind Debugger Recording").ToString(),
+			FolderPath,
+			TEXT(""),
+			*ExtensionStr,
+			EFileDialogFlags::None,
+			OutFilenames
+		);
+	}
+
+	if (OutFilenames.Num() > 0)
+	{
+		if (OutFilenames[0].EndsWith(TEXT(".utrace")))
+		{
+			SaveTrace(OutFilenames[0]);
+		}
+	}
 }
 
 bool FRewindDebugger::ShouldAutoRecordOnPIE() const
@@ -855,8 +989,6 @@ void FRewindDebugger::Tick(float DeltaTime)
 	
 	if (const TraceServices::IAnalysisSession* Session = GetAnalysisSession())
 	{
-		RefreshDebugTracks();
-		
 		const IAnimationProvider* AnimationProvider = Session->ReadProvider<IAnimationProvider>("AnimationProvider");
 		const IGameplayProvider* GameplayProvider = Session->ReadProvider<IGameplayProvider>("GameplayProvider");
 
@@ -864,7 +996,16 @@ void FRewindDebugger::Tick(float DeltaTime)
 		{
 			TraceServices::FAnalysisSessionReadScope SessionReadScope(*Session);
 
-			RecordingDuration.Set(GameplayProvider->GetRecordingDuration());
+			double RecordingDurationValue = GameplayProvider->GetRecordingDuration();
+			if (IsTraceFileLoaded() && RecordingDurationValue > RecordingDuration.Get())
+			{
+				// while trace file is loading up, recording duration changes - autoscroll
+				SetCurrentScrubTime(RecordingDurationValue);
+				TrackCursorDelegate.ExecuteIfBound(false);	
+			}
+			RecordingDuration.Set(RecordingDurationValue);
+			
+			RefreshDebugTracks();
 			
 			UWorld* World = GetWorldToVisualize();
 
@@ -873,7 +1014,7 @@ void FRewindDebugger::Tick(float DeltaTime)
 				if (bRecording)
 				{
 					TRACE_CPUPROFILER_EVENT_SCOPE(FRewindDebugger::Tick_UpdateSimulating);
-					SetCurrentScrubTime(RecordingDuration.Get());
+					SetCurrentScrubTime(RecordingDurationValue);
 					TrackCursorDelegate.ExecuteIfBound(false);
 				}
 			}
@@ -1218,6 +1359,25 @@ void FRewindDebugger::RegisterToolBar()
 				FSlateIcon("RewindDebuggerStyle", "RewindDebugger.StopRecording.small")));
 
 	Section.AddSeparator(NAME_None);
+
+	Section.AddEntry(FToolMenuEntry::InitToolBarButton(
+				Commands.OpenTrace,
+				LOCTEXT("Blank",""),
+				TAttribute<FText>(),
+				 FSlateIcon(FAppStyle::GetAppStyleSetName(), "Icons.FolderOpen")));
+	
+	Section.AddEntry(FToolMenuEntry::InitToolBarButton(
+    			Commands.SaveTrace,
+    			LOCTEXT("Blank",""),
+    			TAttribute<FText>(),
+    			 FSlateIcon(FAppStyle::GetAppStyleSetName(), "Icons.Save")));
+				 
+	Section.AddEntry(FToolMenuEntry::InitToolBarButton(
+				Commands.ClearTrace,
+				LOCTEXT("Blank",""),
+				TAttribute<FText>(),
+				 FSlateIcon(FAppStyle::GetAppStyleSetName(), "Icons.Delete")));
+
 	
 	Section.AddEntry(FToolMenuEntry::InitToolBarButton(
 				Commands.AutoEject,
@@ -1231,6 +1391,7 @@ void FRewindDebugger::RegisterToolBar()
 				LOCTEXT("Blank",""),
 				TAttribute<FText>(),
 				FSlateIcon("RewindDebuggerStyle", "RewindDebugger.AutoRecord")));
+	
 
 	
 	Menu->SetStyleSet(&FAppStyle::Get());
diff --git a/Engine/Plugins/Animation/GameplayInsights/Source/RewindDebugger/Private/RewindDebugger.h b/Engine/Plugins/Animation/GameplayInsights/Source/RewindDebugger/Private/RewindDebugger.h
--- a/Engine/Plugins/Animation/GameplayInsights/Source/RewindDebugger/Private/RewindDebugger.h
+++ b/Engine/Plugins/Animation/GameplayInsights/Source/RewindDebugger/Private/RewindDebugger.h
@@ -43,6 +43,7 @@ public:
 	virtual UWorld* GetWorldToVisualize() const override;
 	virtual bool IsRecording() const override { return bRecording; }
 	virtual bool IsPIESimulating() const override { return bPIESimulating; }
+	virtual bool IsTraceFileLoaded() const override { return bTraceFileLoaded; }
 	virtual double GetRecordingDuration() const override { return RecordingDuration.Get(); }
 	virtual TSharedPtr<FDebugObjectInfo> GetSelectedComponent() const override;
 	virtual TSharedPtr<RewindDebugger::FRewindDebuggerTrack> GetSelectedTrack() const override;
@@ -65,9 +66,19 @@ public:
 
 	// Start a new Recording:  Start tracing Object + Animation data, increment the current recording index, and reset the recording elapsed time to 0
 	void StartRecording();
-
 	bool CanStartRecording() const { return !IsRecording() && bPIESimulating; }
+	
+	bool CanOpenTrace() const;
+   	void OpenTrace(const FString& FilePath);
+   	void OpenTrace();
 
+	bool CanClearTrace() const;
+	void ClearTrace();
+	
+	bool CanSaveTrace() const;
+	void SaveTrace(FString FileName);
+	void SaveTrace();
+        	
 	bool ShouldAutoRecordOnPIE() const;
 	void SetShouldAutoRecordOnPIE(bool value);
 	
@@ -122,7 +133,7 @@ public:
 	void OnTrackCursor(const FOnTrackCursor& TrackCursorCallback);
 
 	TBindableProperty<double>* GetTraceTimeProperty() { return &TraceTime; }
-	TBindableProperty<float>* GetRecordingDurationProperty() { return &RecordingDuration; }
+	TBindableProperty<double>* GetRecordingDurationProperty() { return &RecordingDuration; }
 	TBindableProperty<FString, BindingType_Out>* GetDebugTargetActorProperty() { return &DebugTargetActor; }
 
 	virtual void OpenDetailsPanel() override;
@@ -143,7 +154,7 @@ private:
 	void SetCurrentScrubTime(double Time);
 
 	TBindableProperty<double> TraceTime;
-	TBindableProperty<float> RecordingDuration;
+	TBindableProperty<double> RecordingDuration;
 	TBindableProperty<FString, BindingType_Out> DebugTargetActor;
 
 	enum class EControlState : int8
@@ -153,22 +164,23 @@ private:
 		Pause
 	};
 
-	EControlState ControlState;
+	EControlState ControlState = EControlState::Pause;
 
 	FOnComponentListChanged ComponentListChangedDelegate;
 	FOnTrackCursor TrackCursorDelegate;
 
 	bool bTraceJustConnected = false;
-	bool bPIEStarted;
-	bool bPIESimulating;
+	bool bPIEStarted = false;
+	bool bPIESimulating = false;
+	bool bTraceFileLoaded = false;
 	
-	bool bRecording;
+	bool bRecording = false;
 
-	double PreviousTraceTime;
-	double CurrentScrubTime;
-	TRange<double> CurrentViewRange;
-	TRange<double> CurrentTraceRange;
-	uint16 RecordingIndex;
+	double PreviousTraceTime = -1;
+	double CurrentScrubTime = 0;
+	TRange<double> CurrentViewRange {0, 0};
+	TRange<double> CurrentTraceRange {0, 0};
+	uint16 RecordingIndex = 0;
 
 	struct FScrubTimeInformation
 	{
@@ -193,10 +205,10 @@ private:
 	mutable class IUnrealInsightsModule *UnrealInsightsModule;
 	FTSTicker::FDelegateHandle TickerHandle;
 
-	bool bTargetActorPositionValid;
+	bool bTargetActorPositionValid = false;
 	FVector TargetActorPosition;
 
 	TArray<RewindDebugger::FRewindDebuggerTrackType> TrackTypes;
 
-	bool bIsDetailsPanelOpen;
+	bool bIsDetailsPanelOpen = true;
 };
diff --git a/Engine/Plugins/Animation/GameplayInsights/Source/RewindDebugger/Private/RewindDebuggerAnimation.cpp b/Engine/Plugins/Animation/GameplayInsights/Source/RewindDebugger/Private/RewindDebuggerAnimation.cpp
--- a/Engine/Plugins/Animation/GameplayInsights/Source/RewindDebugger/Private/RewindDebuggerAnimation.cpp
+++ b/Engine/Plugins/Animation/GameplayInsights/Source/RewindDebugger/Private/RewindDebuggerAnimation.cpp
@@ -16,6 +16,7 @@
 #include "ToolMenus.h"
 #include "ObjectTrace.h"
 #include "TraceServices/Model/Frames.h"
+#include "Engine/SkeletalMesh.h"
 
 #define LOCTEXT_NAMESPACE "RewindDebuggerAnimation"
 
@@ -25,6 +26,7 @@ FRewindDebuggerAnimation::FRewindDebuggerAnimation()
 
 void FRewindDebuggerAnimation::RecordingStarted(IRewindDebugger*)
 {
+	ClearSpawnedComponents();
 	UE::Trace::ToggleChannel(TEXT("Animation"), true);
 }
 
@@ -33,6 +35,31 @@ void FRewindDebuggerAnimation::RecordingStopped(IRewindDebugger*)
 	UE::Trace::ToggleChannel(TEXT("Animation"), false);
 }
 
+void FRewindDebuggerAnimation::ClearSpawnedComponents()
+{
+	for(auto& MeshComponentInfo : SpawnedMeshComponents)
+	{
+		if(MeshComponentInfo.Value.Component)
+		{
+			MeshComponentInfo.Value.Component->UnregisterComponent();
+			MeshComponentInfo.Value.Component->MarkAsGarbage();
+			MeshComponentInfo.Value.Component = nullptr;
+		}
+
+		if(MeshComponentInfo.Value.Actor)
+		{
+			MeshComponentInfo.Value.Actor->Destroy();
+			MeshComponentInfo.Value.Actor = nullptr;
+		}
+	}
+	SpawnedMeshComponents.Empty();
+}
+
+void FRewindDebuggerAnimation::Clear(IRewindDebugger* RewindDebugger)
+{
+	ClearSpawnedComponents();
+}
+
 void FRewindDebuggerAnimation::Initialize()
 {
 	FEditorDelegates::ResumePIE.AddRaw(this, &FRewindDebuggerAnimation::OnPIEResumed);
@@ -82,6 +109,58 @@ void FRewindDebuggerAnimation::OnPIEStopped(bool bSimulating)
 }
 
 
+void FRewindDebuggerAnimation::ApplyPoseToMesh(const IAnimationProvider* AnimationProvider, const IGameplayProvider* GameplayProvider, const TraceServices::FFrame& Frame,
+	const IAnimationProvider::SkeletalMeshPoseTimeline& TimelineData, USkeletalMeshComponent* MeshComponent, uint64 ObjectId, bool bQueueForReset, bool bApplyMesh)
+{
+	const FSkeletalMeshPoseMessage * PoseMessage = nullptr;
+	
+	// Get last pose in frame
+	TimelineData.EnumerateEvents(Frame.StartTime, Frame.EndTime,
+		[&PoseMessage](double InStartTime, double InEndTime, uint32 InDepth, const FSkeletalMeshPoseMessage& InPoseMessage)
+		{
+			PoseMessage = &InPoseMessage;
+			return TraceServices::EEventEnumerate::Continue;
+		});
+
+	// Update mesh based on pose
+	if (PoseMessage)
+	{
+		FTransform ComponentWorldTransform;
+		if (const FSkeletalMeshInfo* SkeletalMeshInfo = AnimationProvider->FindSkeletalMeshInfo(PoseMessage->MeshId))
+		{
+			AnimationProvider->GetSkeletalMeshComponentSpacePose(*PoseMessage, *SkeletalMeshInfo, ComponentWorldTransform, MeshComponent->GetEditableComponentSpaceTransforms());
+			MeshComponent->ApplyEditedComponentSpaceTransforms();
+
+			if (bApplyMesh)
+			{
+				const FObjectInfo* SkeletalMeshObjectInfo = GameplayProvider->FindObjectInfo(PoseMessage->MeshId);
+				USkeletalMesh* SkeletalMesh = TSoftObjectPtr<USkeletalMesh>(FSoftObjectPath(SkeletalMeshObjectInfo->PathName)).LoadSynchronous();
+										
+				if(SkeletalMesh)
+				{
+					MeshComponent->SetSkeletalMesh(SkeletalMesh);
+				}
+			}
+			
+			if (bQueueForReset)
+			{
+				if (MeshComponentsToReset.Find(ObjectId) == nullptr)
+				{
+					FMeshComponentResetData ResetData;
+					ResetData.Component = MeshComponent;
+					ResetData.RelativeTransform = MeshComponent->GetRelativeTransform();
+					MeshComponentsToReset.Add(ObjectId, ResetData);
+				}
+			}
+
+			MeshComponent->SetWorldTransform(ComponentWorldTransform, false, nullptr, ETeleportType::TeleportPhysics);
+			MeshComponent->SetForcedLOD(PoseMessage->LodIndex + 1);
+			MeshComponent->UpdateChildTransforms(EUpdateTransformFlags::None, ETeleportType::TeleportPhysics);
+		}
+	}
+}
+
+
 void FRewindDebuggerAnimation::Update(float DeltaTime, IRewindDebugger* RewindDebugger)
 {
 	check(RewindDebugger);
@@ -115,52 +194,65 @@ void FRewindDebuggerAnimation::Update(float DeltaTime, IRewindDebugger* RewindDe
 					// - apply the recorded pose for the current Frame
 					{
 						TRACE_CPUPROFILER_EVENT_SCOPE(FRewindDebugger::Tick_UpdatePoses);
-						AnimationProvider->EnumerateSkeletalMeshPoseTimelines([this, &Frame, AnimationProvider, GameplayProvider](uint64 ObjectId, const IAnimationProvider::SkeletalMeshPoseTimeline& TimelineData)
+						AnimationProvider->EnumerateSkeletalMeshPoseTimelines([this, RewindDebugger, &Frame, AnimationProvider, GameplayProvider](uint64 ObjectId, const IAnimationProvider::SkeletalMeshPoseTimeline& TimelineData)
 						{
+							USkeletalMeshComponent* MeshComponent = nullptr;
+							bool bQueueForReset = false;
+							bool bLoadMesh = false;
+
 #if OBJECT_TRACE_ENABLED
-							if(UObject* ObjectInstance = FObjectTrace::GetObjectFromId(ObjectId))
+							if (!RewindDebugger->IsTraceFileLoaded())
 							{
-								if(USkeletalMeshComponent* MeshComponent = Cast<USkeletalMeshComponent>(ObjectInstance))
+								if(UObject* ObjectInstance = FObjectTrace::GetObjectFromId(ObjectId))
 								{
-									AnimationProvider->ReadSkeletalMeshPoseTimeline(ObjectId, [this, &Frame, ObjectId, MeshComponent, AnimationProvider](const IAnimationProvider::SkeletalMeshPoseTimeline& TimelineData, bool bHasCurves)
+									MeshComponent = Cast<USkeletalMeshComponent>(ObjectInstance);
+									if (MeshComponent)
 									{
-										const FSkeletalMeshPoseMessage * PoseMessage = nullptr;
+										bQueueForReset = true;
+									}
+								}
+							}
+#endif // OBJECT_TRACE_ENABLED
+							if (MeshComponent == nullptr)
+							{
+								// display pose on a spawned mesh component
+								FSpawnedMeshComponentInfo* MeshComponentInfo = SpawnedMeshComponents.Find(ObjectId);
+								if (MeshComponentInfo == nullptr)
+								{
+									MeshComponentInfo = &SpawnedMeshComponents.Add(ObjectId);
+									// spawn one
 
-										// Get last pose in frame
-										TimelineData.EnumerateEvents(Frame.StartTime, Frame.EndTime,
-											[&PoseMessage](double InStartTime, double InEndTime, uint32 InDepth, const FSkeletalMeshPoseMessage& InPoseMessage)
-											{
-												PoseMessage = &InPoseMessage;
-												return TraceServices::EEventEnumerate::Continue;
-											});
+									UEditorEngine* EditorEngine = Cast<UEditorEngine>(GEngine);
+									UWorld* World = EditorEngine->GetEditorWorldContext().World();
 
-										// Update mesh based on pose
-										if (PoseMessage)
-										{
-											FTransform ComponentWorldTransform;
-											if (const FSkeletalMeshInfo* SkeletalMeshInfo = AnimationProvider->FindSkeletalMeshInfo(PoseMessage->MeshId))
-											{
-												AnimationProvider->GetSkeletalMeshComponentSpacePose(*PoseMessage, *SkeletalMeshInfo, ComponentWorldTransform, MeshComponent->GetEditableComponentSpaceTransforms());
-												MeshComponent->ApplyEditedComponentSpaceTransforms();
+									FActorSpawnParameters ActorSpawnParameters;
+									ActorSpawnParameters.bHideFromSceneOutliner = true;
+									ActorSpawnParameters.ObjectFlags |= RF_Transient;
 
-												if (MeshComponentsToReset.Find(ObjectId) == nullptr)
-												{
-													FMeshComponentResetData ResetData;
-													ResetData.Component = MeshComponent;
-													ResetData.RelativeTransform = MeshComponent->GetRelativeTransform();
-													MeshComponentsToReset.Add(ObjectId, ResetData);
-												}
+									MeshComponentInfo->Actor = World->SpawnActor<AActor>(ActorSpawnParameters);
+									MeshComponentInfo->Actor->SetActorLabel(TEXT("RewindDebugger"));
 
-												MeshComponent->SetWorldTransform(ComponentWorldTransform, false, nullptr, ETeleportType::TeleportPhysics);
-												MeshComponent->SetForcedLOD(PoseMessage->LodIndex + 1);
-												MeshComponent->UpdateChildTransforms(EUpdateTransformFlags::None, ETeleportType::TeleportPhysics);
-											}
-										}
-									});
+									MeshComponentInfo->Component = NewObject<USkeletalMeshComponent>(MeshComponentInfo->Actor);
+									MeshComponentInfo->Component->PrimaryComponentTick.bStartWithTickEnabled = false;
+									MeshComponentInfo->Component->PrimaryComponentTick.bCanEverTick = false;
+
+									MeshComponentInfo->Actor->AddInstanceComponent(MeshComponentInfo->Component);
+
+									MeshComponentInfo->Component->SetAnimationMode(EAnimationMode::AnimationCustomMode);
+									MeshComponentInfo->Component->RegisterComponentWithWorld(World);
+
+									bLoadMesh = true;
 								}
+
+								MeshComponent = MeshComponentInfo->Component;
 							}
-#endif // OBJECT_TRACE_ENABLED
-							});
+
+							if (MeshComponent)
+							{
+								ApplyPoseToMesh(AnimationProvider, GameplayProvider, Frame, TimelineData, MeshComponent, ObjectId, bQueueForReset, bLoadMesh);
+							}
+						}
+						);
 						}
         
 						{
diff --git a/Engine/Plugins/Animation/GameplayInsights/Source/RewindDebugger/Private/RewindDebuggerAnimation.h b/Engine/Plugins/Animation/GameplayInsights/Source/RewindDebugger/Private/RewindDebuggerAnimation.h
--- a/Engine/Plugins/Animation/GameplayInsights/Source/RewindDebugger/Private/RewindDebuggerAnimation.h
+++ b/Engine/Plugins/Animation/GameplayInsights/Source/RewindDebugger/Private/RewindDebuggerAnimation.h
@@ -2,10 +2,18 @@
 #pragma once
 
 #include "CoreMinimal.h"
+#include "IAnimationProvider.h"
+#include "IGameplayProvider.h"
 #include "IRewindDebuggerExtension.h"
 #include "RewindDebuggerSettings.h"
 #include "UObject/WeakObjectPtr.h"
 
+
+namespace TraceServices
+{
+	struct FFrame;
+}
+
 // Rewind debugger extension for animation support
 //  replay of animated pose data
 //  updating animation blueprint debugger
@@ -23,18 +31,36 @@ public:
 	virtual void Update(float DeltaTime, IRewindDebugger* RewindDebugger) override;
 	virtual void RecordingStarted(IRewindDebugger* RewindDebugger) override;
 	virtual void RecordingStopped(IRewindDebugger* RewindDebugger) override;
+	void ClearSpawnedComponents();
+	virtual void Clear(IRewindDebugger* RewindDebugger) override;
 	
 	void OnPIEResumed(bool bSimulating);
 	void OnPIEStopped(bool bSimulating);
 	void OnPIESingleStepped(bool bSimulating);
 
 private:
+	
+	void ApplyPoseToMesh(const class IAnimationProvider* AnimationProvider, const IGameplayProvider* GameplayProvider, const TraceServices::FFrame& Frame,
+		const IAnimationProvider::SkeletalMeshPoseTimeline& TimelineData, USkeletalMeshComponent* MeshComponent, uint64 ObjectId, bool bQueueForReset, bool bApplyMesh);
+	
 	struct FMeshComponentResetData
 	{
 		TWeakObjectPtr<USkeletalMeshComponent> Component;
 		FTransform RelativeTransform;
 	};
 
+	struct FSpawnedMeshComponentInfo
+	{
+		// mesh component object id
+		uint64 id;
+		// actor to hold the mesh component
+		TObjectPtr<AActor> Actor;
+		// mesh
+		TObjectPtr<USkeletalMeshComponent> Component;
+	};
+
+	TMap<uint64, FSpawnedMeshComponentInfo> SpawnedMeshComponents;
+
 	TMap<uint64, FMeshComponentResetData> MeshComponentsToReset;
 	double LastScrubTime = 0.0;
 };
diff --git a/Engine/Plugins/Animation/GameplayInsights/Source/RewindDebugger/Private/RewindDebuggerCommands.h b/Engine/Plugins/Animation/GameplayInsights/Source/RewindDebugger/Private/RewindDebuggerCommands.h
--- a/Engine/Plugins/Animation/GameplayInsights/Source/RewindDebugger/Private/RewindDebuggerCommands.h
+++ b/Engine/Plugins/Animation/GameplayInsights/Source/RewindDebugger/Private/RewindDebuggerCommands.h
@@ -33,6 +33,9 @@ public:
 		UI_COMMAND(LastFrame, "Last Frame", "Jump to last recorded frame", EUserInterfaceActionType::Button, FInputChord(EModifierKey::Control, EKeys::Up));
 		UI_COMMAND(AutoEject, "Auto Eject", "Automatically detach player control when PIE is paused", EUserInterfaceActionType::ToggleButton, FInputChord());
 		UI_COMMAND(AutoRecord, "Auto Record", "Automatically start recording when PIE is started", EUserInterfaceActionType::ToggleButton, FInputChord());
+		UI_COMMAND(OpenTrace, "Open Recording", "Open a Trace file containing a rewind debugger recording", EUserInterfaceActionType::Button, FInputChord());
+		UI_COMMAND(SaveTrace, "Save Recording", "Save a copy of the current recording", EUserInterfaceActionType::Button, FInputChord());
+		UI_COMMAND(ClearTrace, "Clear Recording", "Clear current recording data", EUserInterfaceActionType::Button, FInputChord());
 	}
 
 
@@ -48,6 +51,9 @@ public:
 	TSharedPtr<FUICommandInfo> PauseOrPlay;
 	TSharedPtr<FUICommandInfo> AutoEject;
 	TSharedPtr<FUICommandInfo> AutoRecord;
+	TSharedPtr<FUICommandInfo> OpenTrace;
+	TSharedPtr<FUICommandInfo> SaveTrace;
+	TSharedPtr<FUICommandInfo> ClearTrace;
 };
 
 
diff --git a/Engine/Plugins/Animation/GameplayInsights/Source/RewindDebugger/Private/RewindDebuggerModule.cpp b/Engine/Plugins/Animation/GameplayInsights/Source/RewindDebugger/Private/RewindDebuggerModule.cpp
--- a/Engine/Plugins/Animation/GameplayInsights/Source/RewindDebugger/Private/RewindDebuggerModule.cpp
+++ b/Engine/Plugins/Animation/GameplayInsights/Source/RewindDebugger/Private/RewindDebuggerModule.cpp
@@ -154,7 +154,25 @@ TSharedRef<SDockTab> FRewindDebuggerModule::SpawnRewindDebuggerTab(const FSpawnT
 							 FCanExecuteAction(),
 							 FCanExecuteAction::CreateRaw(DebuggerInstance, &FRewindDebugger::ShouldAutoRecordOnPIE),
 							 FIsActionButtonVisible());
-		 
+	
+	CommandList->MapAction(Commands.OpenTrace,
+							 FExecuteAction::CreateRaw(DebuggerInstance, &FRewindDebugger::OpenTrace),
+							 FCanExecuteAction::CreateRaw(DebuggerInstance, &FRewindDebugger::CanOpenTrace),
+							 FIsActionChecked(),
+							 FIsActionButtonVisible());
+		
+	CommandList->MapAction(Commands.SaveTrace,
+							 FExecuteAction::CreateRaw(DebuggerInstance, &FRewindDebugger::SaveTrace),
+							 FCanExecuteAction::CreateRaw(DebuggerInstance, &FRewindDebugger::CanSaveTrace),
+							 FIsActionChecked(),
+							 FIsActionButtonVisible());
+	
+	CommandList->MapAction(Commands.ClearTrace,
+							 FExecuteAction::CreateRaw(DebuggerInstance, &FRewindDebugger::ClearTrace),
+							 FCanExecuteAction::CreateRaw(DebuggerInstance, &FRewindDebugger::CanClearTrace),
+							 FIsActionChecked(),
+							 FIsActionButtonVisible());
+
 
 	// Register PIE Rewind Debugger Commands
 	if (GEditor != nullptr)
diff --git a/Engine/Plugins/Animation/GameplayInsights/Source/RewindDebugger/Private/SRewindDebugger.h b/Engine/Plugins/Animation/GameplayInsights/Source/RewindDebugger/Private/SRewindDebugger.h
--- a/Engine/Plugins/Animation/GameplayInsights/Source/RewindDebugger/Private/SRewindDebugger.h
+++ b/Engine/Plugins/Animation/GameplayInsights/Source/RewindDebugger/Private/SRewindDebugger.h
@@ -36,7 +36,7 @@ public:
 		SLATE_ARGUMENT( TArray< TSharedPtr< RewindDebugger::FRewindDebuggerTrack > >*, DebugComponents );
 		SLATE_ARGUMENT(DebugTargetInitializer, DebugTargetActor);
 		SLATE_ARGUMENT(TBindablePropertyInitializer<double>, TraceTime);
-		SLATE_ARGUMENT(TBindablePropertyInitializer<float>, RecordingDuration);
+		SLATE_ARGUMENT(TBindablePropertyInitializer<double>, RecordingDuration);
 		SLATE_ATTRIBUTE(TArrayView<RewindDebugger::FRewindDebuggerTrackType>, TrackTypes);
 		SLATE_ATTRIBUTE(double, ScrubTime);
 		SLATE_ATTRIBUTE(bool, IsPIESimulating);
@@ -89,7 +89,7 @@ private:
 	FOnViewRangeChanged OnViewRangeChanged;
 	TRange<double> ViewRange;
 	TBindableProperty<double> TraceTime;
-	TBindableProperty<float> RecordingDuration;
+	TBindableProperty<double> RecordingDuration;
 
 	TSharedPtr<FUICommandList> CommandList;
 	const FRewindDebuggerCommands & Commands;
diff --git a/Engine/Source/Editor/RewindDebuggerInterface/Public/IRewindDebugger.h b/Engine/Source/Editor/RewindDebuggerInterface/Public/IRewindDebugger.h
--- a/Engine/Source/Editor/RewindDebuggerInterface/Public/IRewindDebugger.h
+++ b/Engine/Source/Editor/RewindDebuggerInterface/Public/IRewindDebugger.h
@@ -90,6 +90,9 @@ public:
 	// returns true if recording is active
 	virtual bool IsRecording() const = 0;
 
+	// returns true if we have loaded a trace file from disk
+	virtual bool IsTraceFileLoaded() const = 0;
+
 	// returns true if PIE is running and not paused
 	virtual bool IsPIESimulating() const = 0;
 
diff --git a/Engine/Source/Editor/RewindDebuggerInterface/Public/IRewindDebuggerExtension.h b/Engine/Source/Editor/RewindDebuggerInterface/Public/IRewindDebuggerExtension.h
--- a/Engine/Source/Editor/RewindDebuggerInterface/Public/IRewindDebuggerExtension.h
+++ b/Engine/Source/Editor/RewindDebuggerInterface/Public/IRewindDebuggerExtension.h
@@ -26,4 +26,7 @@ public:
 
 	// called when recording has ended
 	virtual void RecordingStopped(IRewindDebugger* RewindDebugger) {};
+	
+	// called when recording is unloaded
+	virtual void Clear(IRewindDebugger* RewindDebugger) {};
 };
