Fix eyedropper actor selection in rewind debugger when loading trace files

diff --git a/Engine/Plugins/Animation/GameplayInsights/Source/RewindDebugger/Private/RewindDebuggerAnimation.cpp b/Engine/Plugins/Animation/GameplayInsights/Source/RewindDebugger/Private/RewindDebuggerAnimation.cpp
--- a/Engine/Plugins/Animation/GameplayInsights/Source/RewindDebugger/Private/RewindDebuggerAnimation.cpp
+++ b/Engine/Plugins/Animation/GameplayInsights/Source/RewindDebugger/Private/RewindDebuggerAnimation.cpp
@@ -160,6 +160,52 @@ void FRewindDebuggerAnimation::ApplyPoseToMesh(const IAnimationProvider* Animati
 	}
 }
 
+bool IsActorClass(const IGameplayProvider* GameplayProvider, uint64 ClassId)
+{
+	static const FString ActorName = "Actor";
+	while(true)
+	{
+		const FClassInfo& ClassInfo = GameplayProvider->GetClassInfo(ClassId);
+		if (ActorName == ClassInfo.Name)
+		{
+			return true;
+		}
+		else
+		{
+			if (ClassInfo.SuperId != 0)
+			{
+				ClassId = ClassInfo.SuperId;
+			}
+			else
+			{
+				return false;
+			}
+		}
+	}
+}
+
+const FObjectInfo* FindOwningActorInfo(const IGameplayProvider* GameplayProvider, uint64 ObjectId)
+{
+	while(true)
+	{
+		const FObjectInfo& ObjectInfo = GameplayProvider->GetObjectInfo(ObjectId);
+		if (IsActorClass(GameplayProvider, ObjectInfo.ClassId))
+		{
+			return &ObjectInfo;
+		}
+		else
+		{
+			if (ObjectInfo.OuterId != 0)
+			{
+				ObjectId = ObjectInfo.OuterId;
+			}
+			else
+			{
+				return nullptr;
+			}
+		}
+	}
+}
 
 void FRewindDebuggerAnimation::Update(float DeltaTime, IRewindDebugger* RewindDebugger)
 {
@@ -229,8 +275,13 @@ void FRewindDebuggerAnimation::Update(float DeltaTime, IRewindDebugger* RewindDe
 									ActorSpawnParameters.bHideFromSceneOutliner = true;
 									ActorSpawnParameters.ObjectFlags |= RF_Transient;
 
+
 									MeshComponentInfo->Actor = World->SpawnActor<AActor>(ActorSpawnParameters);
-									MeshComponentInfo->Actor->SetActorLabel(TEXT("RewindDebugger"));
+
+									if(const FObjectInfo* ActorInfo = FindOwningActorInfo(GameplayProvider, ObjectId))
+									{
+										MeshComponentInfo->Actor->SetActorLabel(FString(TEXT("RewindDebugger: ") + FString(ActorInfo->Name)));
+									}
 
 									MeshComponentInfo->Component = NewObject<USkeletalMeshComponent>(MeshComponentInfo->Actor);
 									MeshComponentInfo->Component->PrimaryComponentTick.bStartWithTickEnabled = false;
diff --git a/Engine/Plugins/Animation/GameplayInsights/Source/RewindDebugger/Private/SRewindDebugger.cpp b/Engine/Plugins/Animation/GameplayInsights/Source/RewindDebugger/Private/SRewindDebugger.cpp
--- a/Engine/Plugins/Animation/GameplayInsights/Source/RewindDebugger/Private/SRewindDebugger.cpp
+++ b/Engine/Plugins/Animation/GameplayInsights/Source/RewindDebugger/Private/SRewindDebugger.cpp
@@ -145,7 +145,17 @@ FReply SRewindDebugger::OnPreviewKeyDown(const FGeometry& MyGeometry, const FKey
 
 void SRewindDebugger::SetDebugTargetActor(AActor* Actor)
 {
-	DebugTargetActor.Set(Actor->GetName());
+	FString ActorLabel = Actor->GetActorLabel();
+	// Spawned actors have a "RewindDebugger: " prefix on their label
+	if (ActorLabel.StartsWith("RewindDebugger: "))
+	{
+		ActorLabel.RemoveFromStart("RewindDebugger: ");
+		DebugTargetActor.Set(ActorLabel);
+	}
+	else
+	{
+		DebugTargetActor.Set(Actor->GetName());
+	}
 }
 
 TSharedRef<SWidget> SRewindDebugger::MakeSelectActorMenu()
