Rewind debugger target selection menu reimplemented to be based on actors traced.

diff --git a/Engine/Plugins/Animation/GameplayInsights/Source/GameplayInsights/Private/GameplayProvider.cpp b/Engine/Plugins/Animation/GameplayInsights/Source/GameplayInsights/Private/GameplayProvider.cpp
--- a/Engine/Plugins/Animation/GameplayInsights/Source/GameplayInsights/Private/GameplayProvider.cpp
+++ b/Engine/Plugins/Animation/GameplayInsights/Source/GameplayInsights/Private/GameplayProvider.cpp
@@ -298,6 +298,32 @@ const FClassInfo* FGameplayProvider::FindClassInfo(const TCHAR* InClassPath) con
 	return nullptr;
 }
 
+bool FGameplayProvider::IsSubClassOf(uint64 InSubClassId, uint64 InParentClassId) const
+{
+	Session.ReadAccessCheck();
+
+	uint64 ClassId = InSubClassId;
+	while(true)
+	{
+		if (ClassId == InParentClassId)
+		{
+			return true;
+		}
+		else
+		{
+			const FClassInfo& ClassInfo = GetClassInfo(ClassId);
+			if (ClassInfo.SuperId != 0)
+			{
+				ClassId = ClassInfo.SuperId;
+			}
+			else
+			{
+				return false;
+			}
+		}
+	}
+}
+
 const FObjectInfo* FGameplayProvider::FindObjectInfo(uint64 InObjectId) const
 {
 	Session.ReadAccessCheck();
diff --git a/Engine/Plugins/Animation/GameplayInsights/Source/GameplayInsights/Private/GameplayProvider.h b/Engine/Plugins/Animation/GameplayInsights/Source/GameplayInsights/Private/GameplayProvider.h
--- a/Engine/Plugins/Animation/GameplayInsights/Source/GameplayInsights/Private/GameplayProvider.h
+++ b/Engine/Plugins/Animation/GameplayInsights/Source/GameplayInsights/Private/GameplayProvider.h
@@ -31,6 +31,7 @@ public:
 	virtual const FClassInfo* FindClassInfo(uint64 InClassId) const override;
 	virtual const UClass* FindClass(uint64 InClassId) const override;
 	virtual const FClassInfo* FindClassInfo(const TCHAR* InClassPath) const override;
+	virtual bool IsSubClassOf(uint64 InSubClassId, uint64 InParentClassId) const override;
 	virtual const FObjectInfo* FindObjectInfo(uint64 InObjectId) const override;
 	virtual const FWorldInfo* FindWorldInfo(uint64 InObjectId) const override;
 	virtual const FWorldInfo* FindWorldInfoFromObject(uint64 InObjectId) const override;
diff --git a/Engine/Plugins/Animation/GameplayInsights/Source/GameplayInsights/Public/IGameplayProvider.h b/Engine/Plugins/Animation/GameplayInsights/Source/GameplayInsights/Public/IGameplayProvider.h
--- a/Engine/Plugins/Animation/GameplayInsights/Source/GameplayInsights/Public/IGameplayProvider.h
+++ b/Engine/Plugins/Animation/GameplayInsights/Source/GameplayInsights/Public/IGameplayProvider.h
@@ -121,6 +121,7 @@ public:
 	virtual const FClassInfo* FindClassInfo(uint64 InClassId) const = 0;
 	virtual const UClass* FindClass(uint64 InClassId) const = 0;
 	virtual const FClassInfo* FindClassInfo(const TCHAR* InClassPath) const = 0;
+	virtual bool IsSubClassOf(uint64 InSubClassId, uint64 InParentClassId) const = 0;
 	virtual const FObjectInfo* FindObjectInfo(uint64 InObjectId) const = 0;
 	virtual const FWorldInfo* FindWorldInfo(uint64 InObjectId) const = 0;
 	virtual const FWorldInfo* FindWorldInfoFromObject(uint64 InObjectId) const = 0;
diff --git a/Engine/Plugins/Animation/GameplayInsights/Source/RewindDebugger/Private/RewindDebugger.cpp b/Engine/Plugins/Animation/GameplayInsights/Source/RewindDebugger/Private/RewindDebugger.cpp
--- a/Engine/Plugins/Animation/GameplayInsights/Source/RewindDebugger/Private/RewindDebugger.cpp
+++ b/Engine/Plugins/Animation/GameplayInsights/Source/RewindDebugger/Private/RewindDebugger.cpp
@@ -968,21 +968,17 @@ void FRewindDebugger::Tick(float DeltaTime)
 			{
 				FObjectTrace::ResetWorldElapsedTime(*World);
 				FObjectTrace::SetWorldRecordingIndex(*World, RecordingIndex);
-				
+		 		
 				TRACE_WORLD(*World);
 					
-				for (TActorIterator<AActor> Iterator(*World); Iterator; ++Iterator)
+				for (TActorIterator<AController> Iterator(*World); Iterator; ++Iterator)
 				{
-					TRACE_OBJECT_LIFETIME_BEGIN(*Iterator);
-					if (APawn* Pawn = Cast<APawn>(*Iterator))
+					if (APawn* Pawn = Iterator->GetPawn())
 					{
-						if (AController* Controller = Pawn->GetController())
-						{
-							TRACE_PAWN_POSSESS(static_cast<UObject*>(Controller), static_cast<UObject*>(Pawn));
-						}
+						TRACE_PAWN_POSSESS(static_cast<UObject*>(*Iterator), static_cast<UObject*>(Pawn));
 					}
 				}
-			}
+		 	}
 		}
 	#endif // OBJECT_TRACE_ENABLED
 	}
diff --git a/Engine/Plugins/Animation/GameplayInsights/Source/RewindDebugger/Private/RewindDebuggerAnimation.cpp b/Engine/Plugins/Animation/GameplayInsights/Source/RewindDebugger/Private/RewindDebuggerAnimation.cpp
--- a/Engine/Plugins/Animation/GameplayInsights/Source/RewindDebugger/Private/RewindDebuggerAnimation.cpp
+++ b/Engine/Plugins/Animation/GameplayInsights/Source/RewindDebugger/Private/RewindDebuggerAnimation.cpp
@@ -160,36 +160,14 @@ void FRewindDebuggerAnimation::ApplyPoseToMesh(const IAnimationProvider* Animati
 	}
 }
 
-bool IsActorClass(const IGameplayProvider* GameplayProvider, uint64 ClassId)
-{
-	static const FString ActorName = "Actor";
-	while(true)
-	{
-		const FClassInfo& ClassInfo = GameplayProvider->GetClassInfo(ClassId);
-		if (ActorName == ClassInfo.Name)
-		{
-			return true;
-		}
-		else
-		{
-			if (ClassInfo.SuperId != 0)
-			{
-				ClassId = ClassInfo.SuperId;
-			}
-			else
-			{
-				return false;
-			}
-		}
-	}
-}
-
 const FObjectInfo* FindOwningActorInfo(const IGameplayProvider* GameplayProvider, uint64 ObjectId)
 {
+	const FClassInfo* ActorClassInfo = GameplayProvider->FindClassInfo(*AActor::StaticClass()->GetPathName());
+	
 	while(true)
 	{
 		const FObjectInfo& ObjectInfo = GameplayProvider->GetObjectInfo(ObjectId);
-		if (IsActorClass(GameplayProvider, ObjectInfo.ClassId))
+		if (GameplayProvider->IsSubClassOf(ObjectInfo.ClassId, ActorClassInfo->Id))
 		{
 			return &ObjectInfo;
 		}
diff --git a/Engine/Plugins/Animation/GameplayInsights/Source/RewindDebugger/Private/SRewindDebugger.cpp b/Engine/Plugins/Animation/GameplayInsights/Source/RewindDebugger/Private/SRewindDebugger.cpp
--- a/Engine/Plugins/Animation/GameplayInsights/Source/RewindDebugger/Private/SRewindDebugger.cpp
+++ b/Engine/Plugins/Animation/GameplayInsights/Source/RewindDebugger/Private/SRewindDebugger.cpp
@@ -4,7 +4,9 @@
 
 #include "ActorPickerMode.h"
 #include "Editor.h"
+#include "IGameplayProvider.h"
 #include "IRewindDebuggerTrackCreator.h"
+#include "LocalizationDescriptor.h"
 #include "Editor/EditorEngine.h"
 #include "Framework/Application/SlateApplication.h"
 #include "Framework/MultiBox/MultiBoxBuilder.h"
@@ -28,6 +30,7 @@
 #include "Widgets/Layout/SScrollBox.h"
 #include "Kismet2/DebuggerCommands.h"
 #include "RewindDebuggerSettings.h"
+#include "Widgets/Input/STextComboBox.h"
 
 #define LOCTEXT_NAMESPACE "SRewindDebugger"
 
@@ -160,21 +163,19 @@ void SRewindDebugger::SetDebugTargetActor(AActor* Actor)
 
 TSharedRef<SWidget> SRewindDebugger::MakeSelectActorMenu()
 {
-	// this menu is partially duplicated from LevelSequenceEditorActorBinding which has a similar workflow for adding actors to sequencer
-
 	FMenuBuilder MenuBuilder(true, nullptr);
 
-	// Set up a menu entry to choosing the selected actor(s) to the sequencer (maybe move this to a submenu and put each selected actor there)
+	MenuBuilder.AddSearchWidget();
+
+	// Menu entry for each actor selected in the scene
 	TArray<AActor*> SelectedActors;
 	GEditor->GetSelectedActors()->GetSelectedObjects(SelectedActors);
 
-	if (SelectedActors.Num() >= 1)
+	MenuBuilder.BeginSection("From Selection Section", LOCTEXT("FromSelection", "From Scene Selection"));
+	if (SelectedActors.Num() > 0)
 	{
-		MenuBuilder.BeginSection("From Selection Section", LOCTEXT("FromSelection", "From Selection"));
-		if (SelectedActors.Num() == 1)
+		for(AActor* SelectedActor : SelectedActors)
 		{
-			AActor* SelectedActor = SelectedActors[0];
-
 			FText SelectedLabel = FText::FromString(SelectedActor->GetActorLabel());
 			FSlateIcon ActorIcon = FSlateIconFinder::FindIconForClass(SelectedActors[0]->GetClass());
 
@@ -183,72 +184,64 @@ TSharedRef<SWidget> SRewindDebugger::MakeSelectActorMenu()
 				SetDebugTargetActor(SelectedActor);
 			}));
 		}
-		else if (SelectedActors.Num() >= 1)
-		{
-			MenuBuilder.AddSubMenu(
-				LOCTEXT("FromSelection", "From Selection"),
-				LOCTEXT("FromSelection_Tooltip", "Select an Actor from the list of selected Actors"),
-				FNewMenuDelegate::CreateLambda([this, SelectedActors](FMenuBuilder& SubMenuBuilder)
-				{
-					for(AActor* SelectedActor : SelectedActors)
-					{
-						FText SelectedLabel = FText::FromString(SelectedActor->GetActorLabel());
-						FSlateIcon ActorIcon = FSlateIconFinder::FindIconForClass(SelectedActors[0]->GetClass());
-
-						SubMenuBuilder.AddMenuEntry(SelectedLabel, FText(), ActorIcon, FExecuteAction::CreateLambda([this, SelectedActor]{
-							FSlateApplication::Get().DismissAllMenus();
-							SetDebugTargetActor(SelectedActor);
-						}));
-					}
-
-				})
-			);
-
-		}
-		MenuBuilder.EndSection();
 	}
-
-	// todo: add special menu item for player controlled character
-
-	MenuBuilder.BeginSection("ChooseActorSection", LOCTEXT("ChooseActor", "Choose Actor:"));
-
-	// Set up a menu entry to select any arbitrary actor
-	FSceneOutlinerInitializationOptions InitOptions;
+	else
 	{
-		// We hide the header row to keep the UI compact.
-		InitOptions.bShowHeaderRow = false;
-		InitOptions.bShowSearchBox = true;
-		InitOptions.bShowCreateNewFolder = false;
-		InitOptions.bFocusSearchBoxWhenOpened = true;
-
-		// Only want the actor label column
-		InitOptions.ColumnMap.Add(FSceneOutlinerBuiltInColumnTypes::Label(), FSceneOutlinerColumnInfo(ESceneOutlinerColumnVisibility::Visible, 0));
-
-		// todo: optionally filter for only actors that have debug data
-		//InitOptions.Filters->AddFilterPredicate<FActorTreeItem>(FActorTreeItem::FFilterPredicate::CreateLambda(IsActorValidForPossession, ExistingPossessedObjects));
+		MenuBuilder.AddMenuEntry(LOCTEXT("No scene selection", "No scene selection"),
+		 LOCTEXT("SceneSelectionToolTip", "If you select an object in the scene, then it will be listed here"),
+		 FSlateIcon(),
+		  FUIAction(FExecuteAction(),
+		  	FCanExecuteAction::CreateLambda([](){return false;})));
 	}
+	MenuBuilder.EndSection();
 
-	// actor selector to allow the user to choose an actor
-	FSceneOutlinerModule& SceneOutlinerModule = FModuleManager::LoadModuleChecked<FSceneOutlinerModule>("SceneOutliner");
-	TSharedRef< SWidget > MiniSceneOutliner =
-		SNew(SBox)
-		.MaxDesiredHeight(400.0f)
-		.WidthOverride(300.0f)
-		[
-			SceneOutlinerModule.CreateActorPicker(
-				InitOptions,
-				FOnActorPicked::CreateLambda([this](AActor* Actor){
-					// Create a new binding for this actor
-					FSlateApplication::Get().DismissAllMenus();
-					SetDebugTargetActor(Actor);
-				})
-			)
-		];
+	MenuBuilder.BeginSection("From Recording Section", LOCTEXT("FromRecording", "From Recording:"));
 
-	MenuBuilder.AddWidget(MiniSceneOutliner, FText::GetEmpty(), true);
-	MenuBuilder.EndSection();
+	if (IRewindDebugger* RewindDebugger = IRewindDebugger::Instance())
+	{
+		if (const TraceServices::IAnalysisSession* Session = RewindDebugger->GetAnalysisSession())
+		{
+			TraceServices::FAnalysisSessionReadScope SessionReadScope(*Session);
+			const IGameplayProvider* GameplayProvider = Session->ReadProvider<IGameplayProvider>("GameplayProvider");
 
+			const FClassInfo* ActorClassInfo = GameplayProvider->FindClassInfo(*AActor::StaticClass()->GetPathName());
+			
+			GameplayProvider->EnumerateObjects([this, &MenuBuilder, GameplayProvider, ActorClassInfo](const FObjectInfo& ObjectInfo)
+			{
+				if (GameplayProvider->IsSubClassOf(ObjectInfo.ClassId, ActorClassInfo->Id))
+				{
+					FString ActorName = ObjectInfo.Name;
+					FText SelectedLabel = FText::FromString(ActorName);
+					FSlateIcon ActorIcon;
+
+					const FClassInfo& ClassInfo = GameplayProvider->GetClassInfo(ObjectInfo.ClassId);
+					if (UClass* FoundClass = UClass::TryFindTypeSlow<UClass>(ClassInfo.Name))
+					{
+						ActorIcon = FSlateIconFinder::FindIconForClass(FoundClass);
+					}
+					else
+					{
+						ActorIcon = FSlateIconFinder::FindIconForClass(AActor::StaticClass());
+					}
 
+					MenuBuilder.AddMenuEntry(SelectedLabel, FText(), ActorIcon, FExecuteAction::CreateLambda([this, ActorName]()
+					{
+						FSlateApplication::Get().DismissAllMenus();
+						DebugTargetActor.Set(ActorName);
+					}));
+				}
+			});
+		}
+		else
+		{
+			MenuBuilder.AddMenuEntry(LOCTEXT("No recording loaded", "No recording loaded"),
+			 LOCTEXT("NoRecordingToolTip", "Start or load a recording, and the recorded actors will be listed here"),
+			 FSlateIcon(),
+			  FUIAction(FExecuteAction(),
+				FCanExecuteAction::CreateLambda([](){return false;})));
+        }
+		
+	}
 
 	return MenuBuilder.MakeWidget();
 }
@@ -385,6 +378,7 @@ void SRewindDebugger::Construct(const FArguments& InArgs, TSharedRef<FUICommandL
 					SNew(SHorizontalBox)
 					+SHorizontalBox::Slot() .FillWidth(1.0f)
 					[
+						// SNew(SRewindDebuggerTargetSelector)
 						SNew(SComboButton)
 						.OnGetMenuContent(this, &SRewindDebugger::MakeSelectActorMenu)
 						.ButtonContent()
