Make Animation Blueprint debugging work with rewind debugger traces loaded from file

diff --git a/Engine/Plugins/Animation/GameplayInsights/Source/RewindDebugger/Private/AnimInstanceHelpers.cpp b/Engine/Plugins/Animation/GameplayInsights/Source/RewindDebugger/Private/AnimInstanceHelpers.cpp
--- a/Engine/Plugins/Animation/GameplayInsights/Source/RewindDebugger/Private/AnimInstanceHelpers.cpp
+++ b/Engine/Plugins/Animation/GameplayInsights/Source/RewindDebugger/Private/AnimInstanceHelpers.cpp
@@ -8,6 +8,7 @@
 #include "Animation/AnimBlueprintGeneratedClass.h"
 #include "Editor.h"
 #include "EdGraph/EdGraph.h"
+#include "RewindDebuggerAnimation.h"
 #include "Subsystems/AssetEditorSubsystem.h"
 #include "IAnimationBlueprintEditor.h"
 #include "Insights/IUnrealInsightsModule.h"
@@ -15,7 +16,6 @@
 
 #define LOCTEXT_NAMESPACE "SRewindDebuggerAnimBPTools"
 
-#if OBJECT_TRACE_ENABLED
 static bool OpenAnimBlueprintAndAttachDebugger(const TraceServices::IAnalysisSession* Session, uint64 ObjectId)
 {
 	TraceServices::FAnalysisSessionReadScope SessionReadScope(*Session);
@@ -31,8 +31,20 @@ static bool OpenAnimBlueprintAndAttachDebugger(const TraceServices::IAnalysisSes
 					if(UAnimBlueprint* AnimBlueprint = Cast<UAnimBlueprint>(InstanceClass->ClassGeneratedBy))
 					{
 						GEditor->GetEditorSubsystem<UAssetEditorSubsystem>()->OpenEditorForAsset(AnimBlueprint);
+						
+						UObject* SelectedInstance;
+#if OBJECT_TRACE_ENABLED
+						SelectedInstance = FObjectTrace::GetObjectFromId(ObjectId);
+#endif
+						if (SelectedInstance == nullptr)
+						{
+							if (FRewindDebuggerAnimation* RewindDebuggerAnimation = FRewindDebuggerAnimation::GetInstance())
+							{
+								SelectedInstance = RewindDebuggerAnimation->GetDebugAnimInstance(ObjectId);
+							}
+						}
 
-						if (UObject* SelectedInstance = FObjectTrace::GetObjectFromId(ObjectId))
+						if(SelectedInstance)
 						{
 							AnimBlueprint->SetObjectBeingDebugged(SelectedInstance);
 						}
@@ -54,17 +66,14 @@ static bool OpenAnimBlueprintAndAttachDebugger(const TraceServices::IAnalysisSes
 	}
 	return false;
 }
-#endif
 
 bool FAnimInstanceDoubleClickHandler::HandleDoubleClick(IRewindDebugger* RewindDebugger)
 {
-#if OBJECT_TRACE_ENABLED
 	TSharedPtr<FDebugObjectInfo> SelectedObject = RewindDebugger->GetSelectedComponent();
 	if (SelectedObject.IsValid())
 	{
 		return OpenAnimBlueprintAndAttachDebugger(RewindDebugger->GetAnalysisSession(), SelectedObject->ObjectId);
 	}
-#endif
 	return false;
 }
 
diff --git a/Engine/Plugins/Animation/GameplayInsights/Source/RewindDebugger/Private/RewindDebuggerAnimation.cpp b/Engine/Plugins/Animation/GameplayInsights/Source/RewindDebugger/Private/RewindDebuggerAnimation.cpp
--- a/Engine/Plugins/Animation/GameplayInsights/Source/RewindDebugger/Private/RewindDebuggerAnimation.cpp
+++ b/Engine/Plugins/Animation/GameplayInsights/Source/RewindDebugger/Private/RewindDebuggerAnimation.cpp
@@ -17,9 +17,12 @@
 #include "ObjectTrace.h"
 #include "TraceServices/Model/Frames.h"
 #include "Engine/SkeletalMesh.h"
+#include "UObject/SoftObjectPath.h"
 
 #define LOCTEXT_NAMESPACE "RewindDebuggerAnimation"
 
+FRewindDebuggerAnimation* FRewindDebuggerAnimation::Instance;
+
 FRewindDebuggerAnimation::FRewindDebuggerAnimation()
 {
 }
@@ -53,6 +56,7 @@ void FRewindDebuggerAnimation::ClearSpawnedComponents()
 		}
 	}
 	SpawnedMeshComponents.Empty();
+	SpawnedAnimInstances.Empty();
 }
 
 void FRewindDebuggerAnimation::Clear(IRewindDebugger* RewindDebugger)
@@ -65,6 +69,8 @@ void FRewindDebuggerAnimation::Initialize()
 	FEditorDelegates::ResumePIE.AddRaw(this, &FRewindDebuggerAnimation::OnPIEResumed);
 	FEditorDelegates::EndPIE.AddRaw(this, &FRewindDebuggerAnimation::OnPIEStopped);
 	FEditorDelegates::SingleStepPIE.AddRaw(this, &FRewindDebuggerAnimation::OnPIESingleStepped);
+
+	Instance = this;
 }
 
 void FRewindDebuggerAnimation::Shutdown()
@@ -72,6 +78,8 @@ void FRewindDebuggerAnimation::Shutdown()
 	FEditorDelegates::ResumePIE.RemoveAll(this);
 	FEditorDelegates::EndPIE.RemoveAll(this);
 	FEditorDelegates::SingleStepPIE.RemoveAll(this);
+
+	Instance = nullptr;
 }
 
 void FRewindDebuggerAnimation::OnPIEResumed(bool bSimulating)
@@ -185,6 +193,76 @@ const FObjectInfo* FindOwningActorInfo(const IGameplayProvider* GameplayProvider
 	}
 }
 
+FRewindDebuggerAnimation::FSpawnedMeshComponentInfo* FRewindDebuggerAnimation::SpawnMesh(uint64 ObjectId, const IGameplayProvider* GameplayProvider)
+{
+	FSpawnedMeshComponentInfo* MeshComponentInfo = &SpawnedMeshComponents.Add(ObjectId);
+
+	UEditorEngine* EditorEngine = Cast<UEditorEngine>(GEngine);
+	UWorld* World = EditorEngine->GetEditorWorldContext().World();
+
+	FActorSpawnParameters ActorSpawnParameters;
+	ActorSpawnParameters.bHideFromSceneOutliner = true;
+	ActorSpawnParameters.ObjectFlags |= RF_Transient;
+
+
+	MeshComponentInfo->Actor = World->SpawnActor<AActor>(ActorSpawnParameters);
+
+	if(const FObjectInfo* ActorInfo = FindOwningActorInfo(GameplayProvider, ObjectId))
+	{
+		MeshComponentInfo->Actor->SetActorLabel(FString(TEXT("RewindDebugger: ") + FString(ActorInfo->Name)));
+	}
+
+	MeshComponentInfo->Component = NewObject<USkeletalMeshComponent>(MeshComponentInfo->Actor);
+	MeshComponentInfo->Component->PrimaryComponentTick.bStartWithTickEnabled = false;
+	MeshComponentInfo->Component->PrimaryComponentTick.bCanEverTick = false;
+
+	MeshComponentInfo->Actor->AddInstanceComponent(MeshComponentInfo->Component);
+
+	MeshComponentInfo->Component->SetAnimationMode(EAnimationMode::AnimationCustomMode);
+	MeshComponentInfo->Component->RegisterComponentWithWorld(World);
+
+	return MeshComponentInfo;
+}
+
+UAnimInstance* FRewindDebuggerAnimation::SpawnAnimInstance(uint64 ObjectId, const IGameplayProvider* GameplayProvider)
+{
+	if (FSpawnedAnimInstanceInfo* FoundInfo = SpawnedAnimInstances.Find(ObjectId))
+	{
+		return FoundInfo->AnimInstance;
+	}
+	
+	FSpawnedAnimInstanceInfo* AnimInstanceInfo = &SpawnedAnimInstances.Add(ObjectId);
+
+	if (const FObjectInfo* ObjectInfo = GameplayProvider->FindObjectInfo(ObjectId))
+	{
+		if (const FClassInfo* ClassInfo = GameplayProvider->FindClassInfo(ObjectInfo->ClassId))
+		{
+			TSoftObjectPtr<UAnimBlueprintGeneratedClass> ClassPtr =	TSoftObjectPtr<UAnimBlueprintGeneratedClass>(FSoftObjectPath(FStringView(ClassInfo->PathName)));
+			if(ClassPtr.LoadSynchronous())
+			{
+				if (const FSpawnedMeshComponentInfo* MeshInfo = SpawnedMeshComponents.Find(ObjectInfo->OuterId))
+				{
+					AnimInstanceInfo->AnimInstance = NewObject<UAnimInstance>(MeshInfo->Component, ClassPtr.Get());
+					return AnimInstanceInfo->AnimInstance;
+				}
+			}
+		}
+	}
+
+	SpawnedAnimInstances.Remove(ObjectId);
+	return nullptr;
+}
+
+UAnimInstance* FRewindDebuggerAnimation::GetDebugAnimInstance(uint64 ObjectId)
+{
+	if (FSpawnedAnimInstanceInfo* FoundInfo = SpawnedAnimInstances.Find(ObjectId))
+	{
+		return FoundInfo->AnimInstance;
+	}
+	return nullptr;
+}
+
+
 void FRewindDebuggerAnimation::Update(float DeltaTime, IRewindDebugger* RewindDebugger)
 {
 	check(RewindDebugger);
@@ -243,36 +321,10 @@ void FRewindDebuggerAnimation::Update(float DeltaTime, IRewindDebugger* RewindDe
 								FSpawnedMeshComponentInfo* MeshComponentInfo = SpawnedMeshComponents.Find(ObjectId);
 								if (MeshComponentInfo == nullptr)
 								{
-									MeshComponentInfo = &SpawnedMeshComponents.Add(ObjectId);
-									// spawn one
-
-									UEditorEngine* EditorEngine = Cast<UEditorEngine>(GEngine);
-									UWorld* World = EditorEngine->GetEditorWorldContext().World();
-
-									FActorSpawnParameters ActorSpawnParameters;
-									ActorSpawnParameters.bHideFromSceneOutliner = true;
-									ActorSpawnParameters.ObjectFlags |= RF_Transient;
-
-
-									MeshComponentInfo->Actor = World->SpawnActor<AActor>(ActorSpawnParameters);
-
-									if(const FObjectInfo* ActorInfo = FindOwningActorInfo(GameplayProvider, ObjectId))
-									{
-										MeshComponentInfo->Actor->SetActorLabel(FString(TEXT("RewindDebugger: ") + FString(ActorInfo->Name)));
-									}
-
-									MeshComponentInfo->Component = NewObject<USkeletalMeshComponent>(MeshComponentInfo->Actor);
-									MeshComponentInfo->Component->PrimaryComponentTick.bStartWithTickEnabled = false;
-									MeshComponentInfo->Component->PrimaryComponentTick.bCanEverTick = false;
-
-									MeshComponentInfo->Actor->AddInstanceComponent(MeshComponentInfo->Component);
-
-									MeshComponentInfo->Component->SetAnimationMode(EAnimationMode::AnimationCustomMode);
-									MeshComponentInfo->Component->RegisterComponentWithWorld(World);
-
+									MeshComponentInfo = SpawnMesh(ObjectId, GameplayProvider);
 									bLoadMesh = true;
 								}
-
+								
 								MeshComponent = MeshComponentInfo->Component;
 							}
 
@@ -290,185 +342,192 @@ void FRewindDebuggerAnimation::Update(float DeltaTime, IRewindDebugger* RewindDe
 							// - enumerate over all anim graph timelines
 							// - check if their instance class still exists and is the debugging target for the Animation Blueprint Editor
 							// - if it is copy that debug data into the class debug data for the blueprint debugger
-							AnimationProvider->EnumerateAnimGraphTimelines([&Frame, AnimationProvider, GameplayProvider](uint64 ObjectId, const IAnimationProvider::AnimGraphTimeline& AnimGraphTimeline)
+							AnimationProvider->EnumerateAnimGraphTimelines([this, &Frame, AnimationProvider, GameplayProvider](uint64 ObjectId, const IAnimationProvider::AnimGraphTimeline& AnimGraphTimeline)
 							{
+								UAnimInstance* AnimInstance = nullptr;
 #if OBJECT_TRACE_ENABLED
 								if(UObject* ObjectInstance = FObjectTrace::GetObjectFromId(ObjectId))
 								{
-									if(UAnimInstance* AnimInstance = Cast<UAnimInstance>(ObjectInstance))
+									AnimInstance = Cast<UAnimInstance>(ObjectInstance);
+								}
+#endif
+								if (AnimInstance == nullptr)
+								{
+									AnimInstance = SpawnAnimInstance(ObjectId, GameplayProvider);
+								}
+								
+								if (AnimInstance)
+								{
+									if(UAnimBlueprintGeneratedClass* InstanceClass = Cast<UAnimBlueprintGeneratedClass>(AnimInstance->GetClass()))
 									{
-										if(UAnimBlueprintGeneratedClass* InstanceClass = Cast<UAnimBlueprintGeneratedClass>(AnimInstance->GetClass()))
+										if(UAnimBlueprint* AnimBlueprint = Cast<UAnimBlueprint>(InstanceClass->ClassGeneratedBy))
 										{
-											if(UAnimBlueprint* AnimBlueprint = Cast<UAnimBlueprint>(InstanceClass->ClassGeneratedBy))
+											// for child Animation Blueprints, we actually want to debug the root blueprint (since the child doesn't contain any anim graphs)
+											if (UAnimBlueprint* RootAnimBP = UAnimBlueprint::FindRootAnimBlueprint(AnimBlueprint))
 											{
-												// for child Animation Blueprints, we actually want to debug the root blueprint (since the child doesn't contain any anim graphs)
-												if (UAnimBlueprint* RootAnimBP = UAnimBlueprint::FindRootAnimBlueprint(AnimBlueprint))
+												if (UAnimBlueprintGeneratedClass* RootInstanceClass = Cast<UAnimBlueprintGeneratedClass>(RootAnimBP->GeneratedClass))
 												{
-													if (UAnimBlueprintGeneratedClass* RootInstanceClass = Cast<UAnimBlueprintGeneratedClass>(RootAnimBP->GeneratedClass))
-													{
-														AnimBlueprint = RootAnimBP;
-														InstanceClass = RootInstanceClass;
-													}
+													AnimBlueprint = RootAnimBP;
+													InstanceClass = RootInstanceClass;
 												}
+											}
 
-												if(AnimBlueprint->IsObjectBeingDebugged(AnimInstance))
+											if(AnimBlueprint->IsObjectBeingDebugged(AnimInstance))
+											{
+												TRACE_CPUPROFILER_EVENT_SCOPE(FRewindDebugger::Tick_UpdateBlueprintDebug);
+												// update debug info for attached Animation Blueprint editors
+												uint64 Id = FObjectTrace::GetObjectId(AnimInstance);
+												const int32 NodeCount = InstanceClass->GetAnimNodeProperties().Num();
+						
+												FAnimBlueprintDebugData& DebugData = InstanceClass->GetAnimBlueprintDebugData();
 												{
-													TRACE_CPUPROFILER_EVENT_SCOPE(FRewindDebugger::Tick_UpdateBlueprintDebug);
-													// update debug info for attached Animation Blueprint editors
-													uint64 Id = FObjectTrace::GetObjectId(AnimInstance);
-													const int32 NodeCount = InstanceClass->GetAnimNodeProperties().Num();
-							
-													FAnimBlueprintDebugData& DebugData = InstanceClass->GetAnimBlueprintDebugData();
-													{
-														TRACE_CPUPROFILER_EVENT_SCOPE(ResetNodeVisitStates);
-														DebugData.ResetNodeVisitSites();
-													}
-													
-													// Anim node values can come from all phases
-													AnimationProvider->ReadAnimNodeValuesTimeline(ObjectId, [&Frame,AnimationProvider, &DebugData](const IAnimationProvider::AnimNodeValuesTimeline& InNodeValuesTimeline)
+													TRACE_CPUPROFILER_EVENT_SCOPE(ResetNodeVisitStates);
+													DebugData.ResetNodeVisitSites();
+												}
+												
+												// Anim node values can come from all phases
+												AnimationProvider->ReadAnimNodeValuesTimeline(ObjectId, [&Frame,AnimationProvider, &DebugData](const IAnimationProvider::AnimNodeValuesTimeline& InNodeValuesTimeline)
+												{
+													TRACE_CPUPROFILER_EVENT_SCOPE(AnimGraphNodeValues);
+													InNodeValuesTimeline.EnumerateEvents(Frame.StartTime, Frame.EndTime, [AnimationProvider, &DebugData](double InStartTime, double InEndTime, uint32 InDepth, const FAnimNodeValueMessage& InMessage)
 													{
-														TRACE_CPUPROFILER_EVENT_SCOPE(AnimGraphNodeValues);
-														InNodeValuesTimeline.EnumerateEvents(Frame.StartTime, Frame.EndTime, [AnimationProvider, &DebugData](double InStartTime, double InEndTime, uint32 InDepth, const FAnimNodeValueMessage& InMessage)
+														// don't send "Name" Node value for display in the graph
+														if (FPlatformString::Strcmp(InMessage.Key, TEXT("Name")) != 0)
 														{
-															// don't send "Name" Node value for display in the graph
-															if (FPlatformString::Strcmp(InMessage.Key, TEXT("Name")) != 0)
-															{
-																FText Text = AnimationProvider->FormatNodeKeyValue(InMessage);
-																DebugData.RecordNodeValue(InMessage.NodeId, Text.ToString());
-															}
-															return TraceServices::EEventEnumerate::Continue;
-														});
+															FText Text = AnimationProvider->FormatNodeKeyValue(InMessage);
+															DebugData.RecordNodeValue(InMessage.NodeId, Text.ToString());
+														}
+														return TraceServices::EEventEnumerate::Continue;
 													});
+												});
 
-													DebugData.DisableAllPoseWatches();
-						
-													AnimGraphTimeline.EnumerateEvents(Frame.StartTime, Frame.EndTime, [Id, AnimationProvider, GameplayProvider, &DebugData, NodeCount](double InGraphStartTime, double InGraphEndTime, uint32 InDepth, const FAnimGraphMessage& InMessage)
+												DebugData.DisableAllPoseWatches();
+					
+												AnimGraphTimeline.EnumerateEvents(Frame.StartTime, Frame.EndTime, [Id, AnimationProvider, GameplayProvider, &DebugData, NodeCount](double InGraphStartTime, double InGraphEndTime, uint32 InDepth, const FAnimGraphMessage& InMessage)
+												{
+													TRACE_CPUPROFILER_EVENT_SCOPE(AnimGraphTimelineEvent);
+															
+													// Basic verification - check node count is the same
+													// @TODO: could add some form of node hash/CRC to the class to improve this
+													if(InMessage.NodeCount == NodeCount)
 													{
-														TRACE_CPUPROFILER_EVENT_SCOPE(AnimGraphTimelineEvent);
-																
-														// Basic verification - check node count is the same
-														// @TODO: could add some form of node hash/CRC to the class to improve this
-														if(InMessage.NodeCount == NodeCount)
+														// Check for an update phase (which contains weights)
+														if(InMessage.Phase == EAnimGraphPhase::Update)
 														{
-															// Check for an update phase (which contains weights)
-															if(InMessage.Phase == EAnimGraphPhase::Update)
+															AnimationProvider->ReadAnimNodesTimeline(Id, [InGraphStartTime, InGraphEndTime, &DebugData](const IAnimationProvider::AnimNodesTimeline& InNodesTimeline)
 															{
-																AnimationProvider->ReadAnimNodesTimeline(Id, [InGraphStartTime, InGraphEndTime, &DebugData](const IAnimationProvider::AnimNodesTimeline& InNodesTimeline)
+																TRACE_CPUPROFILER_EVENT_SCOPE(AnimGraphDebugNodeVisits);
+																InNodesTimeline.EnumerateEvents(InGraphStartTime, InGraphEndTime, [&DebugData](double InStartTime, double InEndTime, uint32 InDepth, const FAnimNodeMessage& InMessage)
 																{
-																	TRACE_CPUPROFILER_EVENT_SCOPE(AnimGraphDebugNodeVisits);
-																	InNodesTimeline.EnumerateEvents(InGraphStartTime, InGraphEndTime, [&DebugData](double InStartTime, double InEndTime, uint32 InDepth, const FAnimNodeMessage& InMessage)
-																	{
-																		DebugData.RecordNodeVisit(InMessage.NodeId, InMessage.PreviousNodeId, InMessage.Weight);
-																		return TraceServices::EEventEnumerate::Continue;
-																	});
+																	DebugData.RecordNodeVisit(InMessage.NodeId, InMessage.PreviousNodeId, InMessage.Weight);
+																	return TraceServices::EEventEnumerate::Continue;
 																});
-						
-																AnimationProvider->ReadStateMachinesTimeline(Id, [InGraphStartTime, InGraphEndTime, &DebugData](const IAnimationProvider::StateMachinesTimeline& InStateMachinesTimeline)
+															});
+					
+															AnimationProvider->ReadStateMachinesTimeline(Id, [InGraphStartTime, InGraphEndTime, &DebugData](const IAnimationProvider::StateMachinesTimeline& InStateMachinesTimeline)
+															{
+																TRACE_CPUPROFILER_EVENT_SCOPE(AnimGraphDebugStateMachine);
+																InStateMachinesTimeline.EnumerateEvents(InGraphStartTime, InGraphEndTime, [&DebugData](double InStartTime, double InEndTime, uint32 InDepth, const FAnimStateMachineMessage& InMessage)
 																{
-																	TRACE_CPUPROFILER_EVENT_SCOPE(AnimGraphDebugStateMachine);
-																	InStateMachinesTimeline.EnumerateEvents(InGraphStartTime, InGraphEndTime, [&DebugData](double InStartTime, double InEndTime, uint32 InDepth, const FAnimStateMachineMessage& InMessage)
-																	{
-																		DebugData.RecordStateData(InMessage.StateMachineIndex, InMessage.StateIndex, InMessage.StateWeight, InMessage.ElapsedTime);
-																		return TraceServices::EEventEnumerate::Continue;
-																	});
+																	DebugData.RecordStateData(InMessage.StateMachineIndex, InMessage.StateIndex, InMessage.StateWeight, InMessage.ElapsedTime);
+																	return TraceServices::EEventEnumerate::Continue;
 																});
-						
-																AnimationProvider->ReadAnimSequencePlayersTimeline(Id, [InGraphStartTime, InGraphEndTime, GameplayProvider, &DebugData](const IAnimationProvider::AnimSequencePlayersTimeline& InSequencePlayersTimeline)
+															});
+					
+															AnimationProvider->ReadAnimSequencePlayersTimeline(Id, [InGraphStartTime, InGraphEndTime, GameplayProvider, &DebugData](const IAnimationProvider::AnimSequencePlayersTimeline& InSequencePlayersTimeline)
+															{
+																TRACE_CPUPROFILER_EVENT_SCOPE(AnimGraphDebugSequencePlayers);
+																InSequencePlayersTimeline.EnumerateEvents(InGraphStartTime, InGraphEndTime, [&DebugData](double InStartTime, double InEndTime, uint32 InDepth, const FAnimSequencePlayerMessage& InMessage)
 																{
-																	TRACE_CPUPROFILER_EVENT_SCOPE(AnimGraphDebugSequencePlayers);
-																	InSequencePlayersTimeline.EnumerateEvents(InGraphStartTime, InGraphEndTime, [&DebugData](double InStartTime, double InEndTime, uint32 InDepth, const FAnimSequencePlayerMessage& InMessage)
-																	{
-																		DebugData.RecordSequencePlayer(InMessage.NodeId, InMessage.Position, InMessage.Length, InMessage.FrameCounter);
-																		return TraceServices::EEventEnumerate::Continue;
-																	});
+																	DebugData.RecordSequencePlayer(InMessage.NodeId, InMessage.Position, InMessage.Length, InMessage.FrameCounter);
+																	return TraceServices::EEventEnumerate::Continue;
 																});
-						
-																AnimationProvider->ReadAnimBlendSpacePlayersTimeline(Id, [InGraphStartTime, InGraphEndTime, GameplayProvider, &DebugData](const IAnimationProvider::BlendSpacePlayersTimeline& InBlendSpacePlayersTimeline)
+															});
+					
+															AnimationProvider->ReadAnimBlendSpacePlayersTimeline(Id, [InGraphStartTime, InGraphEndTime, GameplayProvider, &DebugData](const IAnimationProvider::BlendSpacePlayersTimeline& InBlendSpacePlayersTimeline)
+															{
+																TRACE_CPUPROFILER_EVENT_SCOPE(AnimGraphBlendSpaces);
+																InBlendSpacePlayersTimeline.EnumerateEvents(InGraphStartTime, InGraphEndTime, [GameplayProvider, &DebugData](double InStartTime, double InEndTime, uint32 InDepth, const FBlendSpacePlayerMessage& InMessage)
 																{
-																	TRACE_CPUPROFILER_EVENT_SCOPE(AnimGraphBlendSpaces);
-																	InBlendSpacePlayersTimeline.EnumerateEvents(InGraphStartTime, InGraphEndTime, [GameplayProvider, &DebugData](double InStartTime, double InEndTime, uint32 InDepth, const FBlendSpacePlayerMessage& InMessage)
+																	UBlendSpace* BlendSpace = nullptr;
+																	const FObjectInfo* BlendSpaceInfo = GameplayProvider->FindObjectInfo(InMessage.BlendSpaceId);
+																	if(BlendSpaceInfo)
 																	{
-																		UBlendSpace* BlendSpace = nullptr;
-																		const FObjectInfo* BlendSpaceInfo = GameplayProvider->FindObjectInfo(InMessage.BlendSpaceId);
-																		if(BlendSpaceInfo)
-																		{
-																			BlendSpace = TSoftObjectPtr<UBlendSpace>(FSoftObjectPath(BlendSpaceInfo->PathName)).LoadSynchronous();
-																		}
-						
-																		DebugData.RecordBlendSpacePlayer(InMessage.NodeId, BlendSpace, FVector(InMessage.PositionX, InMessage.PositionY, InMessage.PositionZ), FVector(InMessage.FilteredPositionX, InMessage.FilteredPositionY, InMessage.FilteredPositionZ));
-																		return TraceServices::EEventEnumerate::Continue;
-																	});
+																		BlendSpace = TSoftObjectPtr<UBlendSpace>(FSoftObjectPath(BlendSpaceInfo->PathName)).LoadSynchronous();
+																	}
+					
+																	DebugData.RecordBlendSpacePlayer(InMessage.NodeId, BlendSpace, FVector(InMessage.PositionX, InMessage.PositionY, InMessage.PositionZ), FVector(InMessage.FilteredPositionX, InMessage.FilteredPositionY, InMessage.FilteredPositionZ));
+																	return TraceServices::EEventEnumerate::Continue;
 																});
-						
-																AnimationProvider->ReadAnimSyncTimeline(Id, [InGraphStartTime, InGraphEndTime, AnimationProvider, &DebugData](const IAnimationProvider::AnimSyncTimeline& InAnimSyncTimeline)
+															});
+					
+															AnimationProvider->ReadAnimSyncTimeline(Id, [InGraphStartTime, InGraphEndTime, AnimationProvider, &DebugData](const IAnimationProvider::AnimSyncTimeline& InAnimSyncTimeline)
+															{
+																TRACE_CPUPROFILER_EVENT_SCOPE(AnimGraphAnimSync);
+																InAnimSyncTimeline.EnumerateEvents(InGraphStartTime, InGraphEndTime, [AnimationProvider, &DebugData](double InStartTime, double InEndTime, uint32 InDepth, const FAnimSyncMessage& InMessage)
 																{
-																	TRACE_CPUPROFILER_EVENT_SCOPE(AnimGraphAnimSync);
-																	InAnimSyncTimeline.EnumerateEvents(InGraphStartTime, InGraphEndTime, [AnimationProvider, &DebugData](double InStartTime, double InEndTime, uint32 InDepth, const FAnimSyncMessage& InMessage)
+																	const TCHAR* GroupName = AnimationProvider->GetName(InMessage.GroupNameId);
+																	if(GroupName)
 																	{
-																		const TCHAR* GroupName = AnimationProvider->GetName(InMessage.GroupNameId);
-																		if(GroupName)
-																		{
-																			DebugData.RecordNodeSync(InMessage.SourceNodeId, FName(GroupName));
-																		}
-															
-																		return TraceServices::EEventEnumerate::Continue;
-																	});
+																		DebugData.RecordNodeSync(InMessage.SourceNodeId, FName(GroupName));
+																	}
+														
+																	return TraceServices::EEventEnumerate::Continue;
 																});
-															}
-						
-															// Some traces come from both update and evaluate phases
-															if(InMessage.Phase == EAnimGraphPhase::Update || InMessage.Phase == EAnimGraphPhase::Evaluate)
+															});
+														}
+					
+														// Some traces come from both update and evaluate phases
+														if(InMessage.Phase == EAnimGraphPhase::Update || InMessage.Phase == EAnimGraphPhase::Evaluate)
+														{
+															AnimationProvider->ReadAnimAttributesTimeline(Id, [InGraphStartTime, InGraphEndTime, AnimationProvider, &DebugData](const IAnimationProvider::AnimAttributeTimeline& InAnimAttributeTimeline)
 															{
-																AnimationProvider->ReadAnimAttributesTimeline(Id, [InGraphStartTime, InGraphEndTime, AnimationProvider, &DebugData](const IAnimationProvider::AnimAttributeTimeline& InAnimAttributeTimeline)
+																TRACE_CPUPROFILER_EVENT_SCOPE(AnimGraphAttributes);
+																InAnimAttributeTimeline.EnumerateEvents(InGraphStartTime, InGraphEndTime, [AnimationProvider, &DebugData](double InStartTime, double InEndTime, uint32 InDepth, const FAnimAttributeMessage& InMessage)
 																{
-																	TRACE_CPUPROFILER_EVENT_SCOPE(AnimGraphAttributes);
-																	InAnimAttributeTimeline.EnumerateEvents(InGraphStartTime, InGraphEndTime, [AnimationProvider, &DebugData](double InStartTime, double InEndTime, uint32 InDepth, const FAnimAttributeMessage& InMessage)
+																	const TCHAR* AttributeName = AnimationProvider->GetName(InMessage.AttributeNameId);
+																	if(AttributeName)
 																	{
-																		const TCHAR* AttributeName = AnimationProvider->GetName(InMessage.AttributeNameId);
-																		if(AttributeName)
-																		{
-																			DebugData.RecordNodeAttribute(InMessage.TargetNodeId, InMessage.SourceNodeId, FName(AttributeName));
-																		}
-															
-																		return TraceServices::EEventEnumerate::Continue;
-																	});
+																		DebugData.RecordNodeAttribute(InMessage.TargetNodeId, InMessage.SourceNodeId, FName(AttributeName));
+																	}
+														
+																	return TraceServices::EEventEnumerate::Continue;
 																});
+															});
 
-																
-																AnimationProvider->ReadPoseWatchTimeline(Id, [InGraphStartTime, InGraphEndTime, AnimationProvider, &DebugData](const IAnimationProvider::PoseWatchTimeline& InPoseWatchTimeline)
-																	{
-																		InPoseWatchTimeline.EnumerateEvents(InGraphStartTime, InGraphEndTime, [AnimationProvider, &DebugData](double InStartTime, double InEndTime, uint32 InDepth, const FPoseWatchMessage& InMessage)
+															
+															AnimationProvider->ReadPoseWatchTimeline(Id, [InGraphStartTime, InGraphEndTime, AnimationProvider, &DebugData](const IAnimationProvider::PoseWatchTimeline& InPoseWatchTimeline)
+																{
+																	InPoseWatchTimeline.EnumerateEvents(InGraphStartTime, InGraphEndTime, [AnimationProvider, &DebugData](double InStartTime, double InEndTime, uint32 InDepth, const FPoseWatchMessage& InMessage)
+																		{
+																			for (FAnimNodePoseWatch& PoseWatch : DebugData.AnimNodePoseWatch)
 																			{
-																				for (FAnimNodePoseWatch& PoseWatch : DebugData.AnimNodePoseWatch)
+																				if (PoseWatch.NodeID == InMessage.PoseWatchId)
 																				{
-																					if (PoseWatch.NodeID == InMessage.PoseWatchId)
-																					{
-																						TArray<FBoneIndexType> RequiredBones;
-																						TArray<FTransform> BoneTransforms;
-																						AnimationProvider->GetPoseWatchData(InMessage, BoneTransforms, RequiredBones);
-
-																						PoseWatch.SetPose(RequiredBones, BoneTransforms);
-																						PoseWatch.SetWorldTransform(InMessage.WorldTransform);
-
-																						PoseWatch.PoseWatch->SetIsNodeEnabled(true);
-																						break;
-																					}
+																					TArray<FBoneIndexType> RequiredBones;
+																					TArray<FTransform> BoneTransforms;
+																					AnimationProvider->GetPoseWatchData(InMessage, BoneTransforms, RequiredBones);
+
+																					PoseWatch.SetPose(RequiredBones, BoneTransforms);
+																					PoseWatch.SetWorldTransform(InMessage.WorldTransform);
+
+																					PoseWatch.PoseWatch->SetIsNodeEnabled(true);
+																					break;
 																				}
-																				return TraceServices::EEventEnumerate::Continue;
-																			});
-																	});
+																			}
+																			return TraceServices::EEventEnumerate::Continue;
+																		});
+																});
 
-															}
-						
 														}
-														return TraceServices::EEventEnumerate::Continue;
-													});
-												}
+					
+													}
+													return TraceServices::EEventEnumerate::Continue;
+												});
 											}
 										}
 									}
 								}
-#endif // OBJECT_TRACE_ENABLED
 								return TraceServices::EEventEnumerate::Continue;
 							});
 						}
diff --git a/Engine/Plugins/Animation/GameplayInsights/Source/RewindDebugger/Private/RewindDebuggerAnimation.h b/Engine/Plugins/Animation/GameplayInsights/Source/RewindDebugger/Private/RewindDebuggerAnimation.h
--- a/Engine/Plugins/Animation/GameplayInsights/Source/RewindDebugger/Private/RewindDebuggerAnimation.h
+++ b/Engine/Plugins/Animation/GameplayInsights/Source/RewindDebugger/Private/RewindDebuggerAnimation.h
@@ -14,6 +14,8 @@ namespace TraceServices
 	struct FFrame;
 }
 
+class UAnimInstance;
+
 // Rewind debugger extension for animation support
 //  replay of animated pose data
 //  updating animation blueprint debugger
@@ -22,7 +24,6 @@ class FRewindDebuggerAnimation : public IRewindDebuggerExtension
 {
 public:
 
-
 	FRewindDebuggerAnimation();
 	virtual ~FRewindDebuggerAnimation() {};
 	void Initialize();
@@ -38,8 +39,10 @@ public:
 	void OnPIEStopped(bool bSimulating);
 	void OnPIESingleStepped(bool bSimulating);
 
+	UAnimInstance* GetDebugAnimInstance(uint64 ObjectId);
+
+	static FRewindDebuggerAnimation* GetInstance() { return Instance; }
 private:
-	
 	void ApplyPoseToMesh(const class IAnimationProvider* AnimationProvider, const IGameplayProvider* GameplayProvider, const TraceServices::FFrame& Frame,
 		const IAnimationProvider::SkeletalMeshPoseTimeline& TimelineData, USkeletalMeshComponent* MeshComponent, uint64 ObjectId, bool bQueueForReset, bool bApplyMesh);
 	
@@ -58,9 +61,23 @@ private:
 		// mesh
 		TObjectPtr<USkeletalMeshComponent> Component;
 	};
+	
+	struct FSpawnedAnimInstanceInfo
+	{
+		// AnimInstance id
+		uint64 id;
+		// Data used for anim BP debugging
+		TObjectPtr<UAnimInstance> AnimInstance;
+	};
+	
+	FSpawnedMeshComponentInfo* SpawnMesh(uint64 ObjectId, const IGameplayProvider* GameplayProvider);
+	UAnimInstance* SpawnAnimInstance(uint64 ObjectId, const IGameplayProvider* GameplayProvider);
 
 	TMap<uint64, FSpawnedMeshComponentInfo> SpawnedMeshComponents;
+	TMap<uint64, FSpawnedAnimInstanceInfo> SpawnedAnimInstances;
 
 	TMap<uint64, FMeshComponentResetData> MeshComponentsToReset;
 	double LastScrubTime = 0.0;
+
+	static FRewindDebuggerAnimation* Instance;
 };
