Fix editor release build

diff --git a/Engine/Plugins/Animation/GameplayInsights/Source/RewindDebugger/Private/RewindDebuggerAnimation.cpp b/Engine/Plugins/Animation/GameplayInsights/Source/RewindDebugger/Private/RewindDebuggerAnimation.cpp
--- a/Engine/Plugins/Animation/GameplayInsights/Source/RewindDebugger/Private/RewindDebuggerAnimation.cpp
+++ b/Engine/Plugins/Animation/GameplayInsights/Source/RewindDebugger/Private/RewindDebuggerAnimation.cpp
@@ -376,7 +376,6 @@ void FRewindDebuggerAnimation::Update(float DeltaTime, IRewindDebugger* RewindDe
 											{
 												TRACE_CPUPROFILER_EVENT_SCOPE(FRewindDebugger::Tick_UpdateBlueprintDebug);
 												// update debug info for attached Animation Blueprint editors
-												uint64 Id = FObjectTrace::GetObjectId(AnimInstance);
 												const int32 NodeCount = InstanceClass->GetAnimNodeProperties().Num();
 						
 												FAnimBlueprintDebugData& DebugData = InstanceClass->GetAnimBlueprintDebugData();
@@ -403,7 +402,7 @@ void FRewindDebuggerAnimation::Update(float DeltaTime, IRewindDebugger* RewindDe
 
 												DebugData.DisableAllPoseWatches();
 					
-												AnimGraphTimeline.EnumerateEvents(Frame.StartTime, Frame.EndTime, [Id, AnimationProvider, GameplayProvider, &DebugData, NodeCount](double InGraphStartTime, double InGraphEndTime, uint32 InDepth, const FAnimGraphMessage& InMessage)
+												AnimGraphTimeline.EnumerateEvents(Frame.StartTime, Frame.EndTime, [ObjectId, AnimationProvider, GameplayProvider, &DebugData, NodeCount](double InGraphStartTime, double InGraphEndTime, uint32 InDepth, const FAnimGraphMessage& InMessage)
 												{
 													TRACE_CPUPROFILER_EVENT_SCOPE(AnimGraphTimelineEvent);
 															
@@ -414,7 +413,7 @@ void FRewindDebuggerAnimation::Update(float DeltaTime, IRewindDebugger* RewindDe
 														// Check for an update phase (which contains weights)
 														if(InMessage.Phase == EAnimGraphPhase::Update)
 														{
-															AnimationProvider->ReadAnimNodesTimeline(Id, [InGraphStartTime, InGraphEndTime, &DebugData](const IAnimationProvider::AnimNodesTimeline& InNodesTimeline)
+															AnimationProvider->ReadAnimNodesTimeline(ObjectId, [InGraphStartTime, InGraphEndTime, &DebugData](const IAnimationProvider::AnimNodesTimeline& InNodesTimeline)
 															{
 																TRACE_CPUPROFILER_EVENT_SCOPE(AnimGraphDebugNodeVisits);
 																InNodesTimeline.EnumerateEvents(InGraphStartTime, InGraphEndTime, [&DebugData](double InStartTime, double InEndTime, uint32 InDepth, const FAnimNodeMessage& InMessage)
@@ -424,7 +423,7 @@ void FRewindDebuggerAnimation::Update(float DeltaTime, IRewindDebugger* RewindDe
 																});
 															});
 					
-															AnimationProvider->ReadStateMachinesTimeline(Id, [InGraphStartTime, InGraphEndTime, &DebugData](const IAnimationProvider::StateMachinesTimeline& InStateMachinesTimeline)
+															AnimationProvider->ReadStateMachinesTimeline(ObjectId, [InGraphStartTime, InGraphEndTime, &DebugData](const IAnimationProvider::StateMachinesTimeline& InStateMachinesTimeline)
 															{
 																TRACE_CPUPROFILER_EVENT_SCOPE(AnimGraphDebugStateMachine);
 																InStateMachinesTimeline.EnumerateEvents(InGraphStartTime, InGraphEndTime, [&DebugData](double InStartTime, double InEndTime, uint32 InDepth, const FAnimStateMachineMessage& InMessage)
@@ -434,7 +433,7 @@ void FRewindDebuggerAnimation::Update(float DeltaTime, IRewindDebugger* RewindDe
 																});
 															});
 					
-															AnimationProvider->ReadAnimSequencePlayersTimeline(Id, [InGraphStartTime, InGraphEndTime, GameplayProvider, &DebugData](const IAnimationProvider::AnimSequencePlayersTimeline& InSequencePlayersTimeline)
+															AnimationProvider->ReadAnimSequencePlayersTimeline(ObjectId, [InGraphStartTime, InGraphEndTime, GameplayProvider, &DebugData](const IAnimationProvider::AnimSequencePlayersTimeline& InSequencePlayersTimeline)
 															{
 																TRACE_CPUPROFILER_EVENT_SCOPE(AnimGraphDebugSequencePlayers);
 																InSequencePlayersTimeline.EnumerateEvents(InGraphStartTime, InGraphEndTime, [&DebugData](double InStartTime, double InEndTime, uint32 InDepth, const FAnimSequencePlayerMessage& InMessage)
@@ -444,7 +443,7 @@ void FRewindDebuggerAnimation::Update(float DeltaTime, IRewindDebugger* RewindDe
 																});
 															});
 					
-															AnimationProvider->ReadAnimBlendSpacePlayersTimeline(Id, [InGraphStartTime, InGraphEndTime, GameplayProvider, &DebugData](const IAnimationProvider::BlendSpacePlayersTimeline& InBlendSpacePlayersTimeline)
+															AnimationProvider->ReadAnimBlendSpacePlayersTimeline(ObjectId, [InGraphStartTime, InGraphEndTime, GameplayProvider, &DebugData](const IAnimationProvider::BlendSpacePlayersTimeline& InBlendSpacePlayersTimeline)
 															{
 																TRACE_CPUPROFILER_EVENT_SCOPE(AnimGraphBlendSpaces);
 																InBlendSpacePlayersTimeline.EnumerateEvents(InGraphStartTime, InGraphEndTime, [GameplayProvider, &DebugData](double InStartTime, double InEndTime, uint32 InDepth, const FBlendSpacePlayerMessage& InMessage)
@@ -461,7 +460,7 @@ void FRewindDebuggerAnimation::Update(float DeltaTime, IRewindDebugger* RewindDe
 																});
 															});
 					
-															AnimationProvider->ReadAnimSyncTimeline(Id, [InGraphStartTime, InGraphEndTime, AnimationProvider, &DebugData](const IAnimationProvider::AnimSyncTimeline& InAnimSyncTimeline)
+															AnimationProvider->ReadAnimSyncTimeline(ObjectId, [InGraphStartTime, InGraphEndTime, AnimationProvider, &DebugData](const IAnimationProvider::AnimSyncTimeline& InAnimSyncTimeline)
 															{
 																TRACE_CPUPROFILER_EVENT_SCOPE(AnimGraphAnimSync);
 																InAnimSyncTimeline.EnumerateEvents(InGraphStartTime, InGraphEndTime, [AnimationProvider, &DebugData](double InStartTime, double InEndTime, uint32 InDepth, const FAnimSyncMessage& InMessage)
@@ -480,7 +479,7 @@ void FRewindDebuggerAnimation::Update(float DeltaTime, IRewindDebugger* RewindDe
 														// Some traces come from both update and evaluate phases
 														if(InMessage.Phase == EAnimGraphPhase::Update || InMessage.Phase == EAnimGraphPhase::Evaluate)
 														{
-															AnimationProvider->ReadAnimAttributesTimeline(Id, [InGraphStartTime, InGraphEndTime, AnimationProvider, &DebugData](const IAnimationProvider::AnimAttributeTimeline& InAnimAttributeTimeline)
+															AnimationProvider->ReadAnimAttributesTimeline(ObjectId, [InGraphStartTime, InGraphEndTime, AnimationProvider, &DebugData](const IAnimationProvider::AnimAttributeTimeline& InAnimAttributeTimeline)
 															{
 																TRACE_CPUPROFILER_EVENT_SCOPE(AnimGraphAttributes);
 																InAnimAttributeTimeline.EnumerateEvents(InGraphStartTime, InGraphEndTime, [AnimationProvider, &DebugData](double InStartTime, double InEndTime, uint32 InDepth, const FAnimAttributeMessage& InMessage)
@@ -496,7 +495,7 @@ void FRewindDebuggerAnimation::Update(float DeltaTime, IRewindDebugger* RewindDe
 															});
 
 															
-															AnimationProvider->ReadPoseWatchTimeline(Id, [InGraphStartTime, InGraphEndTime, AnimationProvider, &DebugData](const IAnimationProvider::PoseWatchTimeline& InPoseWatchTimeline)
+															AnimationProvider->ReadPoseWatchTimeline(ObjectId, [InGraphStartTime, InGraphEndTime, AnimationProvider, &DebugData](const IAnimationProvider::PoseWatchTimeline& InPoseWatchTimeline)
 																{
 																	InPoseWatchTimeline.EnumerateEvents(InGraphStartTime, InGraphEndTime, [AnimationProvider, &DebugData](double InStartTime, double InEndTime, uint32 InDepth, const FPoseWatchMessage& InMessage)
 																		{
