Fix uninitialized variable

diff --git a/Engine/Plugins/Animation/GameplayInsights/Source/RewindDebugger/Private/AnimInstanceHelpers.cpp b/Engine/Plugins/Animation/GameplayInsights/Source/RewindDebugger/Private/AnimInstanceHelpers.cpp
--- a/Engine/Plugins/Animation/GameplayInsights/Source/RewindDebugger/Private/AnimInstanceHelpers.cpp
+++ b/Engine/Plugins/Animation/GameplayInsights/Source/RewindDebugger/Private/AnimInstanceHelpers.cpp
@@ -32,7 +32,7 @@ static bool OpenAnimBlueprintAndAttachDebugger(const TraceServices::IAnalysisSes
 					{
 						GEditor->GetEditorSubsystem<UAssetEditorSubsystem>()->OpenEditorForAsset(AnimBlueprint);
 						
-						UObject* SelectedInstance;
+						UObject* SelectedInstance = nullptr;
 #if OBJECT_TRACE_ENABLED
 						SelectedInstance = FObjectTrace::GetObjectFromId(ObjectId);
 #endif
