Fix crashes with rewind debugger when loading a new level or manually deleting a rewind debugger mesh.
Stop updating mesh poses continuously with loaded trace files.

diff --git a/Engine/Plugins/Animation/GameplayInsights/Source/RewindDebugger/Private/RewindDebugger.cpp b/Engine/Plugins/Animation/GameplayInsights/Source/RewindDebugger/Private/RewindDebugger.cpp
--- a/Engine/Plugins/Animation/GameplayInsights/Source/RewindDebugger/Private/RewindDebugger.cpp
+++ b/Engine/Plugins/Animation/GameplayInsights/Source/RewindDebugger/Private/RewindDebugger.cpp
@@ -221,10 +221,21 @@ void FRewindDebugger::OnPIESingleStepped(bool bSimulating)
 
 void FRewindDebugger::OnPIEStopped(bool bSimulating)
 {
+	if (IsRecording() && bPIESimulating)
+	{
+#if OBJECT_TRACE_ENABLED
+		UWorld* World = GetWorldToVisualize();
+		SetCurrentScrubTime(FObjectTrace::GetWorldElapsedTime(World));
+#endif // OBJECT_TRACE_ENABLED
+	}
+	
+	bTraceFileLoaded = IsRecording();
+	
 	bPIEStarted = false;
 	bPIESimulating = false;
 
 	StopRecording();
+	
 }
 
 bool FRewindDebugger::GetTargetActorPosition(FVector& OutPosition) const
@@ -501,8 +512,6 @@ void FRewindDebugger::ClearTrace()
 	
 	bTraceFileLoaded = false;
 	TargetObjectIds.Empty();
-	CurrentViewRange.SetLowerBoundValue(0);
-	CurrentViewRange.SetUpperBoundValue(0);
 	CurrentTraceRange.SetLowerBoundValue(0);
 	CurrentTraceRange.SetUpperBoundValue(0);
 	RecordingDuration.Set(0.0);
@@ -762,15 +771,13 @@ UWorld* FRewindDebugger::GetWorldToVisualize() const
 
 	UWorld* World = nullptr;
 
-#if WITH_EDITOR
 	UEditorEngine* EditorEngine = Cast<UEditorEngine>(GEngine);
 	if (GIsEditor && EditorEngine != nullptr && World == nullptr)
 	{
-		// lets use PlayWorld during PIE/Simulate and regular world from editor otherwise, to draw debug information
-		World = EditorEngine->PlayWorld != nullptr ? ToRawPtr(EditorEngine->PlayWorld) : EditorEngine->GetEditorWorldContext().World();
+		// when a trace file is loaded, use the editor world, otherwise use the play world
+		World = bTraceFileLoaded ?  EditorEngine->GetEditorWorldContext().World() : ToRawPtr(EditorEngine->PlayWorld);
 	}
 
-#endif
 	if (!GIsEditor && World == nullptr)
 	{
 		World = GEngine->GetWorld();
@@ -987,7 +994,7 @@ void FRewindDebugger::Tick(float DeltaTime)
 	{
 		const IAnimationProvider* AnimationProvider = Session->ReadProvider<IAnimationProvider>("AnimationProvider");
 		const IGameplayProvider* GameplayProvider = Session->ReadProvider<IGameplayProvider>("GameplayProvider");
-
+		
 		if (AnimationProvider && GameplayProvider)
 		{
 			TraceServices::FAnalysisSessionReadScope SessionReadScope(*Session);
@@ -995,9 +1002,8 @@ void FRewindDebugger::Tick(float DeltaTime)
 			double RecordingDurationValue = GameplayProvider->GetRecordingDuration();
 			if (IsTraceFileLoaded() && RecordingDurationValue > RecordingDuration.Get())
 			{
-				// while trace file is loading up, recording duration changes - autoscroll
-				SetCurrentScrubTime(RecordingDurationValue);
-				TrackCursorDelegate.ExecuteIfBound(false);	
+				// while trace file is loading up, force the trace range to update.
+				SetCurrentViewRange(GetCurrentViewRange());
 			}
 			RecordingDuration.Set(RecordingDurationValue);
 			
diff --git a/Engine/Plugins/Animation/GameplayInsights/Source/RewindDebugger/Private/RewindDebugger.h b/Engine/Plugins/Animation/GameplayInsights/Source/RewindDebugger/Private/RewindDebugger.h
--- a/Engine/Plugins/Animation/GameplayInsights/Source/RewindDebugger/Private/RewindDebugger.h
+++ b/Engine/Plugins/Animation/GameplayInsights/Source/RewindDebugger/Private/RewindDebugger.h
@@ -178,7 +178,7 @@ private:
 
 	double PreviousTraceTime = -1;
 	double CurrentScrubTime = 0;
-	TRange<double> CurrentViewRange {0, 0};
+	TRange<double> CurrentViewRange {0, 10};
 	TRange<double> CurrentTraceRange {0, 0};
 	uint16 RecordingIndex = 0;
 
diff --git a/Engine/Plugins/Animation/GameplayInsights/Source/RewindDebugger/Private/RewindDebuggerAnimation.cpp b/Engine/Plugins/Animation/GameplayInsights/Source/RewindDebugger/Private/RewindDebuggerAnimation.cpp
--- a/Engine/Plugins/Animation/GameplayInsights/Source/RewindDebugger/Private/RewindDebuggerAnimation.cpp
+++ b/Engine/Plugins/Animation/GameplayInsights/Source/RewindDebugger/Private/RewindDebuggerAnimation.cpp
@@ -42,15 +42,15 @@ void FRewindDebuggerAnimation::ClearSpawnedComponents()
 {
 	for(auto& MeshComponentInfo : SpawnedMeshComponents)
 	{
-		if(MeshComponentInfo.Value.Component)
+		if(MeshComponentInfo.Value.Actor.IsValid())
 		{
-			MeshComponentInfo.Value.Component->UnregisterComponent();
-			MeshComponentInfo.Value.Component->MarkAsGarbage();
-			MeshComponentInfo.Value.Component = nullptr;
-		}
+			if(MeshComponentInfo.Value.Component)
+			{
+				MeshComponentInfo.Value.Component->UnregisterComponent();
+				MeshComponentInfo.Value.Component->MarkAsGarbage();
+				MeshComponentInfo.Value.Component = nullptr;
+			}
 
-		if(MeshComponentInfo.Value.Actor)
-		{
 			MeshComponentInfo.Value.Actor->Destroy();
 			MeshComponentInfo.Value.Actor = nullptr;
 		}
@@ -62,6 +62,7 @@ void FRewindDebuggerAnimation::ClearSpawnedComponents()
 void FRewindDebuggerAnimation::Clear(IRewindDebugger* RewindDebugger)
 {
 	ClearSpawnedComponents();
+	LastScrubTime = -1;
 }
 
 void FRewindDebuggerAnimation::Initialize()
@@ -114,6 +115,9 @@ void FRewindDebuggerAnimation::OnPIESingleStepped(bool bSimulating)
 void FRewindDebuggerAnimation::OnPIEStopped(bool bSimulating)
 {
 	MeshComponentsToReset.Empty();
+	
+	// clear last scrub time so that poses will reapply
+	LastScrubTime = -1;
 }
 
 
@@ -136,9 +140,6 @@ void FRewindDebuggerAnimation::ApplyPoseToMesh(const IAnimationProvider* Animati
 		FTransform ComponentWorldTransform;
 		if (const FSkeletalMeshInfo* SkeletalMeshInfo = AnimationProvider->FindSkeletalMeshInfo(PoseMessage->MeshId))
 		{
-			AnimationProvider->GetSkeletalMeshComponentSpacePose(*PoseMessage, *SkeletalMeshInfo, ComponentWorldTransform, MeshComponent->GetEditableComponentSpaceTransforms());
-			MeshComponent->ApplyEditedComponentSpaceTransforms();
-
 			if (bApplyMesh)
 			{
 				const FObjectInfo* SkeletalMeshObjectInfo = GameplayProvider->FindObjectInfo(PoseMessage->MeshId);
@@ -150,6 +151,9 @@ void FRewindDebuggerAnimation::ApplyPoseToMesh(const IAnimationProvider* Animati
 				}
 			}
 			
+			AnimationProvider->GetSkeletalMeshComponentSpacePose(*PoseMessage, *SkeletalMeshInfo, ComponentWorldTransform, MeshComponent->GetEditableComponentSpaceTransforms());
+			MeshComponent->ApplyEditedComponentSpaceTransforms();
+			
 			if (bQueueForReset)
 			{
 				if (MeshComponentsToReset.Find(ObjectId) == nullptr)
@@ -204,7 +208,6 @@ FRewindDebuggerAnimation::FSpawnedMeshComponentInfo* FRewindDebuggerAnimation::S
 	ActorSpawnParameters.bHideFromSceneOutliner = true;
 	ActorSpawnParameters.ObjectFlags |= RF_Transient;
 
-
 	MeshComponentInfo->Actor = World->SpawnActor<AActor>(ActorSpawnParameters);
 
 	if(const FObjectInfo* ActorInfo = FindOwningActorInfo(GameplayProvider, ObjectId))
@@ -212,7 +215,7 @@ FRewindDebuggerAnimation::FSpawnedMeshComponentInfo* FRewindDebuggerAnimation::S
 		MeshComponentInfo->Actor->SetActorLabel(FString(TEXT("RewindDebugger: ") + FString(ActorInfo->Name)));
 	}
 
-	MeshComponentInfo->Component = NewObject<USkeletalMeshComponent>(MeshComponentInfo->Actor);
+	MeshComponentInfo->Component = NewObject<USkeletalMeshComponent>(MeshComponentInfo->Actor.Get());
 	MeshComponentInfo->Component->PrimaryComponentTick.bStartWithTickEnabled = false;
 	MeshComponentInfo->Component->PrimaryComponentTick.bCanEverTick = false;
 
@@ -226,12 +229,19 @@ FRewindDebuggerAnimation::FSpawnedMeshComponentInfo* FRewindDebuggerAnimation::S
 
 UAnimInstance* FRewindDebuggerAnimation::SpawnAnimInstance(uint64 ObjectId, const IGameplayProvider* GameplayProvider)
 {
-	if (FSpawnedAnimInstanceInfo* FoundInfo = SpawnedAnimInstances.Find(ObjectId))
+	FSpawnedAnimInstanceInfo* AnimInstanceInfo = SpawnedAnimInstances.Find(ObjectId);
+	if (AnimInstanceInfo)
 	{
-		return FoundInfo->AnimInstance;
+		if (AnimInstanceInfo->AnimInstance.IsValid())
+		{
+			return AnimInstanceInfo->AnimInstance.Get();
+		}
+	}
+
+	if (AnimInstanceInfo == nullptr)
+	{
+		AnimInstanceInfo = &SpawnedAnimInstances.Add(ObjectId);
 	}
-	
-	FSpawnedAnimInstanceInfo* AnimInstanceInfo = &SpawnedAnimInstances.Add(ObjectId);
 
 	if (const FObjectInfo* ObjectInfo = GameplayProvider->FindObjectInfo(ObjectId))
 	{
@@ -243,7 +253,7 @@ UAnimInstance* FRewindDebuggerAnimation::SpawnAnimInstance(uint64 ObjectId, cons
 				if (const FSpawnedMeshComponentInfo* MeshInfo = SpawnedMeshComponents.Find(ObjectInfo->OuterId))
 				{
 					AnimInstanceInfo->AnimInstance = NewObject<UAnimInstance>(MeshInfo->Component, ClassPtr.Get());
-					return AnimInstanceInfo->AnimInstance;
+					return AnimInstanceInfo->AnimInstance.Get();
 				}
 			}
 		}
@@ -257,7 +267,7 @@ UAnimInstance* FRewindDebuggerAnimation::GetDebugAnimInstance(uint64 ObjectId)
 {
 	if (FSpawnedAnimInstanceInfo* FoundInfo = SpawnedAnimInstances.Find(ObjectId))
 	{
-		return FoundInfo->AnimInstance;
+		return FoundInfo->AnimInstance.Get();
 	}
 	return nullptr;
 }
@@ -275,10 +285,11 @@ void FRewindDebuggerAnimation::Update(float DeltaTime, IRewindDebugger* RewindDe
 	if (const TraceServices::IAnalysisSession* Session = RewindDebugger->GetAnalysisSession())
 	{
 		TraceServices::FAnalysisSessionReadScope SessionReadScope(*Session);
-		double CurrentTraceTime = RewindDebugger->CurrentTraceTime();
 		
+		double CurrentTraceTime = RewindDebugger->CurrentTraceTime();
 		if (CurrentTraceTime != LastScrubTime)
 		{
+			LastScrubTime = CurrentTraceTime;
 			const TraceServices::IFrameProvider& FrameProvider = TraceServices::ReadFrameProvider(*Session);
 			TraceServices::FFrame Frame;
 			if (FrameProvider.GetFrameFromTime(ETraceFrameType::TraceFrameType_Game, CurrentTraceTime, Frame))
@@ -319,10 +330,21 @@ void FRewindDebuggerAnimation::Update(float DeltaTime, IRewindDebugger* RewindDe
 							{
 								// display pose on a spawned mesh component
 								FSpawnedMeshComponentInfo* MeshComponentInfo = SpawnedMeshComponents.Find(ObjectId);
+								if (MeshComponentInfo)
+								{
+									if (!MeshComponentInfo->Actor.IsValid())
+									{
+										// if the actor has been deleted, clear cached data and create it again
+										SpawnedMeshComponents.Remove(ObjectId);
+										MeshComponentInfo = nullptr;
+									}
+								}
+									
 								if (MeshComponentInfo == nullptr)
 								{
 									MeshComponentInfo = SpawnMesh(ObjectId, GameplayProvider);
 									bLoadMesh = true;
+									LastScrubTime = -1;
 								}
 								
 								MeshComponent = MeshComponentInfo->Component;
diff --git a/Engine/Plugins/Animation/GameplayInsights/Source/RewindDebugger/Private/RewindDebuggerAnimation.h b/Engine/Plugins/Animation/GameplayInsights/Source/RewindDebugger/Private/RewindDebuggerAnimation.h
--- a/Engine/Plugins/Animation/GameplayInsights/Source/RewindDebugger/Private/RewindDebuggerAnimation.h
+++ b/Engine/Plugins/Animation/GameplayInsights/Source/RewindDebugger/Private/RewindDebuggerAnimation.h
@@ -57,7 +57,7 @@ private:
 		// mesh component object id
 		uint64 id;
 		// actor to hold the mesh component
-		TObjectPtr<AActor> Actor;
+		TWeakObjectPtr<AActor> Actor;
 		// mesh
 		TObjectPtr<USkeletalMeshComponent> Component;
 	};
@@ -67,7 +67,7 @@ private:
 		// AnimInstance id
 		uint64 id;
 		// Data used for anim BP debugging
-		TObjectPtr<UAnimInstance> AnimInstance;
+		TWeakObjectPtr<UAnimInstance> AnimInstance;
 	};
 	
 	FSpawnedMeshComponentInfo* SpawnMesh(uint64 ObjectId, const IGameplayProvider* GameplayProvider);
@@ -77,7 +77,7 @@ private:
 	TMap<uint64, FSpawnedAnimInstanceInfo> SpawnedAnimInstances;
 
 	TMap<uint64, FMeshComponentResetData> MeshComponentsToReset;
-	double LastScrubTime = 0.0;
+	double LastScrubTime = -1;
 
 	static FRewindDebuggerAnimation* Instance;
 };
