Apply skeleton preview mesh if traced skeletal mesh can't be found

diff --git a/Engine/Plugins/Animation/GameplayInsights/Source/GameplayInsights/Private/AnimationAnalyzer.cpp b/Engine/Plugins/Animation/GameplayInsights/Source/GameplayInsights/Private/AnimationAnalyzer.cpp
--- a/Engine/Plugins/Animation/GameplayInsights/Source/GameplayInsights/Private/AnimationAnalyzer.cpp
+++ b/Engine/Plugins/Animation/GameplayInsights/Source/GameplayInsights/Private/AnimationAnalyzer.cpp
@@ -22,6 +22,7 @@ void FAnimationAnalyzer::OnAnalysisBegin(const FOnAnalysisContext& Context)
 	Builder.RouteEvent(RouteId_TickRecord2, "Animation", "TickRecord2");
 	Builder.RouteEvent(RouteId_SkeletalMesh, "Animation", "SkeletalMesh");
 	Builder.RouteEvent(RouteId_SkeletalMesh2, "Animation", "SkeletalMesh2");
+	Builder.RouteEvent(RouteId_SkeletalMesh3, "Animation", "SkeletalMesh3");
 	Builder.RouteEvent(RouteId_SkeletalMeshComponent, "Animation", "SkeletalMeshComponent");
 	Builder.RouteEvent(RouteId_SkeletalMeshComponent2, "Animation", "SkeletalMeshComponent2");
 	Builder.RouteEvent(RouteId_SkeletalMeshComponent3, "Animation", "SkeletalMeshComponent3");
@@ -113,7 +114,8 @@ bool FAnimationAnalyzer::OnEvent(uint16 RouteId, EStyle Style, const FOnEventCon
 		uint64 Id = EventData.GetValue<uint64>("Id");
 		uint32 BoneCount = EventData.GetValue<uint32>("BoneCount");
 		TArrayView<const int32> ParentIndices(reinterpret_cast<const int32*>(EventData.GetAttachment()), BoneCount);
-		AnimationProvider.AppendSkeletalMesh(Id, ParentIndices);
+		uint64 SkeletonId = 0;
+		AnimationProvider.AppendSkeletalMesh(Id, SkeletonId, ParentIndices);
 		break;
 	}
 	case RouteId_SkeletalMesh2:
@@ -121,7 +123,17 @@ bool FAnimationAnalyzer::OnEvent(uint16 RouteId, EStyle Style, const FOnEventCon
 		uint64 Id = EventData.GetValue<uint64>("Id");
 		uint32 BoneCount = EventData.GetValue<uint32>("BoneCount");
 		TArrayView<const int32> ParentIndices = EventData.GetArrayView<int32>("ParentIndices");
-		AnimationProvider.AppendSkeletalMesh(Id, ParentIndices);
+		uint64 SkeletonId = 0;
+		AnimationProvider.AppendSkeletalMesh(Id, SkeletonId, ParentIndices);
+		break;
+	}
+	case RouteId_SkeletalMesh3:
+	{
+		uint64 Id = EventData.GetValue<uint64>("Id");
+		uint32 BoneCount = EventData.GetValue<uint32>("BoneCount");
+		TArrayView<const int32> ParentIndices = EventData.GetArrayView<int32>("ParentIndices");
+		uint64 SkeletonId = EventData.GetValue<uint64>("SkeletonId");
+		AnimationProvider.AppendSkeletalMesh(Id, SkeletonId, ParentIndices);
 		break;
 	}
 	case RouteId_SkeletalMeshComponent:
diff --git a/Engine/Plugins/Animation/GameplayInsights/Source/GameplayInsights/Private/AnimationAnalyzer.h b/Engine/Plugins/Animation/GameplayInsights/Source/GameplayInsights/Private/AnimationAnalyzer.h
--- a/Engine/Plugins/Animation/GameplayInsights/Source/GameplayInsights/Private/AnimationAnalyzer.h
+++ b/Engine/Plugins/Animation/GameplayInsights/Source/GameplayInsights/Private/AnimationAnalyzer.h
@@ -24,6 +24,7 @@ private:
 		RouteId_TickRecord2,
 		RouteId_SkeletalMesh,
 		RouteId_SkeletalMesh2,
+		RouteId_SkeletalMesh3,
 		RouteId_SkeletalMeshComponent,
 		RouteId_SkeletalMeshComponent2,
 		RouteId_SkeletalMeshComponent3,
diff --git a/Engine/Plugins/Animation/GameplayInsights/Source/GameplayInsights/Private/AnimationProvider.cpp b/Engine/Plugins/Animation/GameplayInsights/Source/GameplayInsights/Private/AnimationProvider.cpp
--- a/Engine/Plugins/Animation/GameplayInsights/Source/GameplayInsights/Private/AnimationProvider.cpp
+++ b/Engine/Plugins/Animation/GameplayInsights/Source/GameplayInsights/Private/AnimationProvider.cpp
@@ -636,7 +636,7 @@ void FAnimationProvider::AppendTickRecord(uint64 InAnimInstanceId, double InProf
 	Session.UpdateDurationSeconds(InProfileTime);
 }
 
-void FAnimationProvider::AppendSkeletalMesh(uint64 InObjectId, const TArrayView<const int32>& InParentIndices)
+void FAnimationProvider::AppendSkeletalMesh(uint64 InObjectId, uint64 SkeletonId, const TArrayView<const int32>& InParentIndices)
 {
 	Session.WriteAccessCheck();
 
@@ -648,6 +648,7 @@ void FAnimationProvider::AppendSkeletalMesh(uint64 InObjectId, const TArrayView<
 		NewSkeletalMeshInfo.Id = InObjectId;
 		NewSkeletalMeshInfo.BoneCount = (uint32)InParentIndices.Num();
 		NewSkeletalMeshInfo.ParentIndicesStartIndex = SkeletalMeshParentIndices.Num();
+		NewSkeletalMeshInfo.SkeletonId = SkeletonId;
 
 		for(const int32& ParentIndex : InParentIndices)
 		{
diff --git a/Engine/Plugins/Animation/GameplayInsights/Source/GameplayInsights/Private/AnimationProvider.h b/Engine/Plugins/Animation/GameplayInsights/Source/GameplayInsights/Private/AnimationProvider.h
--- a/Engine/Plugins/Animation/GameplayInsights/Source/GameplayInsights/Private/AnimationProvider.h
+++ b/Engine/Plugins/Animation/GameplayInsights/Source/GameplayInsights/Private/AnimationProvider.h
@@ -65,7 +65,7 @@ public:
 	void AppendTickRecord(uint64 InAnimInstanceId, double InProfileTime, double InRecordingTime, uint64 InAssetId, int32 InNodeId, float InBlendWeight, float InPlaybackTime, float InRootMotionWeight, float InPlayRate, float InBlendSpacePositionX, float InBlendSpacePositionY, float InBlendSpaceFilteredPositionX, float InBlendSpaceFilteredPositionY, uint16 InFrameCounter, bool bInLooping, bool bInIsBlendSpace);
 
 	/** Add a skeletal mesh */
-	void AppendSkeletalMesh(uint64 InObjectId, const TArrayView<const int32>& ParentIndices);
+	void AppendSkeletalMesh(uint64 InObjectId, uint64 SkeletonId, const TArrayView<const int32>& ParentIndices);
 
 	/** Add a skeletal mesh pose/curves etc. */
 	void AppendSkeletalMeshComponent(uint64 InObjectId, uint64 InMeshId, double InProfileTime, double InRecordingTime, uint16 InLodIndex, uint16 InFrameCounter, const TArrayView<const FTransform>& InPose, const TArrayView<const FSkeletalMeshNamedCurve>& InCurves);
diff --git a/Engine/Plugins/Animation/GameplayInsights/Source/GameplayInsights/Public/IAnimationProvider.h b/Engine/Plugins/Animation/GameplayInsights/Source/GameplayInsights/Public/IAnimationProvider.h
--- a/Engine/Plugins/Animation/GameplayInsights/Source/GameplayInsights/Public/IAnimationProvider.h
+++ b/Engine/Plugins/Animation/GameplayInsights/Source/GameplayInsights/Public/IAnimationProvider.h
@@ -12,6 +12,7 @@ struct FSkeletalMeshInfo
 {
 	uint64 ParentIndicesStartIndex = 0;
 	uint64 Id = 0;
+	uint64 SkeletonId = 0;
 	uint32 BoneCount = 0;
 };
 
diff --git a/Engine/Plugins/Animation/GameplayInsights/Source/RewindDebugger/Private/RewindDebuggerAnimation.cpp b/Engine/Plugins/Animation/GameplayInsights/Source/RewindDebugger/Private/RewindDebuggerAnimation.cpp
--- a/Engine/Plugins/Animation/GameplayInsights/Source/RewindDebugger/Private/RewindDebuggerAnimation.cpp
+++ b/Engine/Plugins/Animation/GameplayInsights/Source/RewindDebugger/Private/RewindDebuggerAnimation.cpp
@@ -145,6 +145,21 @@ void FRewindDebuggerAnimation::ApplyPoseToMesh(const IAnimationProvider* Animati
 				const FObjectInfo* SkeletalMeshObjectInfo = GameplayProvider->FindObjectInfo(PoseMessage->MeshId);
 				USkeletalMesh* SkeletalMesh = TSoftObjectPtr<USkeletalMesh>(FSoftObjectPath(SkeletalMeshObjectInfo->PathName)).LoadSynchronous();
 										
+				if (SkeletalMesh == nullptr)
+				{
+					// if the skeletal mesh asset was not found, try the skeleton asset preview mesh as a fallback
+					if (SkeletalMeshInfo->SkeletonId != 0)
+					{
+						const FObjectInfo& SkeletonInfo = GameplayProvider->GetObjectInfo(SkeletalMeshInfo->SkeletonId);
+						
+						USkeleton* Skeleton = TSoftObjectPtr<USkeleton>(FSoftObjectPath(SkeletonInfo.PathName)).LoadSynchronous();
+						if (Skeleton)
+						{
+							SkeletalMesh = Skeleton->GetPreviewMesh(true);
+						}
+					}
+				}
+				
 				if(SkeletalMesh)
 				{
 					MeshComponent->SetSkeletalMesh(SkeletalMesh);
@@ -201,8 +216,7 @@ FRewindDebuggerAnimation::FSpawnedMeshComponentInfo* FRewindDebuggerAnimation::S
 {
 	FSpawnedMeshComponentInfo* MeshComponentInfo = &SpawnedMeshComponents.Add(ObjectId);
 
-	UEditorEngine* EditorEngine = Cast<UEditorEngine>(GEngine);
-	UWorld* World = EditorEngine->GetEditorWorldContext().World();
+	UWorld* World = IRewindDebugger::Instance()->GetWorldToVisualize();
 
 	FActorSpawnParameters ActorSpawnParameters;
 	ActorSpawnParameters.bHideFromSceneOutliner = true;
diff --git a/Engine/Source/Runtime/Engine/Private/Animation/AnimTrace.cpp b/Engine/Source/Runtime/Engine/Private/Animation/AnimTrace.cpp
--- a/Engine/Source/Runtime/Engine/Private/Animation/AnimTrace.cpp
+++ b/Engine/Source/Runtime/Engine/Private/Animation/AnimTrace.cpp
@@ -38,8 +38,9 @@ UE_TRACE_EVENT_BEGIN(Animation, TickRecord2)
 	UE_TRACE_EVENT_FIELD(bool, IsBlendSpace)
 UE_TRACE_EVENT_END()
 
-UE_TRACE_EVENT_BEGIN(Animation, SkeletalMesh2, NoSync|Important)
+UE_TRACE_EVENT_BEGIN(Animation, SkeletalMesh3, NoSync|Important)
 	UE_TRACE_EVENT_FIELD(uint64, Id)
+	UE_TRACE_EVENT_FIELD(uint64, SkeletonId)
 	UE_TRACE_EVENT_FIELD(int32[], ParentIndices)
 UE_TRACE_EVENT_END()
 
@@ -522,6 +523,7 @@ void FAnimTrace::OutputSkeletalMesh(const USkeletalMesh* InMesh)
 	}
 
 	TRACE_OBJECT(InMesh);
+	TRACE_OBJECT(InMesh->GetSkeleton());
 
 	uint32 BoneCount = (uint32)InMesh->GetRefSkeleton().GetNum();
 
@@ -535,9 +537,10 @@ void FAnimTrace::OutputSkeletalMesh(const USkeletalMesh* InMesh)
 		ParentIndices[BoneIndex++] = BoneInfo.ParentIndex;
 	}
 
-	UE_TRACE_LOG(Animation, SkeletalMesh2, AnimationChannel, ParentIndices.Num() * sizeof(int32))
-		<< SkeletalMesh2.Id(FObjectTrace::GetObjectId(InMesh))
-		<< SkeletalMesh2.ParentIndices(ParentIndices.GetData(), ParentIndices.Num());
+	UE_TRACE_LOG(Animation, SkeletalMesh3, AnimationChannel, ParentIndices.Num() * sizeof(int32))
+		<< SkeletalMesh3.Id(FObjectTrace::GetObjectId(InMesh))
+		<< SkeletalMesh3.SkeletonId(FObjectTrace::GetObjectId(InMesh->GetSkeleton()))
+		<< SkeletalMesh3.ParentIndices(ParentIndices.GetData(), ParentIndices.Num());
 
 	GSkeletalMeshTraceAnnotations.Set(InMesh);
 }
