Fix bugs with rewind debugger and mesh lod changes

diff --git a/Engine/Plugins/Animation/GameplayInsights/Source/RewindDebugger/Private/RewindDebuggerAnimation.cpp b/Engine/Plugins/Animation/GameplayInsights/Source/RewindDebugger/Private/RewindDebuggerAnimation.cpp
--- a/Engine/Plugins/Animation/GameplayInsights/Source/RewindDebugger/Private/RewindDebuggerAnimation.cpp
+++ b/Engine/Plugins/Animation/GameplayInsights/Source/RewindDebugger/Private/RewindDebuggerAnimation.cpp
@@ -91,6 +91,7 @@ void FRewindDebuggerAnimation::OnPIEResumed(bool bSimulating)
 		if (USkeletalMeshComponent* MeshComponent = MeshData.Value.Component.Get())
 		{
 			MeshComponent->SetRelativeTransform(MeshData.Value.RelativeTransform, false, nullptr, ETeleportType::TeleportPhysics);
+			MeshComponent->SetForcedLOD(MeshData.Value.ForcedLod);
 		}
 	}
 
@@ -176,12 +177,14 @@ void FRewindDebuggerAnimation::ApplyPoseToMesh(const IAnimationProvider* Animati
 					FMeshComponentResetData ResetData;
 					ResetData.Component = MeshComponent;
 					ResetData.RelativeTransform = MeshComponent->GetRelativeTransform();
+					ResetData.ForcedLod = MeshComponent->GetForcedLOD();
 					MeshComponentsToReset.Add(ObjectId, ResetData);
 				}
 			}
 
 			MeshComponent->SetWorldTransform(ComponentWorldTransform, false, nullptr, ETeleportType::TeleportPhysics);
 			MeshComponent->SetForcedLOD(PoseMessage->LodIndex + 1);
+			MeshComponent->UpdateLODStatus();
 			MeshComponent->UpdateChildTransforms(EUpdateTransformFlags::None, ETeleportType::TeleportPhysics);
 		}
 	}
@@ -303,7 +306,6 @@ void FRewindDebuggerAnimation::Update(float DeltaTime, IRewindDebugger* RewindDe
 		double CurrentTraceTime = RewindDebugger->CurrentTraceTime();
 		if (CurrentTraceTime != LastScrubTime)
 		{
-			LastScrubTime = CurrentTraceTime;
 			const TraceServices::IFrameProvider& FrameProvider = TraceServices::ReadFrameProvider(*Session);
 			TraceServices::FFrame Frame;
 			if (FrameProvider.GetFrameFromTime(ETraceFrameType::TraceFrameType_Game, CurrentTraceTime, Frame))
@@ -370,205 +372,206 @@ void FRewindDebuggerAnimation::Update(float DeltaTime, IRewindDebugger* RewindDe
 							}
 						}
 						);
-						}
+					}
         
+					{
+						TRACE_CPUPROFILER_EVENT_SCOPE(FRewindDebugger::Tick_AnimBlueprintsDebug);
+						// Apply Animation Blueprint Debugging Data:
+						// - enumerate over all anim graph timelines
+						// - check if their instance class still exists and is the debugging target for the Animation Blueprint Editor
+						// - if it is copy that debug data into the class debug data for the blueprint debugger
+						AnimationProvider->EnumerateAnimGraphTimelines([this, &Frame, AnimationProvider, GameplayProvider](uint64 ObjectId, const IAnimationProvider::AnimGraphTimeline& AnimGraphTimeline)
 						{
-							TRACE_CPUPROFILER_EVENT_SCOPE(FRewindDebugger::Tick_AnimBlueprintsDebug);
-							// Apply Animation Blueprint Debugging Data:
-							// - enumerate over all anim graph timelines
-							// - check if their instance class still exists and is the debugging target for the Animation Blueprint Editor
-							// - if it is copy that debug data into the class debug data for the blueprint debugger
-							AnimationProvider->EnumerateAnimGraphTimelines([this, &Frame, AnimationProvider, GameplayProvider](uint64 ObjectId, const IAnimationProvider::AnimGraphTimeline& AnimGraphTimeline)
-							{
-								UAnimInstance* AnimInstance = nullptr;
+							UAnimInstance* AnimInstance = nullptr;
 #if OBJECT_TRACE_ENABLED
-								if(UObject* ObjectInstance = FObjectTrace::GetObjectFromId(ObjectId))
-								{
-									AnimInstance = Cast<UAnimInstance>(ObjectInstance);
-								}
+							if(UObject* ObjectInstance = FObjectTrace::GetObjectFromId(ObjectId))
+							{
+								AnimInstance = Cast<UAnimInstance>(ObjectInstance);
+							}
 #endif
-								if (AnimInstance == nullptr)
-								{
-									AnimInstance = SpawnAnimInstance(ObjectId, GameplayProvider);
-								}
-								
-								if (AnimInstance)
+							if (AnimInstance == nullptr)
+							{
+								AnimInstance = SpawnAnimInstance(ObjectId, GameplayProvider);
+							}
+							
+							if (AnimInstance)
+							{
+								if(UAnimBlueprintGeneratedClass* InstanceClass = Cast<UAnimBlueprintGeneratedClass>(AnimInstance->GetClass()))
 								{
-									if(UAnimBlueprintGeneratedClass* InstanceClass = Cast<UAnimBlueprintGeneratedClass>(AnimInstance->GetClass()))
+									if(UAnimBlueprint* AnimBlueprint = Cast<UAnimBlueprint>(InstanceClass->ClassGeneratedBy))
 									{
-										if(UAnimBlueprint* AnimBlueprint = Cast<UAnimBlueprint>(InstanceClass->ClassGeneratedBy))
+										// for child Animation Blueprints, we actually want to debug the root blueprint (since the child doesn't contain any anim graphs)
+										if (UAnimBlueprint* RootAnimBP = UAnimBlueprint::FindRootAnimBlueprint(AnimBlueprint))
 										{
-											// for child Animation Blueprints, we actually want to debug the root blueprint (since the child doesn't contain any anim graphs)
-											if (UAnimBlueprint* RootAnimBP = UAnimBlueprint::FindRootAnimBlueprint(AnimBlueprint))
+											if (UAnimBlueprintGeneratedClass* RootInstanceClass = Cast<UAnimBlueprintGeneratedClass>(RootAnimBP->GeneratedClass))
 											{
-												if (UAnimBlueprintGeneratedClass* RootInstanceClass = Cast<UAnimBlueprintGeneratedClass>(RootAnimBP->GeneratedClass))
-												{
-													AnimBlueprint = RootAnimBP;
-													InstanceClass = RootInstanceClass;
-												}
+												AnimBlueprint = RootAnimBP;
+												InstanceClass = RootInstanceClass;
 											}
+										}
 
-											if(AnimBlueprint->IsObjectBeingDebugged(AnimInstance))
+										if(AnimBlueprint->IsObjectBeingDebugged(AnimInstance))
+										{
+											TRACE_CPUPROFILER_EVENT_SCOPE(FRewindDebugger::Tick_UpdateBlueprintDebug);
+											// update debug info for attached Animation Blueprint editors
+											const int32 NodeCount = InstanceClass->GetAnimNodeProperties().Num();
+					
+											FAnimBlueprintDebugData& DebugData = InstanceClass->GetAnimBlueprintDebugData();
 											{
-												TRACE_CPUPROFILER_EVENT_SCOPE(FRewindDebugger::Tick_UpdateBlueprintDebug);
-												// update debug info for attached Animation Blueprint editors
-												const int32 NodeCount = InstanceClass->GetAnimNodeProperties().Num();
-						
-												FAnimBlueprintDebugData& DebugData = InstanceClass->GetAnimBlueprintDebugData();
-												{
-													TRACE_CPUPROFILER_EVENT_SCOPE(ResetNodeVisitStates);
-													DebugData.ResetNodeVisitSites();
-												}
-												
-												// Anim node values can come from all phases
-												AnimationProvider->ReadAnimNodeValuesTimeline(ObjectId, [&Frame,AnimationProvider, &DebugData](const IAnimationProvider::AnimNodeValuesTimeline& InNodeValuesTimeline)
+												TRACE_CPUPROFILER_EVENT_SCOPE(ResetNodeVisitStates);
+												DebugData.ResetNodeVisitSites();
+											}
+											
+											// Anim node values can come from all phases
+											AnimationProvider->ReadAnimNodeValuesTimeline(ObjectId, [&Frame,AnimationProvider, &DebugData](const IAnimationProvider::AnimNodeValuesTimeline& InNodeValuesTimeline)
+											{
+												TRACE_CPUPROFILER_EVENT_SCOPE(AnimGraphNodeValues);
+												InNodeValuesTimeline.EnumerateEvents(Frame.StartTime, Frame.EndTime, [AnimationProvider, &DebugData](double InStartTime, double InEndTime, uint32 InDepth, const FAnimNodeValueMessage& InMessage)
 												{
-													TRACE_CPUPROFILER_EVENT_SCOPE(AnimGraphNodeValues);
-													InNodeValuesTimeline.EnumerateEvents(Frame.StartTime, Frame.EndTime, [AnimationProvider, &DebugData](double InStartTime, double InEndTime, uint32 InDepth, const FAnimNodeValueMessage& InMessage)
+													// don't send "Name" Node value for display in the graph
+													if (FPlatformString::Strcmp(InMessage.Key, TEXT("Name")) != 0)
 													{
-														// don't send "Name" Node value for display in the graph
-														if (FPlatformString::Strcmp(InMessage.Key, TEXT("Name")) != 0)
-														{
-															FText Text = AnimationProvider->FormatNodeKeyValue(InMessage);
-															DebugData.RecordNodeValue(InMessage.NodeId, Text.ToString());
-														}
-														return TraceServices::EEventEnumerate::Continue;
-													});
+														FText Text = AnimationProvider->FormatNodeKeyValue(InMessage);
+														DebugData.RecordNodeValue(InMessage.NodeId, Text.ToString());
+													}
+													return TraceServices::EEventEnumerate::Continue;
 												});
+											});
 
-												DebugData.DisableAllPoseWatches();
-					
-												AnimGraphTimeline.EnumerateEvents(Frame.StartTime, Frame.EndTime, [ObjectId, AnimationProvider, GameplayProvider, &DebugData, NodeCount](double InGraphStartTime, double InGraphEndTime, uint32 InDepth, const FAnimGraphMessage& InMessage)
+											DebugData.DisableAllPoseWatches();
+				
+											AnimGraphTimeline.EnumerateEvents(Frame.StartTime, Frame.EndTime, [ObjectId, AnimationProvider, GameplayProvider, &DebugData, NodeCount](double InGraphStartTime, double InGraphEndTime, uint32 InDepth, const FAnimGraphMessage& InMessage)
+											{
+												TRACE_CPUPROFILER_EVENT_SCOPE(AnimGraphTimelineEvent);
+														
+												// Basic verification - check node count is the same
+												// @TODO: could add some form of node hash/CRC to the class to improve this
+												if(InMessage.NodeCount == NodeCount)
 												{
-													TRACE_CPUPROFILER_EVENT_SCOPE(AnimGraphTimelineEvent);
-															
-													// Basic verification - check node count is the same
-													// @TODO: could add some form of node hash/CRC to the class to improve this
-													if(InMessage.NodeCount == NodeCount)
+													// Check for an update phase (which contains weights)
+													if(InMessage.Phase == EAnimGraphPhase::Update)
 													{
-														// Check for an update phase (which contains weights)
-														if(InMessage.Phase == EAnimGraphPhase::Update)
+														AnimationProvider->ReadAnimNodesTimeline(ObjectId, [InGraphStartTime, InGraphEndTime, &DebugData](const IAnimationProvider::AnimNodesTimeline& InNodesTimeline)
 														{
-															AnimationProvider->ReadAnimNodesTimeline(ObjectId, [InGraphStartTime, InGraphEndTime, &DebugData](const IAnimationProvider::AnimNodesTimeline& InNodesTimeline)
+															TRACE_CPUPROFILER_EVENT_SCOPE(AnimGraphDebugNodeVisits);
+															InNodesTimeline.EnumerateEvents(InGraphStartTime, InGraphEndTime, [&DebugData](double InStartTime, double InEndTime, uint32 InDepth, const FAnimNodeMessage& InMessage)
 															{
-																TRACE_CPUPROFILER_EVENT_SCOPE(AnimGraphDebugNodeVisits);
-																InNodesTimeline.EnumerateEvents(InGraphStartTime, InGraphEndTime, [&DebugData](double InStartTime, double InEndTime, uint32 InDepth, const FAnimNodeMessage& InMessage)
-																{
-																	DebugData.RecordNodeVisit(InMessage.NodeId, InMessage.PreviousNodeId, InMessage.Weight);
-																	return TraceServices::EEventEnumerate::Continue;
-																});
+																DebugData.RecordNodeVisit(InMessage.NodeId, InMessage.PreviousNodeId, InMessage.Weight);
+																return TraceServices::EEventEnumerate::Continue;
 															});
-					
-															AnimationProvider->ReadStateMachinesTimeline(ObjectId, [InGraphStartTime, InGraphEndTime, &DebugData](const IAnimationProvider::StateMachinesTimeline& InStateMachinesTimeline)
+														});
+				
+														AnimationProvider->ReadStateMachinesTimeline(ObjectId, [InGraphStartTime, InGraphEndTime, &DebugData](const IAnimationProvider::StateMachinesTimeline& InStateMachinesTimeline)
+														{
+															TRACE_CPUPROFILER_EVENT_SCOPE(AnimGraphDebugStateMachine);
+															InStateMachinesTimeline.EnumerateEvents(InGraphStartTime, InGraphEndTime, [&DebugData](double InStartTime, double InEndTime, uint32 InDepth, const FAnimStateMachineMessage& InMessage)
 															{
-																TRACE_CPUPROFILER_EVENT_SCOPE(AnimGraphDebugStateMachine);
-																InStateMachinesTimeline.EnumerateEvents(InGraphStartTime, InGraphEndTime, [&DebugData](double InStartTime, double InEndTime, uint32 InDepth, const FAnimStateMachineMessage& InMessage)
-																{
-																	DebugData.RecordStateData(InMessage.StateMachineIndex, InMessage.StateIndex, InMessage.StateWeight, InMessage.ElapsedTime);
-																	return TraceServices::EEventEnumerate::Continue;
-																});
+																DebugData.RecordStateData(InMessage.StateMachineIndex, InMessage.StateIndex, InMessage.StateWeight, InMessage.ElapsedTime);
+																return TraceServices::EEventEnumerate::Continue;
 															});
-					
-															AnimationProvider->ReadAnimSequencePlayersTimeline(ObjectId, [InGraphStartTime, InGraphEndTime, GameplayProvider, &DebugData](const IAnimationProvider::AnimSequencePlayersTimeline& InSequencePlayersTimeline)
+														});
+				
+														AnimationProvider->ReadAnimSequencePlayersTimeline(ObjectId, [InGraphStartTime, InGraphEndTime, GameplayProvider, &DebugData](const IAnimationProvider::AnimSequencePlayersTimeline& InSequencePlayersTimeline)
+														{
+															TRACE_CPUPROFILER_EVENT_SCOPE(AnimGraphDebugSequencePlayers);
+															InSequencePlayersTimeline.EnumerateEvents(InGraphStartTime, InGraphEndTime, [&DebugData](double InStartTime, double InEndTime, uint32 InDepth, const FAnimSequencePlayerMessage& InMessage)
 															{
-																TRACE_CPUPROFILER_EVENT_SCOPE(AnimGraphDebugSequencePlayers);
-																InSequencePlayersTimeline.EnumerateEvents(InGraphStartTime, InGraphEndTime, [&DebugData](double InStartTime, double InEndTime, uint32 InDepth, const FAnimSequencePlayerMessage& InMessage)
-																{
-																	DebugData.RecordSequencePlayer(InMessage.NodeId, InMessage.Position, InMessage.Length, InMessage.FrameCounter);
-																	return TraceServices::EEventEnumerate::Continue;
-																});
+																DebugData.RecordSequencePlayer(InMessage.NodeId, InMessage.Position, InMessage.Length, InMessage.FrameCounter);
+																return TraceServices::EEventEnumerate::Continue;
 															});
-					
-															AnimationProvider->ReadAnimBlendSpacePlayersTimeline(ObjectId, [InGraphStartTime, InGraphEndTime, GameplayProvider, &DebugData](const IAnimationProvider::BlendSpacePlayersTimeline& InBlendSpacePlayersTimeline)
+														});
+				
+														AnimationProvider->ReadAnimBlendSpacePlayersTimeline(ObjectId, [InGraphStartTime, InGraphEndTime, GameplayProvider, &DebugData](const IAnimationProvider::BlendSpacePlayersTimeline& InBlendSpacePlayersTimeline)
+														{
+															TRACE_CPUPROFILER_EVENT_SCOPE(AnimGraphBlendSpaces);
+															InBlendSpacePlayersTimeline.EnumerateEvents(InGraphStartTime, InGraphEndTime, [GameplayProvider, &DebugData](double InStartTime, double InEndTime, uint32 InDepth, const FBlendSpacePlayerMessage& InMessage)
 															{
-																TRACE_CPUPROFILER_EVENT_SCOPE(AnimGraphBlendSpaces);
-																InBlendSpacePlayersTimeline.EnumerateEvents(InGraphStartTime, InGraphEndTime, [GameplayProvider, &DebugData](double InStartTime, double InEndTime, uint32 InDepth, const FBlendSpacePlayerMessage& InMessage)
+																UBlendSpace* BlendSpace = nullptr;
+																const FObjectInfo* BlendSpaceInfo = GameplayProvider->FindObjectInfo(InMessage.BlendSpaceId);
+																if(BlendSpaceInfo)
 																{
-																	UBlendSpace* BlendSpace = nullptr;
-																	const FObjectInfo* BlendSpaceInfo = GameplayProvider->FindObjectInfo(InMessage.BlendSpaceId);
-																	if(BlendSpaceInfo)
-																	{
-																		BlendSpace = TSoftObjectPtr<UBlendSpace>(FSoftObjectPath(BlendSpaceInfo->PathName)).LoadSynchronous();
-																	}
-					
-																	DebugData.RecordBlendSpacePlayer(InMessage.NodeId, BlendSpace, FVector(InMessage.PositionX, InMessage.PositionY, InMessage.PositionZ), FVector(InMessage.FilteredPositionX, InMessage.FilteredPositionY, InMessage.FilteredPositionZ));
-																	return TraceServices::EEventEnumerate::Continue;
-																});
+																	BlendSpace = TSoftObjectPtr<UBlendSpace>(FSoftObjectPath(BlendSpaceInfo->PathName)).LoadSynchronous();
+																}
+				
+																DebugData.RecordBlendSpacePlayer(InMessage.NodeId, BlendSpace, FVector(InMessage.PositionX, InMessage.PositionY, InMessage.PositionZ), FVector(InMessage.FilteredPositionX, InMessage.FilteredPositionY, InMessage.FilteredPositionZ));
+																return TraceServices::EEventEnumerate::Continue;
 															});
-					
-															AnimationProvider->ReadAnimSyncTimeline(ObjectId, [InGraphStartTime, InGraphEndTime, AnimationProvider, &DebugData](const IAnimationProvider::AnimSyncTimeline& InAnimSyncTimeline)
+														});
+				
+														AnimationProvider->ReadAnimSyncTimeline(ObjectId, [InGraphStartTime, InGraphEndTime, AnimationProvider, &DebugData](const IAnimationProvider::AnimSyncTimeline& InAnimSyncTimeline)
+														{
+															TRACE_CPUPROFILER_EVENT_SCOPE(AnimGraphAnimSync);
+															InAnimSyncTimeline.EnumerateEvents(InGraphStartTime, InGraphEndTime, [AnimationProvider, &DebugData](double InStartTime, double InEndTime, uint32 InDepth, const FAnimSyncMessage& InMessage)
 															{
-																TRACE_CPUPROFILER_EVENT_SCOPE(AnimGraphAnimSync);
-																InAnimSyncTimeline.EnumerateEvents(InGraphStartTime, InGraphEndTime, [AnimationProvider, &DebugData](double InStartTime, double InEndTime, uint32 InDepth, const FAnimSyncMessage& InMessage)
+																const TCHAR* GroupName = AnimationProvider->GetName(InMessage.GroupNameId);
+																if(GroupName)
 																{
-																	const TCHAR* GroupName = AnimationProvider->GetName(InMessage.GroupNameId);
-																	if(GroupName)
-																	{
-																		DebugData.RecordNodeSync(InMessage.SourceNodeId, FName(GroupName));
-																	}
-														
-																	return TraceServices::EEventEnumerate::Continue;
-																});
+																	DebugData.RecordNodeSync(InMessage.SourceNodeId, FName(GroupName));
+																}
+													
+																return TraceServices::EEventEnumerate::Continue;
 															});
-														}
-					
-														// Some traces come from both update and evaluate phases
-														if(InMessage.Phase == EAnimGraphPhase::Update || InMessage.Phase == EAnimGraphPhase::Evaluate)
+														});
+													}
+				
+													// Some traces come from both update and evaluate phases
+													if(InMessage.Phase == EAnimGraphPhase::Update || InMessage.Phase == EAnimGraphPhase::Evaluate)
+													{
+														AnimationProvider->ReadAnimAttributesTimeline(ObjectId, [InGraphStartTime, InGraphEndTime, AnimationProvider, &DebugData](const IAnimationProvider::AnimAttributeTimeline& InAnimAttributeTimeline)
 														{
-															AnimationProvider->ReadAnimAttributesTimeline(ObjectId, [InGraphStartTime, InGraphEndTime, AnimationProvider, &DebugData](const IAnimationProvider::AnimAttributeTimeline& InAnimAttributeTimeline)
+															TRACE_CPUPROFILER_EVENT_SCOPE(AnimGraphAttributes);
+															InAnimAttributeTimeline.EnumerateEvents(InGraphStartTime, InGraphEndTime, [AnimationProvider, &DebugData](double InStartTime, double InEndTime, uint32 InDepth, const FAnimAttributeMessage& InMessage)
 															{
-																TRACE_CPUPROFILER_EVENT_SCOPE(AnimGraphAttributes);
-																InAnimAttributeTimeline.EnumerateEvents(InGraphStartTime, InGraphEndTime, [AnimationProvider, &DebugData](double InStartTime, double InEndTime, uint32 InDepth, const FAnimAttributeMessage& InMessage)
+																const TCHAR* AttributeName = AnimationProvider->GetName(InMessage.AttributeNameId);
+																if(AttributeName)
 																{
-																	const TCHAR* AttributeName = AnimationProvider->GetName(InMessage.AttributeNameId);
-																	if(AttributeName)
-																	{
-																		DebugData.RecordNodeAttribute(InMessage.TargetNodeId, InMessage.SourceNodeId, FName(AttributeName));
-																	}
-														
-																	return TraceServices::EEventEnumerate::Continue;
-																});
+																	DebugData.RecordNodeAttribute(InMessage.TargetNodeId, InMessage.SourceNodeId, FName(AttributeName));
+																}
+													
+																return TraceServices::EEventEnumerate::Continue;
 															});
+														});
 
-															
-															AnimationProvider->ReadPoseWatchTimeline(ObjectId, [InGraphStartTime, InGraphEndTime, AnimationProvider, &DebugData](const IAnimationProvider::PoseWatchTimeline& InPoseWatchTimeline)
-																{
-																	InPoseWatchTimeline.EnumerateEvents(InGraphStartTime, InGraphEndTime, [AnimationProvider, &DebugData](double InStartTime, double InEndTime, uint32 InDepth, const FPoseWatchMessage& InMessage)
+														
+														AnimationProvider->ReadPoseWatchTimeline(ObjectId, [InGraphStartTime, InGraphEndTime, AnimationProvider, &DebugData](const IAnimationProvider::PoseWatchTimeline& InPoseWatchTimeline)
+															{
+																InPoseWatchTimeline.EnumerateEvents(InGraphStartTime, InGraphEndTime, [AnimationProvider, &DebugData](double InStartTime, double InEndTime, uint32 InDepth, const FPoseWatchMessage& InMessage)
+																	{
+																		for (FAnimNodePoseWatch& PoseWatch : DebugData.AnimNodePoseWatch)
 																		{
-																			for (FAnimNodePoseWatch& PoseWatch : DebugData.AnimNodePoseWatch)
+																			if (PoseWatch.NodeID == InMessage.PoseWatchId)
 																			{
-																				if (PoseWatch.NodeID == InMessage.PoseWatchId)
-																				{
-																					TArray<FBoneIndexType> RequiredBones;
-																					TArray<FTransform> BoneTransforms;
-																					AnimationProvider->GetPoseWatchData(InMessage, BoneTransforms, RequiredBones);
-
-																					PoseWatch.SetPose(RequiredBones, BoneTransforms);
-																					PoseWatch.SetWorldTransform(InMessage.WorldTransform);
-
-																					PoseWatch.PoseWatch->SetIsNodeEnabled(true);
-																					break;
-																				}
+																				TArray<FBoneIndexType> RequiredBones;
+																				TArray<FTransform> BoneTransforms;
+																				AnimationProvider->GetPoseWatchData(InMessage, BoneTransforms, RequiredBones);
+
+																				PoseWatch.SetPose(RequiredBones, BoneTransforms);
+																				PoseWatch.SetWorldTransform(InMessage.WorldTransform);
+
+																				PoseWatch.PoseWatch->SetIsNodeEnabled(true);
+																				break;
 																			}
-																			return TraceServices::EEventEnumerate::Continue;
-																		});
-																});
+																		}
+																		return TraceServices::EEventEnumerate::Continue;
+																	});
+															});
 
-														}
-					
 													}
-													return TraceServices::EEventEnumerate::Continue;
-												});
-											}
+				
+												}
+												return TraceServices::EEventEnumerate::Continue;
+											});
 										}
 									}
 								}
-								return TraceServices::EEventEnumerate::Continue;
-							});
-						}
+							}
+							return TraceServices::EEventEnumerate::Continue;
+						});
 					}
 				}
 			}
+			LastScrubTime = CurrentTraceTime;
+		}
 	}
 }
 
diff --git a/Engine/Plugins/Animation/GameplayInsights/Source/RewindDebugger/Private/RewindDebuggerAnimation.h b/Engine/Plugins/Animation/GameplayInsights/Source/RewindDebugger/Private/RewindDebuggerAnimation.h
--- a/Engine/Plugins/Animation/GameplayInsights/Source/RewindDebugger/Private/RewindDebuggerAnimation.h
+++ b/Engine/Plugins/Animation/GameplayInsights/Source/RewindDebugger/Private/RewindDebuggerAnimation.h
@@ -50,6 +50,7 @@ private:
 	{
 		TWeakObjectPtr<USkeletalMeshComponent> Component;
 		FTransform RelativeTransform;
+		int ForcedLod = -1;
 	};
 
 	struct FSpawnedMeshComponentInfo
