Icon lookup optimizations for Rewind Debugger

diff --git a/Engine/Plugins/Animation/GameplayInsights/Source/GameplayInsights/Private/GameplayProvider.cpp b/Engine/Plugins/Animation/GameplayInsights/Source/GameplayInsights/Private/GameplayProvider.cpp
--- a/Engine/Plugins/Animation/GameplayInsights/Source/GameplayInsights/Private/GameplayProvider.cpp
+++ b/Engine/Plugins/Animation/GameplayInsights/Source/GameplayInsights/Private/GameplayProvider.cpp
@@ -6,6 +6,7 @@
 #include "AssetRegistry/AssetRegistryModule.h"
 #include "Misc/PackageName.h"
 #include "Model/AsyncEnumerateTask.h"
+#include "Styling/SlateIconFinder.h"
 
 FName FGameplayProvider::ProviderName("GameplayProvider");
 
@@ -55,9 +56,9 @@ static UClass* FindBlueprintClass(const FString& TargetNameRaw)
 
 // This is copied from EditorUtilitySubsystem (where it is not public)
 // Should probably be somewhere shared
-static UClass* FindClassByPathName(const FString& RawTargePathtName)
+static UClass* FindClassByPathName(const FString& RawTargePathName, bool bCheckBlueprintClasses)
 {
-	FString TargetName = RawTargePathtName;
+	FString TargetName = RawTargePathName;
 
 	// Check native classes and loaded assets first before resorting to the asset registry
 	bool bIsValidClassName = true;
@@ -96,7 +97,7 @@ static UClass* FindClassByPathName(const FString& RawTargePathtName)
 	}
 
 	// If we still haven't found anything yet, try the asset registry for blueprints that match the requirements
-	if (ResultClass == nullptr)
+	if (ResultClass == nullptr && bCheckBlueprintClasses)
 	{
 		ResultClass = FindBlueprintClass(TargetName);
 	}
@@ -272,12 +273,12 @@ const FClassInfo* FGameplayProvider::FindClassInfo(uint64 InClassId) const
 	return nullptr;
 }
 
-const UClass* FGameplayProvider::FindClass(uint64 InClassId) const
+const UClass* FGameplayProvider::FindClass(uint64 InClassId, bool bSearchBlueprints) const
 {
 #if WITH_EDITOR
 	if (const FClassInfo* ClassInfo = FindClassInfo(InClassId))
 	{
-		return FindClassByPathName(ClassInfo->PathName);
+		return FindClassByPathName(ClassInfo->PathName, bSearchBlueprints);
 	}
 	return nullptr;
 #else
@@ -647,6 +648,34 @@ double FGameplayProvider::GetRecordingDuration() const
 	return RecordingDuration;
 }
 
+FSlateIcon FGameplayProvider::FindIconForClass(uint64 ClassId) const
+{
+	FSlateIcon Icon;
+#if WITH_EDITOR
+	// traverse class hierarchy until we find a class that exists, and return the icon for that class
+	// loaded trace data may contain classes from different projects than this editor is compiled against or loading
+	// but we will still find a Character icon for example for subclasses of Character
+	while(ClassId)
+	{
+		// don't search blueprints because it's very slow and icons are associated with native classes
+		if (const UClass* FoundClass = FindClass(ClassId, false))
+		{
+			Icon = FSlateIconFinder::FindIconForClass(FoundClass);
+			break;
+		}
+		
+		const FClassInfo& ClassInfo = GetClassInfo(ClassId);
+		ClassId = ClassInfo.SuperId;
+	}
+	
+	if (!Icon.IsSet())
+	{
+		Icon = FSlateIconFinder::FindIconForClass(UObject::StaticClass());
+	}
+#endif
+	return Icon;
+}
+
 void FGameplayProvider::AppendView(uint64 InPlayerId, double InTime, const FVector& InPosition, const FRotator& InRotation, float InFov, float InAspectRatio)
 {
 	Session.WriteAccessCheck();
diff --git a/Engine/Plugins/Animation/GameplayInsights/Source/GameplayInsights/Private/GameplayProvider.h b/Engine/Plugins/Animation/GameplayInsights/Source/GameplayInsights/Private/GameplayProvider.h
--- a/Engine/Plugins/Animation/GameplayInsights/Source/GameplayInsights/Private/GameplayProvider.h
+++ b/Engine/Plugins/Animation/GameplayInsights/Source/GameplayInsights/Private/GameplayProvider.h
@@ -29,7 +29,7 @@ public:
 	virtual void EnumerateSubobjects(uint64 ObjectId, TFunctionRef<void(uint64 SubobjectId)> Callback) const override;
 	virtual const FObjectPropertyValue* FindPropertyValueFromStorageIndex(uint64 InObjectId, int64 InStorageIndex) const override;
 	virtual const FClassInfo* FindClassInfo(uint64 InClassId) const override;
-	virtual const UClass* FindClass(uint64 InClassId) const override;
+	virtual const UClass* FindClass(uint64 InClassId, bool bSearchBlueprints = true) const override;
 	virtual const FClassInfo* FindClassInfo(const TCHAR* InClassPath) const override;
 	virtual bool IsSubClassOf(uint64 InSubClassId, uint64 InParentClassId) const override;
 	virtual const FObjectInfo* FindObjectInfo(uint64 InObjectId) const override;
@@ -44,6 +44,7 @@ public:
 	virtual const RecordingInfoTimeline* GetRecordingInfo(uint32 RecordingId) const override; 
 	virtual void ReadViewTimeline(TFunctionRef<void(const ViewTimeline&)> Callback) const override;
 	virtual double GetRecordingDuration() const override;
+	virtual FSlateIcon FindIconForClass(uint64 ClassId) const override;
 
 
 	/** Add a class message */
diff --git a/Engine/Plugins/Animation/GameplayInsights/Source/GameplayInsights/Public/IGameplayProvider.h b/Engine/Plugins/Animation/GameplayInsights/Source/GameplayInsights/Public/IGameplayProvider.h
--- a/Engine/Plugins/Animation/GameplayInsights/Source/GameplayInsights/Public/IGameplayProvider.h
+++ b/Engine/Plugins/Animation/GameplayInsights/Source/GameplayInsights/Public/IGameplayProvider.h
@@ -98,6 +98,8 @@ struct FWorldInfo
 	bool bIsSimulating = false;
 };
 
+struct FSlateIcon;
+
 // Delegate fired when an object receives an end play event
 DECLARE_MULTICAST_DELEGATE_ThreeParams(FOnObjectEndPlay, uint64 /*ObjectId*/, double /*Time*/, const FObjectInfo& /*ObjectInfo*/);
 
@@ -119,7 +121,7 @@ public:
 	virtual void EnumerateSubobjects(uint64 ObjectId, TFunctionRef<void(uint64 SubobjectId)> Callback) const = 0;
 	virtual const FObjectPropertyValue * FindPropertyValueFromStorageIndex(uint64 InObjectId, int64 InStorageIndex) const = 0;
 	virtual const FClassInfo* FindClassInfo(uint64 InClassId) const = 0;
-	virtual const UClass* FindClass(uint64 InClassId) const = 0;
+	virtual const UClass* FindClass(uint64 InClassId, bool bSearchBlueprints) const = 0;
 	virtual const FClassInfo* FindClassInfo(const TCHAR* InClassPath) const = 0;
 	virtual bool IsSubClassOf(uint64 InSubClassId, uint64 InParentClassId) const = 0;
 	virtual const FObjectInfo* FindObjectInfo(uint64 InObjectId) const = 0;
@@ -137,4 +139,5 @@ public:
 	virtual TRange<double> GetObjectTraceLifetime(uint64 ObjectId) const = 0;
 	virtual TRange<double> GetObjectRecordingLifetime(uint64 ObjectId) const = 0;
 	virtual double GetRecordingDuration() const = 0;
+	virtual FSlateIcon FindIconForClass(uint64 ClassId) const = 0;
 };
\ No newline at end of file
diff --git a/Engine/Plugins/Animation/GameplayInsights/Source/RewindDebugger/Private/RewindDebuggerObjectTrack.cpp b/Engine/Plugins/Animation/GameplayInsights/Source/RewindDebugger/Private/RewindDebuggerObjectTrack.cpp
--- a/Engine/Plugins/Animation/GameplayInsights/Source/RewindDebugger/Private/RewindDebuggerObjectTrack.cpp
+++ b/Engine/Plugins/Animation/GameplayInsights/Source/RewindDebugger/Private/RewindDebuggerObjectTrack.cpp
@@ -45,6 +45,7 @@ FRewindDebuggerObjectTrack::FRewindDebuggerObjectTrack(uint64 InObjectId, const
 	, ObjectId(InObjectId)
 	, bAddController(bInAddController)
 	, bDisplayNameValid(false)
+	, bIconSearched(false)
 {
 	ExistenceRange = MakeShared<SEventTimelineView::FTimelineEventData>();
 	
@@ -197,17 +198,20 @@ bool FRewindDebuggerObjectTrack::UpdateInternal()
 		ExistenceRange->Windows.Add({Existence.GetLowerBoundValue(), Existence.GetUpperBoundValue(), LOCTEXT("Object Existence","Object Existence"), LOCTEXT("Object Existence","Object Existence"), FLinearColor(0.1f,0.11f,0.1f)});
 	}
 
-	if (!Icon.IsSet())
+	if (!bIconSearched)
 	{
 		TRACE_CPUPROFILER_EVENT_SCOPE(FRewindDebuggerObjectTrack::FindIcon);
 		if (const FObjectInfo* ObjectInfo = GameplayProvider->FindObjectInfo(ObjectId))
 		{
-			if (const UClass* FoundClass = GameplayProvider->FindClass(ObjectInfo->ClassId))
-			{
-				Icon = FSlateIconFinder::FindIconForClass(FoundClass);
-				bChanged = true;
-			}
+			Icon = GameplayProvider->FindIconForClass(ObjectInfo->ClassId);
+		}
+						
+		if (!Icon.IsSet())
+		{
+			Icon = FSlateIconFinder::FindIconForClass(UObject::StaticClass());
 		}
+		bIconSearched = true;
+		bChanged = true;
 	}
 
 	TArray<uint64, TInlineAllocator<32>> FoundObjects;
diff --git a/Engine/Plugins/Animation/GameplayInsights/Source/RewindDebugger/Private/RewindDebuggerObjectTrack.h b/Engine/Plugins/Animation/GameplayInsights/Source/RewindDebugger/Private/RewindDebuggerObjectTrack.h
--- a/Engine/Plugins/Animation/GameplayInsights/Source/RewindDebugger/Private/RewindDebuggerObjectTrack.h
+++ b/Engine/Plugins/Animation/GameplayInsights/Source/RewindDebugger/Private/RewindDebuggerObjectTrack.h
@@ -46,7 +46,7 @@ private:
 
 	bool bAddController;
 	mutable bool bDisplayNameValid;
-	
+	bool bIconSearched;
 };
 
 }
\ No newline at end of file
diff --git a/Engine/Plugins/Animation/GameplayInsights/Source/RewindDebugger/Private/SRewindDebugger.cpp b/Engine/Plugins/Animation/GameplayInsights/Source/RewindDebugger/Private/SRewindDebugger.cpp
--- a/Engine/Plugins/Animation/GameplayInsights/Source/RewindDebugger/Private/SRewindDebugger.cpp
+++ b/Engine/Plugins/Animation/GameplayInsights/Source/RewindDebugger/Private/SRewindDebugger.cpp
@@ -212,17 +212,8 @@ TSharedRef<SWidget> SRewindDebugger::MakeSelectActorMenu()
 				{
 					FString ActorName = ObjectInfo.Name;
 					FText SelectedLabel = FText::FromString(ActorName);
-					FSlateIcon ActorIcon;
-
+					FSlateIcon ActorIcon = GameplayProvider->FindIconForClass(ObjectInfo.ClassId);
 					const FClassInfo& ClassInfo = GameplayProvider->GetClassInfo(ObjectInfo.ClassId);
-					if (UClass* FoundClass = UClass::TryFindTypeSlow<UClass>(ClassInfo.Name))
-					{
-						ActorIcon = FSlateIconFinder::FindIconForClass(FoundClass);
-					}
-					else
-					{
-						ActorIcon = FSlateIconFinder::FindIconForClass(AActor::StaticClass());
-					}
 
 					MenuBuilder.AddMenuEntry(SelectedLabel, FText(), ActorIcon, FExecuteAction::CreateLambda([this, ActorName]()
 					{
@@ -389,18 +380,12 @@ void SRewindDebugger::Construct(const FArguments& InArgs, TSharedRef<FUICommandL
 								SNew(SImage)
 								.Image_Lambda([this]
 									{
-										FSlateIcon ActorIcon = FSlateIconFinder::FindIconForClass(AActor::StaticClass());
 										if (DebugComponents != nullptr && DebugComponents->Num()>0)
 										{
-#if OBJECT_TRACE_ENABLED
-											if (UObject* Object = FObjectTrace::GetObjectFromId((*DebugComponents)[0]->GetObjectId()))
-											{
-												ActorIcon = FSlateIconFinder::FindIconForClass(Object->GetClass());
-											}
-#endif // OBJECT_TRACE_ENABLED
+											return (*DebugComponents)[0]->GetIcon().GetIcon();
 										}
 
-										return ActorIcon.GetIcon();
+										return FSlateIconFinder::FindIconForClass(AActor::StaticClass()).GetIcon();
 									}
 								)
 							]
