Update rewind debugger to record and apply mesh component visibility state

diff --git a/Engine/Plugins/Animation/GameplayInsights/Source/GameplayInsights/Private/AnimationAnalyzer.cpp b/Engine/Plugins/Animation/GameplayInsights/Source/GameplayInsights/Private/AnimationAnalyzer.cpp
--- a/Engine/Plugins/Animation/GameplayInsights/Source/GameplayInsights/Private/AnimationAnalyzer.cpp
+++ b/Engine/Plugins/Animation/GameplayInsights/Source/GameplayInsights/Private/AnimationAnalyzer.cpp
@@ -194,7 +194,7 @@ bool FAnimationAnalyzer::OnEvent(uint16 RouteId, EStyle Style, const FOnEventCon
 	{
 		uint64 Cycle = EventData.GetValue<uint64>("Cycle");
 		double RecordingTime = EventData.GetValue<double>("RecordingTime");
-		
+		bool bIsVisible = EventData.GetValue<bool>("IsVisible", true);	
 		uint64 ComponentId = EventData.GetValue<uint64>("ComponentId");
 		uint64 MeshId = EventData.GetValue<uint64>("MeshId");
 		uint32 BoneCount = EventData.GetValue<uint32>("BoneCount");
@@ -215,6 +215,7 @@ bool FAnimationAnalyzer::OnEvent(uint16 RouteId, EStyle Style, const FOnEventCon
 			RecordingTime,
 			LodIndex,
 			FrameCounter,
+			bIsVisible,
 			ComponentToWorldFloatArray,
 			PoseFloatArray,
 			CurveIds,
diff --git a/Engine/Plugins/Animation/GameplayInsights/Source/GameplayInsights/Private/AnimationProvider.cpp b/Engine/Plugins/Animation/GameplayInsights/Source/GameplayInsights/Private/AnimationProvider.cpp
--- a/Engine/Plugins/Animation/GameplayInsights/Source/GameplayInsights/Private/AnimationProvider.cpp
+++ b/Engine/Plugins/Animation/GameplayInsights/Source/GameplayInsights/Private/AnimationProvider.cpp
@@ -868,6 +868,7 @@ void FAnimationProvider::AppendSkeletalMeshComponent(
 	double InRecordingTime,
 	uint16 InLodIndex,
 	uint16 InFrameCounter,
+	bool bInIsVisible,
 	const TArrayView<const float>& InComponentToWorldRaw,
 	const TArrayView<const float>& InPoseRaw,
 	const TArrayView<const uint32>& InCurveIds,
@@ -924,6 +925,7 @@ void FAnimationProvider::AppendSkeletalMeshComponent(
 
 	FSkeletalMeshPoseMessage Message;
 	Message.RecordingTime = InRecordingTime;
+	Message.bIsVisible = bInIsVisible;
 	Message.ComponentToWorld = ComponentToWorld;
 	Message.TransformStartIndex = SkeletalMeshPoseTransforms.Num();
 	Message.CurveStartIndex = SkeletalMeshCurves.Num();
diff --git a/Engine/Plugins/Animation/GameplayInsights/Source/GameplayInsights/Private/AnimationProvider.h b/Engine/Plugins/Animation/GameplayInsights/Source/GameplayInsights/Private/AnimationProvider.h
--- a/Engine/Plugins/Animation/GameplayInsights/Source/GameplayInsights/Private/AnimationProvider.h
+++ b/Engine/Plugins/Animation/GameplayInsights/Source/GameplayInsights/Private/AnimationProvider.h
@@ -70,7 +70,7 @@ public:
 	/** Add a skeletal mesh pose/curves etc. */
 	void AppendSkeletalMeshComponent(uint64 InObjectId, uint64 InMeshId, double InProfileTime, double InRecordingTime, uint16 InLodIndex, uint16 InFrameCounter, const TArrayView<const FTransform>& InPose, const TArrayView<const FSkeletalMeshNamedCurve>& InCurves);
 	void AppendSkeletalMeshComponent(uint64 InObjectId, uint64 InMeshId, double InProfileTime, double InRecordingTime, uint16 InLodIndex, uint16 InFrameCounter, const TArrayView<const float>& InComponentToWorldRaw, const TArrayView<const float>& InPoseRaw, const TArrayView<const uint32>& InCurveIds, const TArrayView<const float>& InCurveValues);
-	void AppendSkeletalMeshComponent(uint64 InObjectId, uint64 InMeshId, double InProfileTime, double InRecordingTime, uint16 InLodIndex, uint16 InFrameCounter, const TArrayView<const float>& InComponentToWorldRaw, const TArrayView<const float>& InPoseRaw, const TArrayView<const uint32>& InCurveIds, const TArrayView<const float>& InCurveValues, const TArrayView<const float>& InExternalMorphWeights, const TArrayView<const int32>& InExternalMorphCounts);
+	void AppendSkeletalMeshComponent(uint64 InObjectId, uint64 InMeshId, double InProfileTime, double InRecordingTime, uint16 InLodIndex, uint16 InFrameCounter, bool bInIsVisible, const TArrayView<const float>& InComponentToWorldRaw, const TArrayView<const float>& InPoseRaw, const TArrayView<const uint32>& InCurveIds, const TArrayView<const float>& InCurveValues, const TArrayView<const float>& InExternalMorphWeights, const TArrayView<const int32>& InExternalMorphCounts);
 
 	/** Get a skeletal mesh for a specified path. If the mesh no longer exists, make a fake one */
 	USkeletalMesh* GetSkeletalMesh(const TCHAR* InPath);
diff --git a/Engine/Plugins/Animation/GameplayInsights/Source/GameplayInsights/Public/IAnimationProvider.h b/Engine/Plugins/Animation/GameplayInsights/Source/GameplayInsights/Public/IAnimationProvider.h
--- a/Engine/Plugins/Animation/GameplayInsights/Source/GameplayInsights/Public/IAnimationProvider.h
+++ b/Engine/Plugins/Animation/GameplayInsights/Source/GameplayInsights/Public/IAnimationProvider.h
@@ -52,6 +52,7 @@ struct FSkeletalMeshPoseMessage
 	uint16 FrameCounter = 0;
 	uint16 LodIndex = 0;
 	uint16 NumExternalMorphSets = 0;
+	bool bIsVisible = true;
 };
 
 struct FPoseWatchMessage
diff --git a/Engine/Plugins/Animation/GameplayInsights/Source/RewindDebugger/Private/RewindDebuggerAnimation.cpp b/Engine/Plugins/Animation/GameplayInsights/Source/RewindDebugger/Private/RewindDebuggerAnimation.cpp
--- a/Engine/Plugins/Animation/GameplayInsights/Source/RewindDebugger/Private/RewindDebuggerAnimation.cpp
+++ b/Engine/Plugins/Animation/GameplayInsights/Source/RewindDebugger/Private/RewindDebuggerAnimation.cpp
@@ -92,6 +92,7 @@ void FRewindDebuggerAnimation::OnPIEResumed(bool bSimulating)
 		{
 			MeshComponent->SetRelativeTransform(MeshData.Value.RelativeTransform, false, nullptr, ETeleportType::TeleportPhysics);
 			MeshComponent->SetForcedLOD(MeshData.Value.ForcedLod);
+			MeshComponent->SetVisibility(MeshData.Value.bIsVisible);
 		}
 	}
 
@@ -178,6 +179,7 @@ void FRewindDebuggerAnimation::ApplyPoseToMesh(const IAnimationProvider* Animati
 					ResetData.Component = MeshComponent;
 					ResetData.RelativeTransform = MeshComponent->GetRelativeTransform();
 					ResetData.ForcedLod = MeshComponent->GetForcedLOD();
+					ResetData.bIsVisible = MeshComponent->GetVisibleFlag();
 					MeshComponentsToReset.Add(ObjectId, ResetData);
 				}
 			}
@@ -186,6 +188,7 @@ void FRewindDebuggerAnimation::ApplyPoseToMesh(const IAnimationProvider* Animati
 			MeshComponent->SetForcedLOD(PoseMessage->LodIndex + 1);
 			MeshComponent->UpdateLODStatus();
 			MeshComponent->UpdateChildTransforms(EUpdateTransformFlags::None, ETeleportType::TeleportPhysics);
+			MeshComponent->SetVisibility(PoseMessage->bIsVisible);
 		}
 	}
 }
diff --git a/Engine/Plugins/Animation/GameplayInsights/Source/RewindDebugger/Private/RewindDebuggerAnimation.h b/Engine/Plugins/Animation/GameplayInsights/Source/RewindDebugger/Private/RewindDebuggerAnimation.h
--- a/Engine/Plugins/Animation/GameplayInsights/Source/RewindDebugger/Private/RewindDebuggerAnimation.h
+++ b/Engine/Plugins/Animation/GameplayInsights/Source/RewindDebugger/Private/RewindDebuggerAnimation.h
@@ -51,6 +51,7 @@ private:
 		TWeakObjectPtr<USkeletalMeshComponent> Component;
 		FTransform RelativeTransform;
 		int ForcedLod = -1;
+		bool bIsVisible = true;
 	};
 
 	struct FSpawnedMeshComponentInfo
diff --git a/Engine/Source/Runtime/Engine/Private/Animation/AnimTrace.cpp b/Engine/Source/Runtime/Engine/Private/Animation/AnimTrace.cpp
--- a/Engine/Source/Runtime/Engine/Private/Animation/AnimTrace.cpp
+++ b/Engine/Source/Runtime/Engine/Private/Animation/AnimTrace.cpp
@@ -62,6 +62,7 @@ UE_TRACE_EVENT_BEGIN(Animation, SkeletalMeshComponent4)
 	UE_TRACE_EVENT_FIELD(double, RecordingTime)
 	UE_TRACE_EVENT_FIELD(uint64, ComponentId)
 	UE_TRACE_EVENT_FIELD(uint64, MeshId)
+	UE_TRACE_EVENT_FIELD(bool, IsVisible)
 	UE_TRACE_EVENT_FIELD(float[], ComponentToWorld)
 	UE_TRACE_EVENT_FIELD(float[], Pose)
 	UE_TRACE_EVENT_FIELD(uint32[], CurveIds)
@@ -670,6 +671,7 @@ void FAnimTrace::OutputSkeletalMeshComponent(const USkeletalMeshComponent* InCom
 			<< SkeletalMeshComponent4.RecordingTime(FObjectTrace::GetWorldElapsedTime(InComponent->GetWorld()))
 			<< SkeletalMeshComponent4.ComponentId(FObjectTrace::GetObjectId(InComponent))
 			<< SkeletalMeshComponent4.MeshId(FObjectTrace::GetObjectId(InComponent->GetSkeletalMeshAsset()))
+			<< SkeletalMeshComponent4.IsVisible(InComponent->GetVisibleFlag())
 			<< SkeletalMeshComponent4.ComponentToWorld(reinterpret_cast<const float*>(&InComponent->GetComponentToWorld()), sizeof(FTransform) / sizeof(float))
 			<< SkeletalMeshComponent4.Pose(reinterpret_cast<const float*>(InComponent->GetComponentSpaceTransforms().GetData()), BoneCount * (sizeof(FTransform) / sizeof(float)))
 			<< SkeletalMeshComponent4.CurveIds(CurveIds.GetData(), CurveIds.Num())
