Fix crash stopping PIE with a Rewind Debugger Recording

diff --git a/Engine/Plugins/Animation/GameplayInsights/Source/RewindDebugger/Private/RewindDebugger.cpp b/Engine/Plugins/Animation/GameplayInsights/Source/RewindDebugger/Private/RewindDebugger.cpp
--- a/Engine/Plugins/Animation/GameplayInsights/Source/RewindDebugger/Private/RewindDebugger.cpp
+++ b/Engine/Plugins/Animation/GameplayInsights/Source/RewindDebugger/Private/RewindDebugger.cpp
@@ -229,8 +229,6 @@ void FRewindDebugger::OnPIEStopped(bool bSimulating)
 #endif // OBJECT_TRACE_ENABLED
 	}
 	
-	bTraceFileLoaded = IsRecording();
-	
 	bPIEStarted = false;
 	bPIESimulating = false;
 
@@ -449,8 +447,6 @@ void FRewindDebugger::StartRecording()
 	UnrealInsightsModule->StartAnalysisForLastLiveSession(5.0);
 
 	TargetObjectIds.Empty(2);
-
-	bTraceFileLoaded = false;
 }
 
 bool FRewindDebugger::CanOpenTrace() const
@@ -465,7 +461,6 @@ void FRewindDebugger::OpenTrace(const FString& FilePath)
 
 	IUnrealInsightsModule& TraceInsightsModule = FModuleManager::LoadModuleChecked<IUnrealInsightsModule>("TraceInsights");
 	TraceInsightsModule.StartAnalysisForTraceFile(*FilePath);
-	bTraceFileLoaded = GetAnalysisSession() != nullptr;
 
 	// todo: optionally open the map the trace file was recorded in
 }
@@ -510,7 +505,6 @@ void FRewindDebugger::ClearTrace()
 	StopRecording();
 	RecordingDuration.Set(0);
 	
-	bTraceFileLoaded = false;
 	TargetObjectIds.Empty();
 	CurrentTraceRange.SetLowerBoundValue(0);
 	CurrentTraceRange.SetUpperBoundValue(0);
@@ -774,18 +768,18 @@ UWorld* FRewindDebugger::GetWorldToVisualize() const
 	UEditorEngine* EditorEngine = Cast<UEditorEngine>(GEngine);
 	if (GIsEditor && EditorEngine != nullptr && World == nullptr)
 	{
-		// when a trace file is loaded, use the editor world, otherwise use the play world
-		World = bTraceFileLoaded ?  EditorEngine->GetEditorWorldContext().World() : ToRawPtr(EditorEngine->PlayWorld);
-	}
-
-	if (!GIsEditor && World == nullptr)
-	{
-		World = GEngine->GetWorld();
+		// lets use PlayWorld during PIE/Simulate and regular world from editor otherwise, to draw debug information
+		World = EditorEngine->PlayWorld != nullptr ? ToRawPtr(EditorEngine->PlayWorld) : EditorEngine->GetEditorWorldContext().World();
 	}
 
 	return World;
 }
 
+bool FRewindDebugger::IsTraceFileLoaded() const
+{
+	return GetAnalysisSession()!=nullptr && !bPIEStarted;
+}
+
 void FRewindDebugger::SetCurrentViewRange(const TRange<double>& Range)
 {
 	CurrentViewRange = Range;
diff --git a/Engine/Plugins/Animation/GameplayInsights/Source/RewindDebugger/Private/RewindDebugger.h b/Engine/Plugins/Animation/GameplayInsights/Source/RewindDebugger/Private/RewindDebugger.h
--- a/Engine/Plugins/Animation/GameplayInsights/Source/RewindDebugger/Private/RewindDebugger.h
+++ b/Engine/Plugins/Animation/GameplayInsights/Source/RewindDebugger/Private/RewindDebugger.h
@@ -43,7 +43,7 @@ public:
 	virtual UWorld* GetWorldToVisualize() const override;
 	virtual bool IsRecording() const override { return bRecording; }
 	virtual bool IsPIESimulating() const override { return bPIESimulating; }
-	virtual bool IsTraceFileLoaded() const override { return bTraceFileLoaded; }
+	virtual bool IsTraceFileLoaded() const override;
 	virtual double GetRecordingDuration() const override { return RecordingDuration.Get(); }
 	virtual TSharedPtr<FDebugObjectInfo> GetSelectedComponent() const override;
 	virtual TSharedPtr<RewindDebugger::FRewindDebuggerTrack> GetSelectedTrack() const override;
@@ -172,7 +172,6 @@ private:
 	bool bTraceJustConnected = false;
 	bool bPIEStarted = false;
 	bool bPIESimulating = false;
-	bool bTraceFileLoaded = false;
 	
 	bool bRecording = false;
 
diff --git a/Engine/Plugins/Animation/GameplayInsights/Source/RewindDebugger/Private/RewindDebuggerAnimation.cpp b/Engine/Plugins/Animation/GameplayInsights/Source/RewindDebugger/Private/RewindDebuggerAnimation.cpp
--- a/Engine/Plugins/Animation/GameplayInsights/Source/RewindDebugger/Private/RewindDebuggerAnimation.cpp
+++ b/Engine/Plugins/Animation/GameplayInsights/Source/RewindDebugger/Private/RewindDebuggerAnimation.cpp
@@ -220,31 +220,34 @@ const FObjectInfo* FindOwningActorInfo(const IGameplayProvider* GameplayProvider
 
 FRewindDebuggerAnimation::FSpawnedMeshComponentInfo* FRewindDebuggerAnimation::SpawnMesh(uint64 ObjectId, const IGameplayProvider* GameplayProvider)
 {
-	FSpawnedMeshComponentInfo* MeshComponentInfo = &SpawnedMeshComponents.Add(ObjectId);
-
-	UWorld* World = IRewindDebugger::Instance()->GetWorldToVisualize();
+	IRewindDebugger* RewindDebugger = IRewindDebugger::Instance();
+	check(RewindDebugger);
+	if (UWorld* World = RewindDebugger->GetWorldToVisualize())
+	{
+		FSpawnedMeshComponentInfo* MeshComponentInfo = &SpawnedMeshComponents.Add(ObjectId);
 
-	FActorSpawnParameters ActorSpawnParameters;
-	ActorSpawnParameters.bHideFromSceneOutliner = true;
-	ActorSpawnParameters.ObjectFlags |= RF_Transient;
+		FActorSpawnParameters ActorSpawnParameters;
+		ActorSpawnParameters.bHideFromSceneOutliner = true;
+		ActorSpawnParameters.ObjectFlags |= RF_Transient;
 
-	MeshComponentInfo->Actor = World->SpawnActor<AActor>(ActorSpawnParameters);
+		MeshComponentInfo->Actor = World->SpawnActor<AActor>(ActorSpawnParameters);
 
-	if(const FObjectInfo* ActorInfo = FindOwningActorInfo(GameplayProvider, ObjectId))
-	{
-		MeshComponentInfo->Actor->SetActorLabel(FString(TEXT("RewindDebugger: ") + FString(ActorInfo->Name)));
-	}
-
-	MeshComponentInfo->Component = NewObject<USkeletalMeshComponent>(MeshComponentInfo->Actor.Get());
-	MeshComponentInfo->Component->PrimaryComponentTick.bStartWithTickEnabled = false;
-	MeshComponentInfo->Component->PrimaryComponentTick.bCanEverTick = false;
+		if(const FObjectInfo* ActorInfo = FindOwningActorInfo(GameplayProvider, ObjectId))
+		{
+			MeshComponentInfo->Actor->SetActorLabel(FString(TEXT("RewindDebugger: ") + FString(ActorInfo->Name)));
+		}
 
-	MeshComponentInfo->Actor->AddInstanceComponent(MeshComponentInfo->Component);
+		MeshComponentInfo->Component = NewObject<USkeletalMeshComponent>(MeshComponentInfo->Actor.Get());
+		MeshComponentInfo->Component->PrimaryComponentTick.bStartWithTickEnabled = false;
+		MeshComponentInfo->Component->PrimaryComponentTick.bCanEverTick = false;
 
-	MeshComponentInfo->Component->SetAnimationMode(EAnimationMode::AnimationCustomMode);
-	MeshComponentInfo->Component->RegisterComponentWithWorld(World);
+		MeshComponentInfo->Actor->AddInstanceComponent(MeshComponentInfo->Component);
 
-	return MeshComponentInfo;
+		MeshComponentInfo->Component->SetAnimationMode(EAnimationMode::AnimationCustomMode);
+		MeshComponentInfo->Component->RegisterComponentWithWorld(World);
+		return MeshComponentInfo;
+	}
+	return nullptr;
 }
 
 UAnimInstance* FRewindDebuggerAnimation::SpawnAnimInstance(uint64 ObjectId, const IGameplayProvider* GameplayProvider)
@@ -365,8 +368,11 @@ void FRewindDebuggerAnimation::Update(float DeltaTime, IRewindDebugger* RewindDe
 									bLoadMesh = true;
 									LastScrubTime = -1;
 								}
-								
-								MeshComponent = MeshComponentInfo->Component;
+
+								if (MeshComponentInfo)
+								{
+									MeshComponent = MeshComponentInfo->Component;
+								}
 							}
 
 							if (MeshComponent)
diff --git a/Engine/Plugins/Animation/GameplayInsights/Source/RewindDebugger/Private/RewindDebuggerCamera.cpp b/Engine/Plugins/Animation/GameplayInsights/Source/RewindDebugger/Private/RewindDebuggerCamera.cpp
--- a/Engine/Plugins/Animation/GameplayInsights/Source/RewindDebugger/Private/RewindDebuggerCamera.cpp
+++ b/Engine/Plugins/Animation/GameplayInsights/Source/RewindDebugger/Private/RewindDebuggerCamera.cpp
@@ -148,14 +148,20 @@ void FRewindDebuggerCamera::Update(float DeltaTime, IRewindDebugger* RewindDebug
 									{
 										FActorSpawnParameters SpawnParameters;
 										SpawnParameters.ObjectFlags |= RF_Transient;
-										CameraActor = RewindDebugger->GetWorldToVisualize()->SpawnActor<ACameraActor>(ViewMessage.Position, ViewMessage.Rotation, SpawnParameters);
-										CameraActor->SetActorLabel("RewindDebuggerCamera");
+										if (UWorld* World = RewindDebugger->GetWorldToVisualize())
+										{
+											CameraActor = World->SpawnActor<ACameraActor>(ViewMessage.Position, ViewMessage.Rotation, SpawnParameters);
+											CameraActor->SetActorLabel("RewindDebuggerCamera");
+										}
 									}
 
-									UCameraComponent* Camera = CameraActor->GetCameraComponent();
-									Camera->SetWorldLocationAndRotation(ViewMessage.Position, ViewMessage.Rotation);
-									Camera->SetFieldOfView(ViewMessage.Fov);
-									Camera->SetAspectRatio(ViewMessage.AspectRatio);
+									if (CameraActor.IsValid())
+									{
+										UCameraComponent* Camera = CameraActor->GetCameraComponent();
+										Camera->SetWorldLocationAndRotation(ViewMessage.Position, ViewMessage.Rotation);
+										Camera->SetFieldOfView(ViewMessage.Fov);
+										Camera->SetAspectRatio(ViewMessage.AspectRatio);
+									}
 
 									bCameraTraceDataFound = true;
 
diff --git a/Engine/Plugins/Experimental/Chooser/Source/Chooser/Internal/ObjectColumn.h b/Engine/Plugins/Experimental/Chooser/Source/Chooser/Internal/ObjectColumn.h
--- a/Engine/Plugins/Experimental/Chooser/Source/Chooser/Internal/ObjectColumn.h
+++ b/Engine/Plugins/Experimental/Chooser/Source/Chooser/Internal/ObjectColumn.h
@@ -20,6 +20,7 @@ struct CHOOSER_API FObjectContextProperty : public FChooserParameterObjectBase
 	FChooserObjectPropertyBinding Binding;
 
 	virtual bool GetValue(FChooserEvaluationContext& Context, FSoftObjectPath& OutResult) const override;
+	virtual bool GetValue(FChooserEvaluationContext& Context, UObject* OutResult) const override;
 	virtual bool SetValue(FChooserEvaluationContext& Context, UObject* OutResult) const override;
 
 	CHOOSER_PARAMETER_BOILERPLATE();
diff --git a/Engine/Plugins/Experimental/Chooser/Source/Chooser/Private/ObjectColumn.cpp b/Engine/Plugins/Experimental/Chooser/Source/Chooser/Private/ObjectColumn.cpp
--- a/Engine/Plugins/Experimental/Chooser/Source/Chooser/Private/ObjectColumn.cpp
+++ b/Engine/Plugins/Experimental/Chooser/Source/Chooser/Private/ObjectColumn.cpp
@@ -34,6 +34,32 @@ bool FObjectContextProperty::GetValue(FChooserEvaluationContext& Context, FSoftO
 	return false;
 }
 
+bool FObjectContextProperty::GetValue(FChooserEvaluationContext& Context, UObject* OutResult) const
+{
+	UE::Chooser::FResolvedPropertyChainResult Result;
+	if (UE::Chooser::ResolvePropertyChain(Context, Binding, Result))
+	{
+		if (Result.Function == nullptr)
+		{
+			// if the property is a soft object property, get the path directly
+			if (Result.PropertyType == UE::Chooser::EChooserPropertyAccessType::SoftObjectRef)
+			{
+				const FSoftObjectPtr& SoftObjectPtr = *reinterpret_cast<const FSoftObjectPtr*>(Result.Container + Result.PropertyOffset);
+				OutResult = SoftObjectPtr.Get();
+				return true;
+			}
+		
+			// otherwise get the value from the object property and convert to a soft object path
+			UObject** LoadedObject = reinterpret_cast<UObject**>(Result.Container + Result.PropertyOffset);
+			OutResult = *LoadedObject;
+			return true;
+		}
+	}
+
+	return false;
+}
+
+
 bool FObjectContextProperty::SetValue(FChooserEvaluationContext& Context, UObject* OutputObject) const
 {
 	UE::Chooser::FResolvedPropertyChainResult Result;
diff --git a/Engine/Plugins/Experimental/Chooser/Source/Chooser/Public/IChooserParameterObject.h b/Engine/Plugins/Experimental/Chooser/Source/Chooser/Public/IChooserParameterObject.h
--- a/Engine/Plugins/Experimental/Chooser/Source/Chooser/Public/IChooserParameterObject.h
+++ b/Engine/Plugins/Experimental/Chooser/Source/Chooser/Public/IChooserParameterObject.h
@@ -14,6 +14,7 @@ struct FChooserParameterObjectBase : public FChooserParameterBase
 	GENERATED_BODY()
 
 	virtual bool GetValue(FChooserEvaluationContext& Context, FSoftObjectPath& OutResult) const { return false; }
+	virtual bool GetValue(FChooserEvaluationContext& Context, UObject* OutResult) const { return false; }
 	virtual bool SetValue(FChooserEvaluationContext& Context, UObject* OutResult) const { return false; }
 
 #if WITH_EDITOR
