Crash fix for Rewind Debugger details panels sometimes having a reference to an invalid Analysis Session

diff --git a/Engine/Plugins/Animation/GameplayInsights/Source/RewindDebugger/Private/RewindDebugger.cpp b/Engine/Plugins/Animation/GameplayInsights/Source/RewindDebugger/Private/RewindDebugger.cpp
--- a/Engine/Plugins/Animation/GameplayInsights/Source/RewindDebugger/Private/RewindDebugger.cpp
+++ b/Engine/Plugins/Animation/GameplayInsights/Source/RewindDebugger/Private/RewindDebugger.cpp
@@ -422,6 +422,8 @@ void FRewindDebugger::StartRecording()
 	{
 		return;
 	}
+
+	ClearTrace();
 	
 	// Clear caches
 #if OBJECT_TRACE_ENABLED
@@ -510,6 +512,8 @@ void FRewindDebugger::ClearTrace()
 	CurrentTraceRange.SetUpperBoundValue(0);
 	RecordingDuration.Set(0.0);
 	SetCurrentScrubTime(0.0);
+
+	ComponentSelectionChanged(nullptr);
 	
 	// update extensions
 	IterateExtensions([this](IRewindDebuggerExtension* Extension)
