Delay rewind debugger plugin initialization

diff --git a/Engine/Plugins/Animation/GameplayInsights/GameplayInsights.uplugin b/Engine/Plugins/Animation/GameplayInsights/GameplayInsights.uplugin
--- a/Engine/Plugins/Animation/GameplayInsights/GameplayInsights.uplugin
+++ b/Engine/Plugins/Animation/GameplayInsights/GameplayInsights.uplugin
@@ -32,12 +32,12 @@
 		{
 			"Name": "RewindDebugger",
 			"Type": "Editor",
-			"LoadingPhase": "Default"
+			"LoadingPhase": "PostEngineInit"
 		},
 		{
 			"Name": "RewindDebuggerVLog",
 			"Type": "Editor",
-			"LoadingPhase": "Default"
+			"LoadingPhase": "PostEngineInit"
 		}
 
 	]
