Rewind debugger camera fixes:
- follow target actor fixed for loaded traces
- fix follow target actor teleporting on pause

diff --git a/Engine/Plugins/Animation/GameplayInsights/Source/RewindDebugger/Private/RewindDebugger.cpp b/Engine/Plugins/Animation/GameplayInsights/Source/RewindDebugger/Private/RewindDebugger.cpp
--- a/Engine/Plugins/Animation/GameplayInsights/Source/RewindDebugger/Private/RewindDebugger.cpp
+++ b/Engine/Plugins/Animation/GameplayInsights/Source/RewindDebugger/Private/RewindDebugger.cpp
@@ -946,6 +946,31 @@ const TraceServices::IAnalysisSession* FRewindDebugger::GetAnalysisSession() con
 	return UnrealInsightsModule ? UnrealInsightsModule->GetAnalysisSession().Get() : nullptr;
 }
 
+const FObjectInfo* FRewindDebugger::FindOwningActorInfo(const IGameplayProvider* GameplayProvider, uint64 ObjectId)
+{
+	const FClassInfo* ActorClassInfo = GameplayProvider->FindClassInfo(*AActor::StaticClass()->GetPathName());
+	
+	while(true)
+	{
+		const FObjectInfo& ObjectInfo = GameplayProvider->GetObjectInfo(ObjectId);
+		if (GameplayProvider->IsSubClassOf(ObjectInfo.ClassId, ActorClassInfo->Id))
+		{
+			return &ObjectInfo;
+		}
+		else
+		{
+			if (ObjectInfo.OuterId != 0)
+			{
+				ObjectId = ObjectInfo.OuterId;
+			}
+			else
+			{
+				return nullptr;
+			}
+		}
+	}
+}
+
 void FRewindDebugger::Tick(float DeltaTime)
 {
 	TRACE_CPUPROFILER_EVENT_SCOPE(FRewindDebugger::Tick);
@@ -953,6 +978,7 @@ void FRewindDebugger::Tick(float DeltaTime)
 	if (bTraceJustConnected)
 	{
 		bTraceJustConnected = false;
+		bTargetActorPositionValid = false;
 		
 		UE::Trace::ToggleChannel(TEXT("Object"), true);
 		UE::Trace::ToggleChannel(TEXT("ObjectProperties"), true);
@@ -1017,6 +1043,7 @@ void FRewindDebugger::Tick(float DeltaTime)
 					SetCurrentScrubTime(RecordingDurationValue);
 					TrackCursorDelegate.ExecuteIfBound(false);
 				}
+				bTargetActorPositionValid = false;
 			}
 			else
 			{
@@ -1042,55 +1069,54 @@ void FRewindDebugger::Tick(float DeltaTime)
 					const double CurrentTraceTime = TraceTime.Get();
 					if (CurrentTraceTime != PreviousTraceTime)
 					{
+						TRACE_CPUPROFILER_EVENT_SCOPE(FRewindDebugger::Tick_UpdateActorPosition);
 						PreviousTraceTime = CurrentTraceTime;
-						
+
 						const TraceServices::IFrameProvider& FrameProvider = TraceServices::ReadFrameProvider(*Session);
 						TraceServices::FFrame Frame;
 						if (FrameProvider.GetFrameFromTime(ETraceFrameType::TraceFrameType_Game, CurrentTraceTime, Frame))
 						{
+							bool bNewActor = false;
+							if (!TargetObjectIds.Contains(TargetActorIdForMesh))
 							{
-								TRACE_CPUPROFILER_EVENT_SCOPE(FRewindDebugger::Tick_UpdateActorPosition);
-								// until we have actor transforms traced out, the first skeletal mesh component transform on the target actor be used as as the actor position
-
-								for(uint64 TargetActorId : TargetObjectIds)
+								AnimationProvider->EnumerateSkeletalMeshPoseTimelines([this, &bNewActor, GameplayProvider](uint64 ObjectId, const IAnimationProvider::SkeletalMeshPoseTimeline& TimelineData)
 								{
-#if OBJECT_TRACE_ENABLED
-									if(UObject* ObjectInstance = FObjectTrace::GetObjectFromId(TargetActorId))
+									// until we have actor transforms traced out, the first skeletal mesh component transform on the target actor be used as as the actor position
+
+									if (const FObjectInfo* ActorInfo = FindOwningActorInfo(GameplayProvider, ObjectId))
 									{
-										if (AActor* TargetActor = Cast<AActor>(ObjectInstance))
+										if (TargetObjectIds.Contains(ActorInfo->Id))
 										{
-											TInlineComponentArray<USkeletalMeshComponent*> SkeletalMeshComponents;
-											TargetActor->GetComponents(SkeletalMeshComponents);
-
-											if (SkeletalMeshComponents.Num() > 0)
-											{
-												int64 ObjectId = FObjectTrace::GetObjectId(SkeletalMeshComponents[0]);
-
-												AnimationProvider->ReadSkeletalMeshPoseTimeline(ObjectId, [this, Frame, ObjectId, AnimationProvider](const IAnimationProvider::SkeletalMeshPoseTimeline& TimelineData, bool bHasCurves)
-												{
-													const FSkeletalMeshPoseMessage * PoseMessage = nullptr;
-
-													// Get last pose in frame
-													TimelineData.EnumerateEvents(Frame.StartTime, Frame.EndTime,
-														[&PoseMessage](double InStartTime, double InEndTime, uint32 InDepth, const FSkeletalMeshPoseMessage& InPoseMessage)
-														{
-															PoseMessage = &InPoseMessage;
-															return TraceServices::EEventEnumerate::Continue;
-														});
-
-													// Update position based on pose
-													if (PoseMessage)
-													{
-														bTargetActorPositionValid = true;
-														TargetActorPosition = PoseMessage->ComponentToWorld.GetTranslation();
-													}
-												});
-											}
+											bNewActor = true;
+											TargetActorIdForMesh = ActorInfo->Id;
+											TargetActorMeshId = ObjectId;
 										}
 									}
-#endif // OBJECT_TRACE_ENABLED
-								}
+									
+								});
 							}
+						
+						
+							AnimationProvider->ReadSkeletalMeshPoseTimeline(TargetActorMeshId, [this, &Frame, bNewActor](const IAnimationProvider::SkeletalMeshPoseTimeline& TimelineData, bool bHasCurves)
+							{
+								const FSkeletalMeshPoseMessage * PoseMessage = nullptr;
+
+								// Get last pose in frame
+								TimelineData.EnumerateEvents(Frame.StartTime, Frame.EndTime,
+									[&PoseMessage](double InStartTime, double InEndTime, uint32 InDepth, const FSkeletalMeshPoseMessage& InPoseMessage)
+									{
+										PoseMessage = &InPoseMessage;
+										return TraceServices::EEventEnumerate::Continue;
+									});
+
+								// Update position based on pose
+								if (PoseMessage)
+								{
+									// mark the target position as invalid for a frame when the actor changes, so it will be treated as a teleport by the camera system
+									bTargetActorPositionValid = !bNewActor;
+									TargetActorPosition = PoseMessage->ComponentToWorld.GetTranslation();
+								}
+							});
 						}
 					}
 				}
diff --git a/Engine/Plugins/Animation/GameplayInsights/Source/RewindDebugger/Private/RewindDebugger.h b/Engine/Plugins/Animation/GameplayInsights/Source/RewindDebugger/Private/RewindDebugger.h
--- a/Engine/Plugins/Animation/GameplayInsights/Source/RewindDebugger/Private/RewindDebugger.h
+++ b/Engine/Plugins/Animation/GameplayInsights/Source/RewindDebugger/Private/RewindDebugger.h
@@ -20,6 +20,9 @@ namespace TraceServices
 class SWidget;
 class SDockTab;
 
+struct FObjectInfo;
+class IGameplayProvider;
+
 // Singleton class that handles the logic for the Rewind Debugger
 // handles:
 //  Playback/Scrubbing state
@@ -139,6 +142,8 @@ public:
 	virtual void OpenDetailsPanel() override;
 	void SetIsDetailsPanelOpen(bool bIsOpen) { bIsDetailsPanelOpen = bIsOpen; }
 	bool IsDetailsPanelOpen(bool bIsOpen) { return bIsDetailsPanelOpen; }
+	
+	static const FObjectInfo* FindOwningActorInfo(const IGameplayProvider* GameplayProvider, uint64 ObjectId);
 
 	TArrayView<RewindDebugger::FRewindDebuggerTrackType> GetTrackTypes() { return TrackTypes; };
 
@@ -201,11 +206,13 @@ private:
 
 	TArray<uint64> TargetObjectIds;
 
-	mutable class IUnrealInsightsModule *UnrealInsightsModule;
+	mutable class IUnrealInsightsModule *UnrealInsightsModule = 0;
 	FTSTicker::FDelegateHandle TickerHandle;
 
 	bool bTargetActorPositionValid = false;
 	FVector TargetActorPosition;
+	uint64 TargetActorMeshId = 0;
+	uint64 TargetActorIdForMesh = 0;
 
 	TArray<RewindDebugger::FRewindDebuggerTrackType> TrackTypes;
 
diff --git a/Engine/Plugins/Animation/GameplayInsights/Source/RewindDebugger/Private/RewindDebuggerAnimation.cpp b/Engine/Plugins/Animation/GameplayInsights/Source/RewindDebugger/Private/RewindDebuggerAnimation.cpp
--- a/Engine/Plugins/Animation/GameplayInsights/Source/RewindDebugger/Private/RewindDebuggerAnimation.cpp
+++ b/Engine/Plugins/Animation/GameplayInsights/Source/RewindDebugger/Private/RewindDebuggerAnimation.cpp
@@ -8,7 +8,7 @@
 #include "IAnimationProvider.h"
 #include "IGameplayProvider.h"
 #include "Insights/IUnrealInsightsModule.h"
-#include "IRewindDebugger.h"
+#include "RewindDebugger.h"
 #include "LevelEditor.h"
 #include "Modules/ModuleManager.h"
 #include "SLevelViewport.h"
@@ -193,34 +193,9 @@ void FRewindDebuggerAnimation::ApplyPoseToMesh(const IAnimationProvider* Animati
 	}
 }
 
-const FObjectInfo* FindOwningActorInfo(const IGameplayProvider* GameplayProvider, uint64 ObjectId)
-{
-	const FClassInfo* ActorClassInfo = GameplayProvider->FindClassInfo(*AActor::StaticClass()->GetPathName());
-	
-	while(true)
-	{
-		const FObjectInfo& ObjectInfo = GameplayProvider->GetObjectInfo(ObjectId);
-		if (GameplayProvider->IsSubClassOf(ObjectInfo.ClassId, ActorClassInfo->Id))
-		{
-			return &ObjectInfo;
-		}
-		else
-		{
-			if (ObjectInfo.OuterId != 0)
-			{
-				ObjectId = ObjectInfo.OuterId;
-			}
-			else
-			{
-				return nullptr;
-			}
-		}
-	}
-}
-
 FRewindDebuggerAnimation::FSpawnedMeshComponentInfo* FRewindDebuggerAnimation::SpawnMesh(uint64 ObjectId, const IGameplayProvider* GameplayProvider)
 {
-	IRewindDebugger* RewindDebugger = IRewindDebugger::Instance();
+	const IRewindDebugger* RewindDebugger = IRewindDebugger::Instance();
 	check(RewindDebugger);
 	if (UWorld* World = RewindDebugger->GetWorldToVisualize())
 	{
@@ -232,7 +207,7 @@ FRewindDebuggerAnimation::FSpawnedMeshComponentInfo* FRewindDebuggerAnimation::S
 
 		MeshComponentInfo->Actor = World->SpawnActor<AActor>(ActorSpawnParameters);
 
-		if(const FObjectInfo* ActorInfo = FindOwningActorInfo(GameplayProvider, ObjectId))
+		if(const FObjectInfo* ActorInfo = FRewindDebugger::FindOwningActorInfo(GameplayProvider, ObjectId))
 		{
 			MeshComponentInfo->Actor->SetActorLabel(FString(TEXT("RewindDebugger: ") + FString(ActorInfo->Name)));
 		}
diff --git a/Engine/Plugins/Animation/GameplayInsights/Source/RewindDebugger/Private/RewindDebuggerCamera.cpp b/Engine/Plugins/Animation/GameplayInsights/Source/RewindDebugger/Private/RewindDebuggerCamera.cpp
--- a/Engine/Plugins/Animation/GameplayInsights/Source/RewindDebugger/Private/RewindDebuggerCamera.cpp
+++ b/Engine/Plugins/Animation/GameplayInsights/Source/RewindDebugger/Private/RewindDebuggerCamera.cpp
@@ -17,7 +17,7 @@
 #define LOCTEXT_NAMESPACE "RewindDebuggerCamera"
 
 FRewindDebuggerCamera::FRewindDebuggerCamera()
-	: LastPositionValid(false)
+	: bLastPositionValid(false)
 {
 }
 
@@ -98,9 +98,10 @@ void FRewindDebuggerCamera::Update(float DeltaTime, IRewindDebugger* RewindDebug
 {
 	if (RewindDebugger->IsPIESimulating() || RewindDebugger->GetRecordingDuration() == 0.0)
 	{
+		bLastPositionValid = false;
 		return;
 	}
-
+	
 	if (const TraceServices::IAnalysisSession* Session = RewindDebugger->GetAnalysisSession())
 	{
 		TraceServices::FAnalysisSessionReadScope SessionReadScope(*Session);
@@ -125,7 +126,7 @@ void FRewindDebuggerCamera::Update(float DeltaTime, IRewindDebugger* RewindDebug
 					// Follow Actor mode: apply position changes from the target actor to the camera
 					if (bTargetActorPositionValid)
 					{
-						if(LastPositionValid)
+						if(bLastPositionValid)
 						{
 							LevelViewportClient.SetViewLocation(LevelViewportClient.GetViewLocation() + TargetActorPosition - LastPosition);
 						}
@@ -180,7 +181,7 @@ void FRewindDebuggerCamera::Update(float DeltaTime, IRewindDebugger* RewindDebug
 				}
 
 				LastPosition = TargetActorPosition;
-				LastPositionValid = bTargetActorPositionValid;
+				bLastPositionValid = bTargetActorPositionValid;
 
 				if (bCameraTraceDataFound) // don't update this if there was no trace data found, because when first pausing, it can take a few frames for latest data to get processed
 				{
diff --git a/Engine/Plugins/Animation/GameplayInsights/Source/RewindDebugger/Private/RewindDebuggerCamera.h b/Engine/Plugins/Animation/GameplayInsights/Source/RewindDebugger/Private/RewindDebuggerCamera.h
--- a/Engine/Plugins/Animation/GameplayInsights/Source/RewindDebugger/Private/RewindDebuggerCamera.h
+++ b/Engine/Plugins/Animation/GameplayInsights/Source/RewindDebugger/Private/RewindDebuggerCamera.h
@@ -27,9 +27,8 @@ public:
 	void SetCameraMode(ERewindDebuggerCameraMode Mode);
 	
 private:
-
 	
-	bool LastPositionValid;
+	bool bLastPositionValid;
 	FVector LastPosition;
 	
 	TWeakObjectPtr<ACameraActor> CameraActor; 
