GameplayInsights - fix for SCurveTimelineView::PaintCurve connecting the first and second points even if further away than LargeFrameTime

diff --git a/Engine/Plugins/Animation/GameplayInsights/Source/GameplayInsights/Private/SCurveTimelineView.cpp b/Engine/Plugins/Animation/GameplayInsights/Source/GameplayInsights/Private/SCurveTimelineView.cpp
--- a/Engine/Plugins/Animation/GameplayInsights/Source/GameplayInsights/Private/SCurveTimelineView.cpp
+++ b/Engine/Plugins/Animation/GameplayInsights/Source/GameplayInsights/Private/SCurveTimelineView.cpp
@@ -206,26 +206,44 @@ int32 SCurveTimelineView::PaintCurve(const FGeometry& AllottedGeometry, const FS
 			else
 			{
 				// Lines rendering
-
 				TArray<FVector2D> Points;
 				Points.Reserve(NumPoints);
+		 
+				auto MakeLines = [&OutDrawElements, &LayerId, &AllottedGeometry, &Points, &LineDrawEffects, &LineColor]()
+					{
+						if (Points.IsEmpty())
+						{
+							return;
+						}
+
+						if (Points.Num() == 1)
+						{
+							// adding an additional point AdditionalPixels apart to be able to draw a line
+							static const float AdditionalPixels = 2.f;
+							const FVector2D AdditionalPoint(Points[0].X + AdditionalPixels, Points[0].Y);
+							Points.Add(AdditionalPoint);
+						}
+
+						FSlateDrawElement::MakeLines(
+							OutDrawElements,
+							LayerId++,
+							AllottedGeometry.ToPaintGeometry(),
+							Points,
+							LineDrawEffects,
+							LineColor,
+							false
+						);
+					};
 
 				float PrevTime = CurvePoints[0].Time;
 				for(int i=0; i<NumPoints; i++)
 				{
 					const FTimelineCurveData::CurvePoint& Point = CurvePoints[i];
-					if (Point.Time - PrevTime > LargeFrameTime && Points.Num()>1)
+					if (Point.Time - PrevTime > LargeFrameTime)
 					{
 						// break the line list - data has stopped and started again
-						FSlateDrawElement::MakeLines(
-									OutDrawElements,
-											LayerId++,
-											AllottedGeometry.ToPaintGeometry(),
-											Points,
-											LineDrawEffects,
-											LineColor,
-											false
-											);
+						MakeLines();
+						
 						Points.SetNum(0,EAllowShrinking::No);
 					}
 			
@@ -236,15 +254,8 @@ int32 SCurveTimelineView::PaintCurve(const FGeometry& AllottedGeometry, const FS
 					Points.Add(FVector2D(X,Y));
 				}
 			
-				FSlateDrawElement::MakeLines(
-							OutDrawElements,
-									LayerId,
-									AllottedGeometry.ToPaintGeometry(),
-									Points,
-									LineDrawEffects,
-									LineColor,
-									false
-									);
+				MakeLines();
+				
 			}
 		
 		}
