Rewind Debugger support for making recordings from standalone games

diff --git a/Engine/Plugins/Animation/GameplayInsights/GameplayInsights.uplugin b/Engine/Plugins/Animation/GameplayInsights/GameplayInsights.uplugin
--- a/Engine/Plugins/Animation/GameplayInsights/GameplayInsights.uplugin
+++ b/Engine/Plugins/Animation/GameplayInsights/GameplayInsights.uplugin
@@ -38,7 +38,16 @@
 			"Name": "RewindDebuggerVLog",
 			"Type": "Editor",
 			"LoadingPhase": "PostEngineInit"
+		},
+		{
+			"Name": "RewindDebuggerRuntime",
+			"Type": "Runtime",
+			"LoadingPhase": "Default"
+		},
+		{
+			"Name": "RewindDebuggerVLogRuntime",
+			"Type": "Runtime",
+			"LoadingPhase": "Default"
 		}
-
 	]
 }
diff --git a/Engine/Plugins/Animation/GameplayInsights/Source/RewindDebugger/Private/RewindDebugger.cpp b/Engine/Plugins/Animation/GameplayInsights/Source/RewindDebugger/Private/RewindDebugger.cpp
--- a/Engine/Plugins/Animation/GameplayInsights/Source/RewindDebugger/Private/RewindDebugger.cpp
+++ b/Engine/Plugins/Animation/GameplayInsights/Source/RewindDebugger/Private/RewindDebugger.cpp
@@ -39,6 +39,7 @@
 #include "UnrealEdGlobals.h"
 #include "Editor/UnrealEdEngine.h"
 #include "Kismet2/DebuggerCommands.h"
+#include "RewindDebuggerRuntime/RewindDebuggerRuntime.h"
 #include "TraceServices/AnalysisService.h"
 #include "TraceServices/ITraceServicesModule.h"
 #include "Widgets/Input/SNumericEntryBox.h"
@@ -70,6 +71,18 @@ static void TraceSubobjects(UObject* OuterObject)
 
 FRewindDebugger::FRewindDebugger()
 {
+	if (RewindDebugger::FRewindDebuggerRuntime::Instance() == nullptr)
+	{
+		RewindDebugger::FRewindDebuggerRuntime::Initialize();
+	}
+
+	if (RewindDebugger::FRewindDebuggerRuntime* Runtime = RewindDebugger::FRewindDebuggerRuntime::Instance())
+	{
+		Runtime->ClearRecording.AddRaw(this, &FRewindDebugger::OnClearRecording);
+		Runtime->RecordingStarted.AddRaw(this, &FRewindDebugger::OnRecordingStarted);
+		Runtime->RecordingStarted.AddRaw(this, &FRewindDebugger::OnRecordingStopped);
+	}
+	
 	RewindDebugger::FRewindDebuggerTrackCreators::EnumerateCreators([this](const RewindDebugger::IRewindDebuggerTrackCreator* Creator)
     {
 		Creator->GetTrackTypes(TrackTypes);
@@ -136,6 +149,11 @@ FRewindDebugger::~FRewindDebugger()
 	FEditorDelegates::SingleStepPIE.RemoveAll(this);
 
 	FTSTicker::GetCoreTicker().RemoveTicker(TickerHandle);
+	
+	if (RewindDebugger::FRewindDebuggerRuntime* Runtime = RewindDebugger::FRewindDebuggerRuntime::Instance())
+	{
+		Runtime->RecordingStarted.RemoveAll(this);
+	}
 }
 
 void FRewindDebugger::Initialize() 
@@ -174,7 +192,7 @@ void FRewindDebugger::OnPIEPaused(bool bSimulating)
 	bPIESimulating = false;
 	ControlState = EControlState::Pause;
 	
-	if (bRecording)
+	if (IsRecording())
 	{
     #if OBJECT_TRACE_ENABLED
 		UWorld* World = GetWorldToVisualize();
@@ -209,7 +227,7 @@ void FRewindDebugger::OnPIEResumed(bool bSimulating)
 
 void FRewindDebugger::OnPIESingleStepped(bool bSimulating)
 {
-	if (bRecording)
+	if (IsRecording())
 	{
     #if OBJECT_TRACE_ENABLED
 		UWorld* World = GetWorldToVisualize();
@@ -423,32 +441,44 @@ void FRewindDebugger::StartRecording()
 		return;
 	}
 
+	if (RewindDebugger::FRewindDebuggerRuntime* Runtime = RewindDebugger::FRewindDebuggerRuntime::Instance())
+	{
+		Runtime->StartRecording();
+	}
+}
+
+void FRewindDebugger::OnClearRecording()
+{
 	ClearTrace();
-	
-	// Clear caches
-#if OBJECT_TRACE_ENABLED
-	FObjectTrace::Reset();
-	FAnimTrace::Reset();
-#endif
-	
 	RecordingDuration.Set(0);
-	// RecordingIndex++;
-	bRecording = true;
-
-	// Disable all trace channels, and then enable only the ones needed by RewindDebugger
-	// for systems with RewindDebugger integration, they should enable their channel(s) in an Extension in "RecordingStarted"
-	DisableAllTraceChannels();
-
-	// Clear all buffered data and prevent data from previous recordings from leaking into the new recording
-	FTraceAuxiliary::FOptions Options;
-	Options.bExcludeTail = true;
+	TargetObjectIds.Empty(2);
+	bTargetActorPositionValid = false;
 
-	FTraceAuxiliary::OnConnection.AddRaw(this, &FRewindDebugger::OnConnection); 
+	IterateExtensions([this](IRewindDebuggerExtension* Extension)
+	{
+		Extension->Clear(this);
+	}
+	);
+}
 
-	FTraceAuxiliary::Start(FTraceAuxiliary::EConnectionType::Network, TEXT("127.0.0.1"), TEXT(""), &Options, LogRewindDebugger);
+void FRewindDebugger::OnRecordingStarted()
+{
+	IterateExtensions([this](IRewindDebuggerExtension* Extension)
+	{
+		Extension->RecordingStarted(this);
+	}
+	);
+	
 	UnrealInsightsModule->StartAnalysisForLastLiveSession(5.0);
+}
 
-	TargetObjectIds.Empty(2);
+void FRewindDebugger::OnRecordingStopped()
+{
+	IterateExtensions([this](IRewindDebuggerExtension* Extension)
+	{
+		Extension->RecordingStopped(this);
+	}
+	);
 }
 
 bool FRewindDebugger::CanOpenTrace() const
@@ -611,19 +641,9 @@ void FRewindDebugger::SetShouldAutoEject(bool value)
 
 void FRewindDebugger::StopRecording()
 {
-	if (bRecording)
+	if (RewindDebugger::FRewindDebuggerRuntime* Runtime = RewindDebugger::FRewindDebuggerRuntime::Instance())
 	{
-		// update extensions
-		IterateExtensions([this](IRewindDebuggerExtension* Extension)
-			{
-				Extension->RecordingStopped(this);
-			}
-		);
-
-		bRecording = false;
-		
-		DisableAllTraceChannels();
-		FTraceAuxiliary::Stop();
+		Runtime->StopRecording();
 	}
 }
 
@@ -779,6 +799,15 @@ UWorld* FRewindDebugger::GetWorldToVisualize() const
 	return World;
 }
 
+bool FRewindDebugger::IsRecording() const
+{
+	if (RewindDebugger::FRewindDebuggerRuntime* Runtime = RewindDebugger::FRewindDebuggerRuntime::Instance())
+	{
+		return Runtime->IsRecording();
+	}
+	return false;
+}
+
 bool FRewindDebugger::IsTraceFileLoaded() const
 {
 	return GetAnalysisSession()!=nullptr && !bPIEStarted;
@@ -975,45 +1004,6 @@ void FRewindDebugger::Tick(float DeltaTime)
 {
 	TRACE_CPUPROFILER_EVENT_SCOPE(FRewindDebugger::Tick);
 
-	if (bTraceJustConnected)
-	{
-		bTraceJustConnected = false;
-		bTargetActorPositionValid = false;
-		
-		UE::Trace::ToggleChannel(TEXT("Object"), true);
-		UE::Trace::ToggleChannel(TEXT("ObjectProperties"), true);
-		UE::Trace::ToggleChannel(TEXT("Frame"), true);
-
-		// update extensions
-		IterateExtensions([this](IRewindDebuggerExtension* Extension)
-			{
-				Extension->RecordingStarted(this);
-			}
-		);
-
-	#if OBJECT_TRACE_ENABLED
-		// trace each play-in-editor world, and all the actors in it.
-		for (TObjectIterator<UWorld> World; World; ++World)
-		{
-			if (World->IsPlayInEditor())
-			{
-				FObjectTrace::ResetWorldElapsedTime(*World);
-				FObjectTrace::SetWorldRecordingIndex(*World, RecordingIndex);
-		 		
-				TRACE_WORLD(*World);
-					
-				for (TActorIterator<AController> Iterator(*World); Iterator; ++Iterator)
-				{
-					if (APawn* Pawn = Iterator->GetPawn())
-					{
-						TRACE_PAWN_POSSESS(static_cast<UObject*>(*Iterator), static_cast<UObject*>(Pawn));
-					}
-				}
-		 	}
-		}
-	#endif // OBJECT_TRACE_ENABLED
-	}
-	
 	if (const TraceServices::IAnalysisSession* Session = GetAnalysisSession())
 	{
 		const IAnimationProvider* AnimationProvider = Session->ReadProvider<IAnimationProvider>("AnimationProvider");
@@ -1037,7 +1027,7 @@ void FRewindDebugger::Tick(float DeltaTime)
 
 			if (bPIESimulating)
 			{
-				if (bRecording)
+				if (IsRecording())
 				{
 					TRACE_CPUPROFILER_EVENT_SCOPE(FRewindDebugger::Tick_UpdateSimulating);
 					SetCurrentScrubTime(RecordingDurationValue);
diff --git a/Engine/Plugins/Animation/GameplayInsights/Source/RewindDebugger/Private/RewindDebugger.h b/Engine/Plugins/Animation/GameplayInsights/Source/RewindDebugger/Private/RewindDebugger.h
--- a/Engine/Plugins/Animation/GameplayInsights/Source/RewindDebugger/Private/RewindDebugger.h
+++ b/Engine/Plugins/Animation/GameplayInsights/Source/RewindDebugger/Private/RewindDebugger.h
@@ -44,7 +44,7 @@ public:
 	virtual uint64 GetTargetActorId() const override;
 	virtual bool GetTargetActorPosition(FVector& OutPosition) const override;
 	virtual UWorld* GetWorldToVisualize() const override;
-	virtual bool IsRecording() const override { return bRecording; }
+	virtual bool IsRecording() const override;
 	virtual bool IsPIESimulating() const override { return bPIESimulating; }
 	virtual bool IsTraceFileLoaded() const override;
 	virtual double GetRecordingDuration() const override { return RecordingDuration.Get(); }
@@ -69,6 +69,9 @@ public:
 
 	// Start a new Recording:  Start tracing Object + Animation data, increment the current recording index, and reset the recording elapsed time to 0
 	void StartRecording();
+	void OnClearRecording();
+	void OnRecordingStarted();
+	void OnRecordingStopped();
 	bool CanStartRecording() const { return !IsRecording() && bPIESimulating; }
 	
 	bool CanOpenTrace() const;
diff --git a/Engine/Plugins/Animation/GameplayInsights/Source/RewindDebugger/Private/RewindDebuggerAnimation.cpp b/Engine/Plugins/Animation/GameplayInsights/Source/RewindDebugger/Private/RewindDebuggerAnimation.cpp
--- a/Engine/Plugins/Animation/GameplayInsights/Source/RewindDebugger/Private/RewindDebuggerAnimation.cpp
+++ b/Engine/Plugins/Animation/GameplayInsights/Source/RewindDebugger/Private/RewindDebuggerAnimation.cpp
@@ -27,17 +27,6 @@ FRewindDebuggerAnimation::FRewindDebuggerAnimation()
 {
 }
 
-void FRewindDebuggerAnimation::RecordingStarted(IRewindDebugger*)
-{
-	ClearSpawnedComponents();
-	UE::Trace::ToggleChannel(TEXT("Animation"), true);
-}
-
-void FRewindDebuggerAnimation::RecordingStopped(IRewindDebugger*)
-{
-	UE::Trace::ToggleChannel(TEXT("Animation"), false);
-}
-
 void FRewindDebuggerAnimation::ClearSpawnedComponents()
 {
 	for(auto& MeshComponentInfo : SpawnedMeshComponents)
diff --git a/Engine/Plugins/Animation/GameplayInsights/Source/RewindDebugger/Private/RewindDebuggerAnimation.h b/Engine/Plugins/Animation/GameplayInsights/Source/RewindDebugger/Private/RewindDebuggerAnimation.h
--- a/Engine/Plugins/Animation/GameplayInsights/Source/RewindDebugger/Private/RewindDebuggerAnimation.h
+++ b/Engine/Plugins/Animation/GameplayInsights/Source/RewindDebugger/Private/RewindDebuggerAnimation.h
@@ -30,9 +30,6 @@ public:
 	void Shutdown();
 
 	virtual void Update(float DeltaTime, IRewindDebugger* RewindDebugger) override;
-	virtual void RecordingStarted(IRewindDebugger* RewindDebugger) override;
-	virtual void RecordingStopped(IRewindDebugger* RewindDebugger) override;
-	void ClearSpawnedComponents();
 	virtual void Clear(IRewindDebugger* RewindDebugger) override;
 	
 	void OnPIEResumed(bool bSimulating);
@@ -43,6 +40,7 @@ public:
 
 	static FRewindDebuggerAnimation* GetInstance() { return Instance; }
 private:
+	void ClearSpawnedComponents();
 	void ApplyPoseToMesh(const class IAnimationProvider* AnimationProvider, const IGameplayProvider* GameplayProvider, const TraceServices::FFrame& Frame,
 		const IAnimationProvider::SkeletalMeshPoseTimeline& TimelineData, USkeletalMeshComponent* MeshComponent, uint64 ObjectId, bool bQueueForReset, bool bApplyMesh);
 	
diff --git a/Engine/Plugins/Animation/GameplayInsights/Source/RewindDebugger/RewindDebugger.Build.cs b/Engine/Plugins/Animation/GameplayInsights/Source/RewindDebugger/RewindDebugger.Build.cs
--- a/Engine/Plugins/Animation/GameplayInsights/Source/RewindDebugger/RewindDebugger.Build.cs
+++ b/Engine/Plugins/Animation/GameplayInsights/Source/RewindDebugger/RewindDebugger.Build.cs
@@ -35,6 +35,7 @@ namespace UnrealBuildTool.Rules
 				"TraceInsights",
 				"UnrealEd",
 				"RewindDebuggerInterface",
+				"RewindDebuggerRuntime",
 				"ToolWidgets",
 				"MainFrame",
 				"DeveloperSettings",
diff --git a/Engine/Plugins/Animation/GameplayInsights/Source/RewindDebuggerRuntime/Private/RewindDebuggerAnimationRuntime.h b/Engine/Plugins/Animation/GameplayInsights/Source/RewindDebuggerRuntime/Private/RewindDebuggerAnimationRuntime.h
new file mode 100644
--- /dev/null
+++ b/Engine/Plugins/Animation/GameplayInsights/Source/RewindDebuggerRuntime/Private/RewindDebuggerAnimationRuntime.h
@@ -0,0 +1,25 @@
+// Copyright Epic Games, Inc. All Rights Reserved.
+
+#pragma once
+
+#include "CoreMinimal.h"
+#include "RewindDebuggerRuntimeInterface/IRewindDebuggerRuntimeExtension.h"
+
+
+namespace RewindDebugger
+{
+
+	class REWINDDEBUGGERRUNTIME_API FRewindDebuggerAnimationRuntime : public IRewindDebuggerRuntimeExtension 
+	{
+	public:
+		virtual void RecordingStarted() override
+		{
+			UE::Trace::ToggleChannel(TEXT("Animation"), true);
+		}
+		
+		virtual void RecordingStopped() override
+		{
+			UE::Trace::ToggleChannel(TEXT("Animation"), false);
+		}
+	};
+}
diff --git a/Engine/Plugins/Animation/GameplayInsights/Source/RewindDebuggerRuntime/Private/RewindDebuggerRuntime.cpp b/Engine/Plugins/Animation/GameplayInsights/Source/RewindDebuggerRuntime/Private/RewindDebuggerRuntime.cpp
new file mode 100644
--- /dev/null
+++ b/Engine/Plugins/Animation/GameplayInsights/Source/RewindDebuggerRuntime/Private/RewindDebuggerRuntime.cpp
@@ -0,0 +1,133 @@
+// Copyright Epic Games, Inc. All Rights Reserved.
+
+#include "RewindDebuggerRuntime/RewindDebuggerRuntime.h"
+
+#include "Animation/AnimTrace.h"
+#include "EngineUtils.h"
+#include "Engine/World.h"
+#include "GameFramework/Pawn.h"
+#include "IRewindDebuggerExtension.h"
+#include "ProfilingDebugging/TraceAuxiliary.h"
+#include "UObject/UObjectIterator.h"
+#include "RewindDebuggerRuntimeInterface/IRewindDebuggerRuntimeExtension.h"
+
+DEFINE_LOG_CATEGORY(LogRewindDebuggerRuntime)
+
+namespace RewindDebugger
+{
+
+	FRewindDebuggerRuntime* FRewindDebuggerRuntime::InternalInstance = nullptr;
+
+	void FRewindDebuggerRuntime::Initialize()
+	{
+		InternalInstance = new FRewindDebuggerRuntime();
+	}
+	
+	void FRewindDebuggerRuntime::Shutdown()
+	{
+		delete InternalInstance;
+		InternalInstance = nullptr;
+	}
+
+	static void IterateExtensions(TFunction<void(IRewindDebuggerRuntimeExtension* Extension)> IteratorFunction)
+	{
+		// update extensions
+		IModularFeatures& ModularFeatures = IModularFeatures::Get();
+
+		const int32 NumExtensions = ModularFeatures.GetModularFeatureImplementationCount(IRewindDebuggerRuntimeExtension::ModularFeatureName);
+		for (int32 ExtensionIndex = 0; ExtensionIndex < NumExtensions; ++ExtensionIndex)
+		{
+			IRewindDebuggerRuntimeExtension* Extension = static_cast<IRewindDebuggerRuntimeExtension*>(ModularFeatures.GetModularFeatureImplementation(IRewindDebuggerRuntimeExtension::ModularFeatureName, ExtensionIndex));
+			IteratorFunction(Extension);
+		}
+	}
+	
+	static void DisableAllTraceChannels()
+	{
+		UE::Trace::EnumerateChannels([](const ANSICHAR* ChannelName, bool bEnabled, void*)
+		{
+			if (bEnabled)
+			{
+				FString ChannelNameFString(ChannelName);
+				UE::Trace::ToggleChannel(ChannelNameFString.GetCharArray().GetData(), false);
+			}
+		}
+		, nullptr);
+	}
+	
+	void FRewindDebuggerRuntime::StartRecording()
+	{
+#if OBJECT_TRACE_ENABLED
+		// Clear caches  (maybe move these to plugins?)
+		FObjectTrace::Reset();
+		FAnimTrace::Reset();
+
+		ClearRecording.Broadcast();
+        	
+        // Disable all trace channels, and then enable only the ones needed by RewindDebugger
+        // for systems with RewindDebugger integration, they should enable their channel(s) in an Extension in "RecordingStarted"
+        DisableAllTraceChannels();
+
+		// Clear all buffered data and prevent data from previous recordings from leaking into the new recording
+		FTraceAuxiliary::FOptions Options;
+		Options.bExcludeTail = true;
+	
+		// FTraceAuxiliary::OnConnection.AddRaw(this, &FRewindDebugger::OnConnection); 
+	
+		FTraceAuxiliary::Start(FTraceAuxiliary::EConnectionType::Network, TEXT("127.0.0.1"), TEXT(""), &Options, LogRewindDebuggerRuntime);
+		
+		UE::Trace::ToggleChannel(TEXT("Object"), true);
+		UE::Trace::ToggleChannel(TEXT("ObjectProperties"), true);
+		UE::Trace::ToggleChannel(TEXT("Frame"), true);
+		
+		bIsRecording = true;
+
+		// update extensions
+		IterateExtensions([](IRewindDebuggerRuntimeExtension* Extension)
+			{
+				Extension->RecordingStarted();
+			}
+		);
+	
+		// trace each play-in-editor world, and all the actors in it.
+		for (TObjectIterator<UWorld> World; World; ++World)
+		{
+			FObjectTrace::ResetWorldElapsedTime(*World);
+			
+			TRACE_WORLD(*World);
+				
+			for (TActorIterator<AController> Iterator(*World); Iterator; ++Iterator)
+			{
+				if (APawn* Pawn = Iterator->GetPawn())
+				{
+					TRACE_PAWN_POSSESS(static_cast<UObject*>(*Iterator), static_cast<UObject*>(Pawn));
+				}
+			}
+		}
+
+		RecordingStarted.Broadcast();
+	#endif // OBJECT_TRACE_ENABLED
+	}
+	
+	void FRewindDebuggerRuntime::StopRecording()
+	{
+		if (bIsRecording)
+		{
+			// // update extensions
+			IterateExtensions([](IRewindDebuggerRuntimeExtension* Extension)
+				{
+					Extension->RecordingStopped();
+				}
+			);
+		
+			bIsRecording = false;
+		
+			DisableAllTraceChannels();
+			FTraceAuxiliary::Stop();
+		
+			RecordingStopped.Broadcast();
+		}
+	}
+
+
+}
diff --git a/Engine/Plugins/Animation/GameplayInsights/Source/RewindDebuggerRuntime/Private/RewindDebuggerRuntimeModule.cpp b/Engine/Plugins/Animation/GameplayInsights/Source/RewindDebuggerRuntime/Private/RewindDebuggerRuntimeModule.cpp
new file mode 100644
--- /dev/null
+++ b/Engine/Plugins/Animation/GameplayInsights/Source/RewindDebuggerRuntime/Private/RewindDebuggerRuntimeModule.cpp
@@ -0,0 +1,52 @@
+// Copyright Epic Games, Inc. All Rights Reserved.
+
+#include "RewindDebuggerRuntimeModule.h"
+
+#include "RewindDebuggerRuntime/RewindDebuggerRuntime.h"
+#include "Modules/ModuleManager.h"
+#include "HAL/IConsoleManager.h"
+#include "ObjectTrace.h"
+#include "Features/IModularFeatures.h"
+
+void FRewindDebuggerRuntimeModule::StartupModule()
+{
+	if (RewindDebugger::FRewindDebuggerRuntime::Instance() == nullptr)
+	{
+		RewindDebugger::FRewindDebuggerRuntime::Initialize();
+	}
+	
+	ConsoleObjects.Add(IConsoleManager::Get().RegisterConsoleCommand(
+		TEXT("RewindDebugger.StartRecording"),
+		TEXT("Starts making a rewind debugger recording."),
+		FConsoleCommandDelegate::CreateRaw(RewindDebugger::FRewindDebuggerRuntime::Instance(), &RewindDebugger::FRewindDebuggerRuntime::StartRecording),
+		ECVF_Default
+	));
+	
+	ConsoleObjects.Add(IConsoleManager::Get().RegisterConsoleCommand(
+   		TEXT("RewindDebugger.StopRecording"),
+   		TEXT("Stops the current rewind debugger recording."),
+		FConsoleCommandDelegate::CreateRaw(RewindDebugger::FRewindDebuggerRuntime::Instance(), &RewindDebugger::FRewindDebuggerRuntime::StopRecording),
+   		ECVF_Default
+   	));
+	
+	IModularFeatures::Get().RegisterModularFeature(IRewindDebuggerRuntimeExtension::ModularFeatureName, &AnimationExtension);
+}
+
+void FRewindDebuggerRuntimeModule::ShutdownModule()
+{
+	IConsoleManager &ConsoleManager = IConsoleManager::Get();
+	for (IConsoleObject *ConsoleObject : ConsoleObjects)
+	{
+		if (ConsoleObject)
+		{
+			ConsoleManager.UnregisterConsoleObject(ConsoleObject);
+		}
+	}
+	ConsoleObjects.Empty();
+	
+	RewindDebugger::FRewindDebuggerRuntime::Shutdown();
+	
+	IModularFeatures::Get().UnregisterModularFeature(IRewindDebuggerRuntimeExtension::ModularFeatureName, &AnimationExtension);
+}
+
+IMPLEMENT_MODULE(FRewindDebuggerRuntimeModule, RewindDebuggerRuntime);
diff --git a/Engine/Plugins/Animation/GameplayInsights/Source/RewindDebuggerRuntime/Private/RewindDebuggerRuntimeModule.h b/Engine/Plugins/Animation/GameplayInsights/Source/RewindDebuggerRuntime/Private/RewindDebuggerRuntimeModule.h
new file mode 100644
--- /dev/null
+++ b/Engine/Plugins/Animation/GameplayInsights/Source/RewindDebuggerRuntime/Private/RewindDebuggerRuntimeModule.h
@@ -0,0 +1,20 @@
+// Copyright Epic Games, Inc. All Rights Reserved.
+
+#pragma once
+
+#include "CoreMinimal.h"
+#include "Modules/ModuleInterface.h"
+#include "RewindDebuggerAnimationRuntime.h"
+
+class FRewindDebuggerRuntimeModule : public IModuleInterface
+{
+public:
+	// IModuleInterface interface
+	virtual void StartupModule() override;
+	virtual void ShutdownModule() override;
+	
+private:
+	TArray<IConsoleObject*> ConsoleObjects;
+
+	RewindDebugger::FRewindDebuggerAnimationRuntime AnimationExtension;
+};
diff --git a/Engine/Plugins/Animation/GameplayInsights/Source/RewindDebuggerRuntime/Public/RewindDebuggerRuntime/RewindDebuggerRuntime.h b/Engine/Plugins/Animation/GameplayInsights/Source/RewindDebuggerRuntime/Public/RewindDebuggerRuntime/RewindDebuggerRuntime.h
new file mode 100644
--- /dev/null
+++ b/Engine/Plugins/Animation/GameplayInsights/Source/RewindDebuggerRuntime/Public/RewindDebuggerRuntime/RewindDebuggerRuntime.h
@@ -0,0 +1,32 @@
+// Copyright Epic Games, Inc. All Rights Reserved.
+
+#pragma once
+
+#include "CoreMinimal.h"
+
+REWINDDEBUGGERRUNTIME_API DECLARE_LOG_CATEGORY_EXTERN(LogRewindDebuggerRuntime, Log, All);
+
+
+namespace RewindDebugger
+{
+
+	class REWINDDEBUGGERRUNTIME_API FRewindDebuggerRuntime
+	{
+	public:
+		static void Initialize();
+		static void Shutdown();
+		static FRewindDebuggerRuntime* Instance() { return InternalInstance; }
+			
+		void StartRecording();
+		void StopRecording();
+
+		bool IsRecording() const { return bIsRecording; }
+
+		FSimpleMulticastDelegate RecordingStarted;
+		FSimpleMulticastDelegate ClearRecording;
+		FSimpleMulticastDelegate RecordingStopped;
+	private:
+		bool bIsRecording = false;
+		static FRewindDebuggerRuntime* InternalInstance;
+	};
+}
diff --git a/Engine/Plugins/Animation/GameplayInsights/Source/RewindDebuggerRuntime/RewindDebuggerRuntime.Build.cs b/Engine/Plugins/Animation/GameplayInsights/Source/RewindDebuggerRuntime/RewindDebuggerRuntime.Build.cs
new file mode 100644
--- /dev/null
+++ b/Engine/Plugins/Animation/GameplayInsights/Source/RewindDebuggerRuntime/RewindDebuggerRuntime.Build.cs
@@ -0,0 +1,23 @@
+// Copyright Epic Games, Inc. All Rights Reserved.
+
+namespace UnrealBuildTool.Rules
+{
+	public class RewindDebuggerRuntime : ModuleRules
+	{
+		public RewindDebuggerRuntime(ReadOnlyTargetRules Target) : base(Target)
+		{
+			PrivateDependencyModuleNames.AddRange(
+			new string[]
+			{
+				"Core",
+				"CoreUObject",
+				"Engine",
+				"RewindDebuggerInterface",
+				"RewindDebuggerRuntimeInterface",
+				"TraceLog",
+				"TraceServices",
+			});
+		}
+	}
+}
+
diff --git a/Engine/Plugins/Animation/GameplayInsights/Source/RewindDebuggerVLog/Private/RewindDebuggerVLog.cpp b/Engine/Plugins/Animation/GameplayInsights/Source/RewindDebuggerVLog/Private/RewindDebuggerVLog.cpp
--- a/Engine/Plugins/Animation/GameplayInsights/Source/RewindDebuggerVLog/Private/RewindDebuggerVLog.cpp
+++ b/Engine/Plugins/Animation/GameplayInsights/Source/RewindDebuggerVLog/Private/RewindDebuggerVLog.cpp
@@ -14,7 +14,6 @@
 #include "Engine/World.h"
 #include "HAL/IConsoleManager.h"
 #include "TraceServices/Model/Frames.h"
-#include "VisualLogger/VisualLogger.h"
 #include "VisualLogger/VisualLoggerTraceDevice.h"
 
 #define LOCTEXT_NAMESPACE "RewindDebuggerVLog"
@@ -238,16 +237,4 @@ void FRewindDebuggerVLog::Update(float DeltaTime, IRewindDebugger* RewindDebugge
 	}
 }
 
-void FRewindDebuggerVLog::RecordingStarted(IRewindDebugger*)
-{
-	// start recording visual logger data
-	FVisualLogger::Get().SetIsRecordingToTrace(true);
-}
-
-void FRewindDebuggerVLog::RecordingStopped(IRewindDebugger*)
-{
-	// stop recording visual logger data
-	FVisualLogger::Get().SetIsRecordingToTrace(false);
-}
-
 #undef LOCTEXT_NAMESPACE
\ No newline at end of file
diff --git a/Engine/Plugins/Animation/GameplayInsights/Source/RewindDebuggerVLog/Private/RewindDebuggerVLog.h b/Engine/Plugins/Animation/GameplayInsights/Source/RewindDebuggerVLog/Private/RewindDebuggerVLog.h
--- a/Engine/Plugins/Animation/GameplayInsights/Source/RewindDebuggerVLog/Private/RewindDebuggerVLog.h
+++ b/Engine/Plugins/Animation/GameplayInsights/Source/RewindDebuggerVLog/Private/RewindDebuggerVLog.h
@@ -27,9 +27,6 @@ public:
 	
 	virtual void Update(float DeltaTime, IRewindDebugger* RewindDebugger) override;
 
-	virtual void RecordingStarted(IRewindDebugger* RewindDebugger) override;
-	virtual void RecordingStopped(IRewindDebugger* RewindDebugger) override;
-
 private:
 	void AddLogEntries(const TArray<TSharedPtr<FDebugObjectInfo>>& Components, float StartTime, float EndTime, const class IVisualLoggerProvider* Provider);
 	void ImmediateRender(const UObject* Object, const FVisualLogEntry& Entry);
diff --git a/Engine/Plugins/Animation/GameplayInsights/Source/RewindDebuggerVLogRuntime/Private/RewindDebuggerVLogRuntime.h b/Engine/Plugins/Animation/GameplayInsights/Source/RewindDebuggerVLogRuntime/Private/RewindDebuggerVLogRuntime.h
new file mode 100644
--- /dev/null
+++ b/Engine/Plugins/Animation/GameplayInsights/Source/RewindDebuggerVLogRuntime/Private/RewindDebuggerVLogRuntime.h
@@ -0,0 +1,27 @@
+// Copyright Epic Games, Inc. All Rights Reserved.
+
+#pragma once
+
+#include "CoreMinimal.h"
+#include "RewindDebuggerRuntimeInterface/IRewindDebuggerRuntimeExtension.h"
+#include "VisualLogger/VisualLogger.h"
+
+namespace RewindDebugger
+{
+
+	class REWINDDEBUGGERVLOGRUNTIME_API FRewindDebuggerVLogRuntime : public IRewindDebuggerRuntimeExtension 
+	{
+	public:
+		virtual void RecordingStarted() override
+		{
+			// start recording visual logger data
+			FVisualLogger::Get().SetIsRecordingToTrace(true);
+		}
+		
+		virtual void RecordingStopped() override
+		{
+			// stop recording visual logger data
+			FVisualLogger::Get().SetIsRecordingToTrace(false);
+		}
+	};
+}
diff --git a/Engine/Plugins/Animation/GameplayInsights/Source/RewindDebuggerVLogRuntime/Private/RewindDebuggerVLogRuntimeModule.cpp b/Engine/Plugins/Animation/GameplayInsights/Source/RewindDebuggerVLogRuntime/Private/RewindDebuggerVLogRuntimeModule.cpp
new file mode 100644
--- /dev/null
+++ b/Engine/Plugins/Animation/GameplayInsights/Source/RewindDebuggerVLogRuntime/Private/RewindDebuggerVLogRuntimeModule.cpp
@@ -0,0 +1,17 @@
+// Copyright Epic Games, Inc. All Rights Reserved.
+
+#include "RewindDebuggerVLogRuntimeModule.h"
+#include "Modules/ModuleManager.h"
+#include "Features/IModularFeatures.h"
+
+void FRewindDebuggerVLogRuntimeModule::StartupModule()
+{
+	IModularFeatures::Get().RegisterModularFeature(IRewindDebuggerRuntimeExtension::ModularFeatureName, &VLogExtension);
+}
+
+void FRewindDebuggerVLogRuntimeModule::ShutdownModule()
+{
+	IModularFeatures::Get().UnregisterModularFeature(IRewindDebuggerRuntimeExtension::ModularFeatureName, &VLogExtension);
+}
+
+IMPLEMENT_MODULE(FRewindDebuggerVLogRuntimeModule, RewindDebuggerVLogRuntime);
diff --git a/Engine/Plugins/Animation/GameplayInsights/Source/RewindDebuggerVLogRuntime/Private/RewindDebuggerVLogRuntimeModule.h b/Engine/Plugins/Animation/GameplayInsights/Source/RewindDebuggerVLogRuntime/Private/RewindDebuggerVLogRuntimeModule.h
new file mode 100644
--- /dev/null
+++ b/Engine/Plugins/Animation/GameplayInsights/Source/RewindDebuggerVLogRuntime/Private/RewindDebuggerVLogRuntimeModule.h
@@ -0,0 +1,18 @@
+// Copyright Epic Games, Inc. All Rights Reserved.
+
+#pragma once
+
+#include "CoreMinimal.h"
+#include "Modules/ModuleInterface.h"
+#include "RewindDebuggerVLogRuntime.h"
+
+class FRewindDebuggerVLogRuntimeModule : public IModuleInterface
+{
+public:
+	// IModuleInterface interface
+	virtual void StartupModule() override;
+	virtual void ShutdownModule() override;
+	
+private:
+	RewindDebugger::FRewindDebuggerVLogRuntime VLogExtension;
+};
diff --git a/Engine/Plugins/Animation/GameplayInsights/Source/RewindDebuggerVLogRuntime/RewindDebuggerVLogRuntime.Build.cs b/Engine/Plugins/Animation/GameplayInsights/Source/RewindDebuggerVLogRuntime/RewindDebuggerVLogRuntime.Build.cs
new file mode 100644
--- /dev/null
+++ b/Engine/Plugins/Animation/GameplayInsights/Source/RewindDebuggerVLogRuntime/RewindDebuggerVLogRuntime.Build.cs
@@ -0,0 +1,23 @@
+// Copyright Epic Games, Inc. All Rights Reserved.
+
+namespace UnrealBuildTool.Rules
+{
+	public class RewindDebuggerVLogRuntime : ModuleRules
+	{
+		public RewindDebuggerVLogRuntime(ReadOnlyTargetRules Target) : base(Target)
+		{
+			PrivateDependencyModuleNames.AddRange(
+			new string[]
+			{
+				"Core",
+				"CoreUObject",
+				"Engine",
+				"RewindDebuggerInterface",
+				"TraceLog",
+				"TraceServices",
+				"RewindDebuggerRuntimeInterface",
+			});
+		}
+	}
+}
+
diff --git a/Engine/Plugins/Animation/PoseSearch/Source/Editor/Private/PoseSearchDebugger.cpp b/Engine/Plugins/Animation/PoseSearch/Source/Editor/Private/PoseSearchDebugger.cpp
--- a/Engine/Plugins/Animation/PoseSearch/Source/Editor/Private/PoseSearchDebugger.cpp
+++ b/Engine/Plugins/Animation/PoseSearch/Source/Editor/Private/PoseSearchDebugger.cpp
@@ -337,16 +337,6 @@ void FDebugger::Shutdown()
 	delete Debugger;
 }
 
-void FDebugger::RecordingStarted(IRewindDebugger*)
-{
-	UE::Trace::ToggleChannel(TEXT("PoseSearch"), true);
-}
-
-void FDebugger::RecordingStopped(IRewindDebugger*)
-{
-	UE::Trace::ToggleChannel(TEXT("PoseSearch"), false);
-}
-
 bool FDebugger::IsPIESimulating()
 {
 	return Debugger->RewindDebugger->IsPIESimulating();
diff --git a/Engine/Plugins/Animation/PoseSearch/Source/Editor/Private/PoseSearchDebugger.h b/Engine/Plugins/Animation/PoseSearch/Source/Editor/Private/PoseSearchDebugger.h
--- a/Engine/Plugins/Animation/PoseSearch/Source/Editor/Private/PoseSearchDebugger.h
+++ b/Engine/Plugins/Animation/PoseSearch/Source/Editor/Private/PoseSearchDebugger.h
@@ -27,9 +27,6 @@ public:
 	virtual void Update(float DeltaTime, IRewindDebugger* InRewindDebugger) override;
 	virtual ~FDebugger() = default;
 
-	virtual void RecordingStarted(IRewindDebugger* RewindDebugger) override;
-	virtual void RecordingStopped(IRewindDebugger* RewindDebugger) override;
-
 	static FDebugger* Get() { return Debugger; }
 	static void Initialize();
 	static void Shutdown();
diff --git a/Engine/Plugins/Animation/PoseSearch/Source/Runtime/PoseSearch.Build.cs b/Engine/Plugins/Animation/PoseSearch/Source/Runtime/PoseSearch.Build.cs
--- a/Engine/Plugins/Animation/PoseSearch/Source/Runtime/PoseSearch.Build.cs
+++ b/Engine/Plugins/Animation/PoseSearch/Source/Runtime/PoseSearch.Build.cs
@@ -27,6 +27,12 @@ public class PoseSearch : ModuleRules
 			}
 		);
 
+		PrivateDependencyModuleNames.AddRange(
+		new string[]{
+			"RewindDebuggerRuntimeInterface",
+		}
+		);
+
 		if (Target.bCompileAgainstEditor)
 		{
 			PrivateDependencyModuleNames.AddRange(
diff --git a/Engine/Plugins/Animation/PoseSearch/Source/Runtime/Private/PoseSearchModule.cpp b/Engine/Plugins/Animation/PoseSearch/Source/Runtime/Private/PoseSearchModule.cpp
--- a/Engine/Plugins/Animation/PoseSearch/Source/Runtime/Private/PoseSearchModule.cpp
+++ b/Engine/Plugins/Animation/PoseSearch/Source/Runtime/Private/PoseSearchModule.cpp
@@ -9,6 +9,21 @@
 #include "PoseSearch/PoseSearchLibrary.h"
 #include "PoseSearch/PoseSearchSchema.h"
 #include "PoseSearchFeatureChannel_PermutationTime.h"
+#include "RewindDebuggerRuntimeInterface/IRewindDebuggerRuntimeExtension.h"
+
+class FRewindDebuggerPoseSearchRuntime : public IRewindDebuggerRuntimeExtension 
+{
+public:
+	virtual void RecordingStarted() override
+	{
+		UE::Trace::ToggleChannel(TEXT("PoseSearch"), true);
+	}
+	
+	virtual void RecordingStopped() override
+	{
+		UE::Trace::ToggleChannel(TEXT("PoseSearch"), false);
+	}
+};
 
 class FPoseSearchModule final : public IModuleInterface, public UE::Anim::IPoseSearchProvider
 {
@@ -18,11 +33,13 @@ public:
 	virtual void StartupModule() override
 	{
 		IModularFeatures::Get().RegisterModularFeature(UE::Anim::IPoseSearchProvider::GetModularFeatureName(), this);
+		IModularFeatures::Get().RegisterModularFeature(IRewindDebuggerRuntimeExtension::ModularFeatureName, &RewindDebuggerPoseSearchRuntime);
 	}
 
 	virtual void ShutdownModule() override
 	{
 		IModularFeatures::Get().UnregisterModularFeature(UE::Anim::IPoseSearchProvider::GetModularFeatureName(), this);
+		IModularFeatures::Get().UnregisterModularFeature(IRewindDebuggerRuntimeExtension::ModularFeatureName, &RewindDebuggerPoseSearchRuntime);
 	}
 
 	// IPoseSearchProvider
@@ -84,6 +101,9 @@ public:
 
 		return ProviderResult;
 	}
+
+private:
+	FRewindDebuggerPoseSearchRuntime RewindDebuggerPoseSearchRuntime;
 };
 
 IMPLEMENT_MODULE(FPoseSearchModule, PoseSearch);
\ No newline at end of file
diff --git a/Engine/Plugins/Experimental/Chooser/Source/Chooser/Chooser.Build.cs b/Engine/Plugins/Experimental/Chooser/Source/Chooser/Chooser.Build.cs
--- a/Engine/Plugins/Experimental/Chooser/Source/Chooser/Chooser.Build.cs
+++ b/Engine/Plugins/Experimental/Chooser/Source/Chooser/Chooser.Build.cs
@@ -24,7 +24,8 @@ namespace UnrealBuildTool.Rules
 					"AnimationCore",
 					"AnimGraphRuntime",
 					"BlendStack",
-					"TraceLog"
+					"TraceLog",
+					"RewindDebuggerRuntimeInterface",
 					// ... add private dependencies that you statically link with here ...
 				}
 			);
diff --git a/Engine/Plugins/Experimental/Chooser/Source/Chooser/Private/ChooserModule.cpp b/Engine/Plugins/Experimental/Chooser/Source/Chooser/Private/ChooserModule.cpp
--- a/Engine/Plugins/Experimental/Chooser/Source/Chooser/Private/ChooserModule.cpp
+++ b/Engine/Plugins/Experimental/Chooser/Source/Chooser/Private/ChooserModule.cpp
@@ -1,5 +1,18 @@
 // Copyright Epic Games, Inc. All Rights Reserved.
 
+#include "ChooserModule.h"
 #include "Modules/ModuleManager.h"
+#include "ObjectTrace.h"
+#include "Features/IModularFeatures.h"
 
-IMPLEMENT_MODULE(FDefaultModuleImpl, Chooser);
\ No newline at end of file
+void FChooserModule::StartupModule()
+{
+	IModularFeatures::Get().RegisterModularFeature(IRewindDebuggerRuntimeExtension::ModularFeatureName, &ChooserExtension);
+}
+
+void FChooserModule::ShutdownModule()
+{
+	IModularFeatures::Get().UnregisterModularFeature(IRewindDebuggerRuntimeExtension::ModularFeatureName, &ChooserExtension);
+}
+
+IMPLEMENT_MODULE(FChooserModule, Chooser);
diff --git a/Engine/Plugins/Experimental/Chooser/Source/Chooser/Private/ChooserModule.h b/Engine/Plugins/Experimental/Chooser/Source/Chooser/Private/ChooserModule.h
new file mode 100644
--- /dev/null
+++ b/Engine/Plugins/Experimental/Chooser/Source/Chooser/Private/ChooserModule.h
@@ -0,0 +1,18 @@
+// Copyright Epic Games, Inc. All Rights Reserved.
+
+#pragma once
+
+#include "CoreMinimal.h"
+#include "Modules/ModuleInterface.h"
+#include "RewindDebuggerChooserRuntime.h"
+
+class FChooserModule : public IModuleInterface
+{
+public:
+	// IModuleInterface interface
+	virtual void StartupModule() override;
+	virtual void ShutdownModule() override;
+	
+private:
+	FRewindDebuggerChooserRuntime ChooserExtension;
+};
diff --git a/Engine/Plugins/Experimental/Chooser/Source/Chooser/Private/RewindDebuggerChooserRuntime.cpp b/Engine/Plugins/Experimental/Chooser/Source/Chooser/Private/RewindDebuggerChooserRuntime.cpp
new file mode 100644
--- /dev/null
+++ b/Engine/Plugins/Experimental/Chooser/Source/Chooser/Private/RewindDebuggerChooserRuntime.cpp
@@ -0,0 +1,13 @@
+// Copyright Epic Games, Inc. All Rights Reserved.
+
+#include "RewindDebuggerChooserRuntime.h"
+
+void FRewindDebuggerChooserRuntime::RecordingStarted()
+{
+	UE::Trace::ToggleChannel(TEXT("Chooser"), true);
+}
+
+void FRewindDebuggerChooserRuntime::RecordingStopped()
+{
+	UE::Trace::ToggleChannel(TEXT("Chooser"), false);
+}
\ No newline at end of file
diff --git a/Engine/Plugins/Experimental/Chooser/Source/Chooser/Private/RewindDebuggerChooserRuntime.h b/Engine/Plugins/Experimental/Chooser/Source/Chooser/Private/RewindDebuggerChooserRuntime.h
new file mode 100644
--- /dev/null
+++ b/Engine/Plugins/Experimental/Chooser/Source/Chooser/Private/RewindDebuggerChooserRuntime.h
@@ -0,0 +1,14 @@
+// Copyright Epic Games, Inc. All Rights Reserved.
+#pragma once
+
+#include "CoreMinimal.h"
+#include "RewindDebuggerRuntimeInterface/IRewindDebuggerRuntimeExtension.h"
+
+// Rewind debugger extension for Chooser support
+
+class FRewindDebuggerChooserRuntime : public IRewindDebuggerRuntimeExtension
+{
+public:
+	virtual void RecordingStarted() override;
+	virtual void RecordingStopped() override;
+};
diff --git a/Engine/Plugins/Experimental/Chooser/Source/ChooserEditor/Private/RewindDebuggerChooser.cpp b/Engine/Plugins/Experimental/Chooser/Source/ChooserEditor/Private/RewindDebuggerChooser.cpp
--- a/Engine/Plugins/Experimental/Chooser/Source/ChooserEditor/Private/RewindDebuggerChooser.cpp
+++ b/Engine/Plugins/Experimental/Chooser/Source/ChooserEditor/Private/RewindDebuggerChooser.cpp
@@ -75,14 +75,4 @@ void FRewindDebuggerChooser::Update(float DeltaTime, IRewindDebugger* RewindDebu
 
 		});
 	}
-}
-
-void FRewindDebuggerChooser::RecordingStarted(IRewindDebugger*)
-{
-	UE::Trace::ToggleChannel(TEXT("Chooser"), true);
-}
-
-void FRewindDebuggerChooser::RecordingStopped(IRewindDebugger*)
-{
-	UE::Trace::ToggleChannel(TEXT("Chooser"), false);
 }
\ No newline at end of file
diff --git a/Engine/Plugins/Experimental/Chooser/Source/ChooserEditor/Private/RewindDebuggerChooser.h b/Engine/Plugins/Experimental/Chooser/Source/ChooserEditor/Private/RewindDebuggerChooser.h
--- a/Engine/Plugins/Experimental/Chooser/Source/ChooserEditor/Private/RewindDebuggerChooser.h
+++ b/Engine/Plugins/Experimental/Chooser/Source/ChooserEditor/Private/RewindDebuggerChooser.h
@@ -14,6 +14,4 @@ public:
 	virtual ~FRewindDebuggerChooser() {};
 
 	virtual void Update(float DeltaTime, IRewindDebugger* RewindDebugger) override;
-	virtual void RecordingStarted(IRewindDebugger* RewindDebugger) override;
-	virtual void RecordingStopped(IRewindDebugger* RewindDebugger) override;
 };
diff --git a/Engine/Source/Runtime/RewindDebuggerRuntimeInterface/Private/IRewindDebuggerRuntimeExtension.cpp b/Engine/Source/Runtime/RewindDebuggerRuntimeInterface/Private/IRewindDebuggerRuntimeExtension.cpp
new file mode 100644
--- /dev/null
+++ b/Engine/Source/Runtime/RewindDebuggerRuntimeInterface/Private/IRewindDebuggerRuntimeExtension.cpp
@@ -0,0 +1,6 @@
+// Copyright Epic Games, Inc. All Rights Reserved.
+
+#include "RewindDebuggerRuntimeInterface/IRewindDebuggerRuntimeExtension.h"
+
+const FName IRewindDebuggerRuntimeExtension::ModularFeatureName = "RewindDebuggerRuntimeExtension";
+	
diff --git a/Engine/Source/Runtime/RewindDebuggerRuntimeInterface/Public/RewindDebuggerRuntimeInterface/IRewindDebuggerRuntimeExtension.h b/Engine/Source/Runtime/RewindDebuggerRuntimeInterface/Public/RewindDebuggerRuntimeInterface/IRewindDebuggerRuntimeExtension.h
new file mode 100644
--- /dev/null
+++ b/Engine/Source/Runtime/RewindDebuggerRuntimeInterface/Public/RewindDebuggerRuntimeInterface/IRewindDebuggerRuntimeExtension.h
@@ -0,0 +1,30 @@
+// Copyright Epic Games, Inc. All Rights Reserved.
+
+#pragma once
+
+#include "CoreTypes.h"
+#include "UObject/NameTypes.h"
+#include "Features/IModularFeature.h"
+
+// IRewindDebuggerRuntimeExtension 
+//
+// interface class for extensions which add functionality to the rewind debugger
+// these get a callback on recording start/stop, to enable/disable trace channels and on clear to clean up any cached data
+//
+
+class REWINDDEBUGGERRUNTIMEINTERFACE_API IRewindDebuggerRuntimeExtension : public IModularFeature
+{
+public:
+	virtual ~IRewindDebuggerRuntimeExtension() {}
+	
+	static const FName ModularFeatureName;
+
+	// called when recording has started
+	virtual void RecordingStarted() {};
+
+	// called when recording has ended
+	virtual void RecordingStopped() {};
+	
+	// called when recording is unloaded or a new one starts
+	virtual void Clear() {};
+};
diff --git a/Engine/Source/Runtime/RewindDebuggerRuntimeInterface/RewindDebuggerRuntimeInterface.Build.cs b/Engine/Source/Runtime/RewindDebuggerRuntimeInterface/RewindDebuggerRuntimeInterface.Build.cs
new file mode 100644
--- /dev/null
+++ b/Engine/Source/Runtime/RewindDebuggerRuntimeInterface/RewindDebuggerRuntimeInterface.Build.cs
@@ -0,0 +1,17 @@
+// Copyright Epic Games, Inc. All Rights Reserved.
+
+namespace UnrealBuildTool.Rules
+{
+	public class RewindDebuggerRuntimeInterface : ModuleRules
+	{
+		public RewindDebuggerRuntimeInterface(ReadOnlyTargetRules Target) : base(Target)
+		{
+			PrivateDependencyModuleNames.AddRange(
+				new string[]
+				{
+					"Core",
+					"CoreUObject",
+				});
+		}
+	}
+}
\ No newline at end of file
