Compile fixes for RewindDebugger runtime modules

diff --git a/Engine/Plugins/Animation/GameplayInsights/Source/RewindDebuggerRuntime/RewindDebuggerRuntime.Build.cs b/Engine/Plugins/Animation/GameplayInsights/Source/RewindDebuggerRuntime/RewindDebuggerRuntime.Build.cs
--- a/Engine/Plugins/Animation/GameplayInsights/Source/RewindDebuggerRuntime/RewindDebuggerRuntime.Build.cs
+++ b/Engine/Plugins/Animation/GameplayInsights/Source/RewindDebuggerRuntime/RewindDebuggerRuntime.Build.cs
@@ -15,7 +15,6 @@ namespace UnrealBuildTool.Rules
 				"RewindDebuggerInterface",
 				"RewindDebuggerRuntimeInterface",
 				"TraceLog",
-				"TraceServices",
 			});
 		}
 	}
diff --git a/Engine/Plugins/Animation/GameplayInsights/Source/RewindDebuggerVLogRuntime/Private/RewindDebuggerVLogRuntime.h b/Engine/Plugins/Animation/GameplayInsights/Source/RewindDebuggerVLogRuntime/Private/RewindDebuggerVLogRuntime.h
--- a/Engine/Plugins/Animation/GameplayInsights/Source/RewindDebuggerVLogRuntime/Private/RewindDebuggerVLogRuntime.h
+++ b/Engine/Plugins/Animation/GameplayInsights/Source/RewindDebuggerVLogRuntime/Private/RewindDebuggerVLogRuntime.h
@@ -14,14 +14,18 @@ namespace RewindDebugger
 	public:
 		virtual void RecordingStarted() override
 		{
+#if ENABLE_VISUAL_LOG	
 			// start recording visual logger data
 			FVisualLogger::Get().SetIsRecordingToTrace(true);
+#endif
 		}
 		
 		virtual void RecordingStopped() override
 		{
+#if ENABLE_VISUAL_LOG	
 			// stop recording visual logger data
 			FVisualLogger::Get().SetIsRecordingToTrace(false);
+#endif
 		}
 	};
 }
diff --git a/Engine/Plugins/Animation/GameplayInsights/Source/RewindDebuggerVLogRuntime/RewindDebuggerVLogRuntime.Build.cs b/Engine/Plugins/Animation/GameplayInsights/Source/RewindDebuggerVLogRuntime/RewindDebuggerVLogRuntime.Build.cs
--- a/Engine/Plugins/Animation/GameplayInsights/Source/RewindDebuggerVLogRuntime/RewindDebuggerVLogRuntime.Build.cs
+++ b/Engine/Plugins/Animation/GameplayInsights/Source/RewindDebuggerVLogRuntime/RewindDebuggerVLogRuntime.Build.cs
@@ -14,7 +14,6 @@ namespace UnrealBuildTool.Rules
 				"Engine",
 				"RewindDebuggerInterface",
 				"TraceLog",
-				"TraceServices",
 				"RewindDebuggerRuntimeInterface",
 			});
 		}
