Shipping compile fixes for Rewind Debugger Runtime modules

diff --git a/Engine/Plugins/Animation/GameplayInsights/Source/RewindDebuggerRuntime/Private/RewindDebuggerRuntime.cpp b/Engine/Plugins/Animation/GameplayInsights/Source/RewindDebuggerRuntime/Private/RewindDebuggerRuntime.cpp
--- a/Engine/Plugins/Animation/GameplayInsights/Source/RewindDebuggerRuntime/Private/RewindDebuggerRuntime.cpp
+++ b/Engine/Plugins/Animation/GameplayInsights/Source/RewindDebuggerRuntime/Private/RewindDebuggerRuntime.cpp
@@ -6,10 +6,10 @@
 #include "EngineUtils.h"
 #include "Engine/World.h"
 #include "GameFramework/Pawn.h"
-#include "IRewindDebuggerExtension.h"
 #include "ProfilingDebugging/TraceAuxiliary.h"
 #include "UObject/UObjectIterator.h"
 #include "RewindDebuggerRuntimeInterface/IRewindDebuggerRuntimeExtension.h"
+#include "Features/IModularFeatures.h"
 
 DEFINE_LOG_CATEGORY(LogRewindDebuggerRuntime)
 
diff --git a/Engine/Plugins/Animation/GameplayInsights/Source/RewindDebuggerRuntime/RewindDebuggerRuntime.Build.cs b/Engine/Plugins/Animation/GameplayInsights/Source/RewindDebuggerRuntime/RewindDebuggerRuntime.Build.cs
--- a/Engine/Plugins/Animation/GameplayInsights/Source/RewindDebuggerRuntime/RewindDebuggerRuntime.Build.cs
+++ b/Engine/Plugins/Animation/GameplayInsights/Source/RewindDebuggerRuntime/RewindDebuggerRuntime.Build.cs
@@ -12,7 +12,6 @@ namespace UnrealBuildTool.Rules
 				"Core",
 				"CoreUObject",
 				"Engine",
-				"RewindDebuggerInterface",
 				"RewindDebuggerRuntimeInterface",
 				"TraceLog",
 			});
diff --git a/Engine/Plugins/Animation/GameplayInsights/Source/RewindDebuggerVLogRuntime/RewindDebuggerVLogRuntime.Build.cs b/Engine/Plugins/Animation/GameplayInsights/Source/RewindDebuggerVLogRuntime/RewindDebuggerVLogRuntime.Build.cs
--- a/Engine/Plugins/Animation/GameplayInsights/Source/RewindDebuggerVLogRuntime/RewindDebuggerVLogRuntime.Build.cs
+++ b/Engine/Plugins/Animation/GameplayInsights/Source/RewindDebuggerVLogRuntime/RewindDebuggerVLogRuntime.Build.cs
@@ -12,7 +12,6 @@ namespace UnrealBuildTool.Rules
 				"Core",
 				"CoreUObject",
 				"Engine",
-				"RewindDebuggerInterface",
 				"TraceLog",
 				"RewindDebuggerRuntimeInterface",
 			});
