Fix compile error when ObjectTrace is enabled but AnimTrace is disabled

diff --git a/Engine/Plugins/Animation/GameplayInsights/Source/RewindDebuggerRuntime/Private/RewindDebuggerRuntime.cpp b/Engine/Plugins/Animation/GameplayInsights/Source/RewindDebuggerRuntime/Private/RewindDebuggerRuntime.cpp
--- a/Engine/Plugins/Animation/GameplayInsights/Source/RewindDebuggerRuntime/Private/RewindDebuggerRuntime.cpp
+++ b/Engine/Plugins/Animation/GameplayInsights/Source/RewindDebuggerRuntime/Private/RewindDebuggerRuntime.cpp
@@ -60,7 +60,9 @@ namespace RewindDebugger
 #if OBJECT_TRACE_ENABLED
 		// Clear caches  (maybe move these to plugins?)
 		FObjectTrace::Reset();
+#if ANIM_TRACE_ENABLED
 		FAnimTrace::Reset();
+#endif
 
 		ClearRecording.Broadcast();
         	
