Rewind debugger support for attaching to a live session

diff --git a/Engine/Plugins/Animation/GameplayInsights/Source/RewindDebugger/Private/RewindDebugger.cpp b/Engine/Plugins/Animation/GameplayInsights/Source/RewindDebugger/Private/RewindDebugger.cpp
--- a/Engine/Plugins/Animation/GameplayInsights/Source/RewindDebugger/Private/RewindDebugger.cpp
+++ b/Engine/Plugins/Animation/GameplayInsights/Source/RewindDebugger/Private/RewindDebugger.cpp
@@ -28,6 +28,7 @@
 #include "RewindDebuggerPlaceholderTrack.h"
 #include "RewindDebuggerSettings.h"
 #include "SLevelViewport.h"
+#include "SModalSessionBrowser.h"
 #include "ToolMenus.h"
 #include "Trace/StoreClient.h"
 #include "TraceServices/Model/Frames.h"
@@ -43,6 +44,7 @@
 #include "TraceServices/AnalysisService.h"
 #include "TraceServices/ITraceServicesModule.h"
 #include "Widgets/Input/SNumericEntryBox.h"
+#include "Misc/MessageDialog.h"
 
 #define LOCTEXT_NAMESPACE "RewindDebugger"
 
@@ -527,6 +529,31 @@ void FRewindDebugger::OpenTrace()
 	}
 }
 
+
+void FRewindDebugger::AttachToSession()
+{
+	ClearTrace();
+	const TSharedRef<SModalSessionBrowser> SessionBrowserModal = SNew(SModalSessionBrowser);
+
+	if (SessionBrowserModal->ShowModal() != EAppReturnType::Cancel)
+	{
+		bool bSuccess = false;
+		const SModalSessionBrowser::FTraceSessionInfo SessionInfo = SessionBrowserModal->GetSelectedTraceInfo();
+		if (SessionInfo.bIsValid)
+		{
+			const FString SessionAddress = SessionBrowserModal->GetSelectedTraceStoreAddress();
+			IUnrealInsightsModule& TraceInsightsModule = FModuleManager::LoadModuleChecked<IUnrealInsightsModule>("TraceInsights");
+			TraceInsightsModule.StartAnalysisForTrace(SessionInfo.TraceID);
+			bSuccess = TraceInsightsModule.GetAnalysisSession().IsValid();
+		}
+
+		if (!bSuccess)
+		{
+			FMessageDialog::Open(EAppMsgType::Ok, LOCTEXT("FailedToConnectToSessionMessage", "Failed to connect to session"));	
+		}
+	}
+}
+
 bool FRewindDebugger::CanClearTrace() const
 {
 	return GetAnalysisSession() != nullptr;
@@ -1376,6 +1403,13 @@ void FRewindDebugger::RegisterToolBar()
 
 	Section.AddSeparator(NAME_None);
 
+
+	Section.AddEntry(FToolMenuEntry::InitToolBarButton(
+				Commands.AttachToSession,
+				LOCTEXT("Blank",""),
+				TAttribute<FText>(),
+				FSlateIcon("RewindDebuggerStyle", "RewindDebugger.ConnectToSession")));
+				
 	Section.AddEntry(FToolMenuEntry::InitToolBarButton(
 				Commands.OpenTrace,
 				LOCTEXT("Blank",""),
diff --git a/Engine/Plugins/Animation/GameplayInsights/Source/RewindDebugger/Private/RewindDebugger.h b/Engine/Plugins/Animation/GameplayInsights/Source/RewindDebugger/Private/RewindDebugger.h
--- a/Engine/Plugins/Animation/GameplayInsights/Source/RewindDebugger/Private/RewindDebugger.h
+++ b/Engine/Plugins/Animation/GameplayInsights/Source/RewindDebugger/Private/RewindDebugger.h
@@ -77,6 +77,8 @@ public:
 	bool CanOpenTrace() const;
    	void OpenTrace(const FString& FilePath);
    	void OpenTrace();
+	
+   	void AttachToSession();
 
 	bool CanClearTrace() const;
 	void ClearTrace();
diff --git a/Engine/Plugins/Animation/GameplayInsights/Source/RewindDebugger/Private/RewindDebuggerCommands.h b/Engine/Plugins/Animation/GameplayInsights/Source/RewindDebugger/Private/RewindDebuggerCommands.h
--- a/Engine/Plugins/Animation/GameplayInsights/Source/RewindDebugger/Private/RewindDebuggerCommands.h
+++ b/Engine/Plugins/Animation/GameplayInsights/Source/RewindDebugger/Private/RewindDebuggerCommands.h
@@ -34,6 +34,7 @@ public:
 		UI_COMMAND(AutoEject, "Auto Eject", "Automatically detach player control when PIE is paused", EUserInterfaceActionType::ToggleButton, FInputChord());
 		UI_COMMAND(AutoRecord, "Auto Record", "Automatically start recording when PIE is started", EUserInterfaceActionType::ToggleButton, FInputChord());
 		UI_COMMAND(OpenTrace, "Open Recording", "Open a Trace file containing a rewind debugger recording", EUserInterfaceActionType::Button, FInputChord());
+		UI_COMMAND(AttachToSession, "Attach to Session", "Attach to a live session", EUserInterfaceActionType::Button, FInputChord());
 		UI_COMMAND(SaveTrace, "Save Recording", "Save a copy of the current recording", EUserInterfaceActionType::Button, FInputChord());
 		UI_COMMAND(ClearTrace, "Clear Recording", "Clear current recording data", EUserInterfaceActionType::Button, FInputChord());
 	}
@@ -52,6 +53,7 @@ public:
 	TSharedPtr<FUICommandInfo> AutoEject;
 	TSharedPtr<FUICommandInfo> AutoRecord;
 	TSharedPtr<FUICommandInfo> OpenTrace;
+	TSharedPtr<FUICommandInfo> AttachToSession;
 	TSharedPtr<FUICommandInfo> SaveTrace;
 	TSharedPtr<FUICommandInfo> ClearTrace;
 };
diff --git a/Engine/Plugins/Animation/GameplayInsights/Source/RewindDebugger/Private/RewindDebuggerModule.cpp b/Engine/Plugins/Animation/GameplayInsights/Source/RewindDebugger/Private/RewindDebuggerModule.cpp
--- a/Engine/Plugins/Animation/GameplayInsights/Source/RewindDebugger/Private/RewindDebuggerModule.cpp
+++ b/Engine/Plugins/Animation/GameplayInsights/Source/RewindDebugger/Private/RewindDebuggerModule.cpp
@@ -160,6 +160,12 @@ TSharedRef<SDockTab> FRewindDebuggerModule::SpawnRewindDebuggerTab(const FSpawnT
 							 FCanExecuteAction::CreateRaw(DebuggerInstance, &FRewindDebugger::CanOpenTrace),
 							 FIsActionChecked(),
 							 FIsActionButtonVisible());
+	
+	CommandList->MapAction(Commands.AttachToSession,
+    						 FExecuteAction::CreateRaw(DebuggerInstance, &FRewindDebugger::AttachToSession),
+    						 FCanExecuteAction::CreateRaw(DebuggerInstance, &FRewindDebugger::CanOpenTrace),
+    						 FIsActionChecked(),
+    						 FIsActionButtonVisible());
 		
 	CommandList->MapAction(Commands.SaveTrace,
 							 FExecuteAction::CreateRaw(DebuggerInstance, &FRewindDebugger::SaveTrace),
diff --git a/Engine/Plugins/Animation/GameplayInsights/Source/RewindDebugger/Private/SModalSessionBrowser.cpp b/Engine/Plugins/Animation/GameplayInsights/Source/RewindDebugger/Private/SModalSessionBrowser.cpp
new file mode 100644
--- /dev/null
+++ b/Engine/Plugins/Animation/GameplayInsights/Source/RewindDebugger/Private/SModalSessionBrowser.cpp
@@ -0,0 +1,298 @@
+// Copyright Epic Games, Inc. All Rights Reserved.
+
+#include "SModalSessionBrowser.h"
+
+#include "Editor.h"
+#include "Framework/Application/SlateApplication.h"
+#include "GenericPlatform/GenericPlatformMisc.h"
+#include "Logging/LogMacros.h"
+#include "SModalSessionNameListPicker.h"
+#include "Trace/StoreClient.h"
+#include "Widgets/Layout/SUniformGridPanel.h"
+#include "Widgets/Text/STextBlock.h"
+#include "Widgets/Input/SButton.h"
+
+#define LOCTEXT_NAMESPACE "ModalSessionBrowser"
+
+DEFINE_LOG_CATEGORY_STATIC(LogSessionBrowser, Log, All);
+
+/** Reduced Version of the IP4 structure from the Network Module
+ * used just to convert to String an IP saved as host byte order.
+ */
+struct FCVDIPv4Address
+{
+	union
+	{
+		/** The IP address value as A.B.C.D components. */
+		struct
+		{
+#if PLATFORM_LITTLE_ENDIAN
+#ifdef _MSC_VER
+			uint8 D, C, B, A;
+#else
+			uint8 D GCC_ALIGN(4);
+			uint8 C, B, A;
+#endif
+#else
+			uint8 A, B, C, D;
+#endif
+		};
+
+		/** The IP address value in host byte order. */
+		uint32 Value;
+	};
+	
+	/**
+ * Creates and initializes a new IPv4 address with the specified components.
+ *
+ * The created IP address has the value A.B.C.D.
+ *
+ * @param InA The first component.
+ * @param InB The second component.
+ * @param InC The third component.
+ * @param InD The fourth component.
+ */
+	FCVDIPv4Address(uint8 InA, uint8 InB, uint8 InC, uint8 InD)
+#if PLATFORM_LITTLE_ENDIAN
+		: D(InD)
+		, C(InC)
+		, B(InB)
+		, A(InA)
+#else
+		: A(InA)
+		, B(InB)
+		, C(InC)
+		, D(InD)
+#endif // PLATFORM_LITTLE_ENDIAN
+	{ }
+
+	/**
+	 * Creates and initializes a new IPv4 address with the specified value.
+	 *
+	 * @param InValue The address value (in host byte order).
+	 */
+	FCVDIPv4Address(uint32 InValue)
+		: Value(InValue)
+	{ }
+
+	FString ToString() const { return FString::Printf(TEXT("%i.%i.%i.%i"), A, B, C, D); }
+};
+
+SModalSessionBrowser::~SModalSessionBrowser()
+{
+	if (FSlateApplication::IsInitialized())
+	{
+		if (ModalTickHandle.IsValid())
+		{
+			FSlateApplication::Get().GetOnModalLoopTickEvent().Remove(ModalTickHandle);
+		}
+	}
+}
+
+void SModalSessionBrowser::Construct(const FArguments& InArgs)
+{
+	CurrentTraceStoreAddress = TEXT("127.0.0.1");
+	
+	SWindow::Construct(SWindow::FArguments()
+	.Title(LOCTEXT("SModalSessionBrowser_Title", "Live Session Browser"))
+	.SupportsMinimize(false)
+	.SupportsMaximize(false)
+	.UserResizeBorder(0)
+	.ClientSize(FVector2D(450, 120))
+	[
+		SNew(SVerticalBox)
+		+SVerticalBox::Slot()
+		.Padding(5)
+		[
+			SNew(SBorder)
+			.BorderImage(FAppStyle::GetBrush("ToolPanel.GroupBorder"))
+			[
+				SNew(SVerticalBox)
+				+SVerticalBox::Slot()
+				.Padding(0.0f, 0.0f, 0.0f, 5.0f)
+				.AutoHeight()
+				[
+					SNew(SHorizontalBox)
+					+SHorizontalBox::Slot()
+					.HAlign(HAlign_Left)
+					.VAlign(VAlign_Center)
+					[
+						SNew(STextBlock)
+						.Text(LOCTEXT("TraceStoreAddress", "Trace Store Address"))
+						.Font(FCoreStyle::GetDefaultFontStyle("Regular", 12))
+					]
+					+SHorizontalBox::Slot()
+					.HAlign(HAlign_Right)
+					.VAlign(VAlign_Center)
+					[
+						SAssignNew(TraceStoreAddressWidget, SEditableTextBox)
+						.Text(FText::AsCultureInvariant(CurrentTraceStoreAddress))
+						.Font(FCoreStyle::GetDefaultFontStyle("Regular", 12))
+						.OnTextCommitted_Raw(this, &SModalSessionBrowser::OnTraceStoreAddressUpdated)
+					]
+				]
+				+SVerticalBox::Slot()
+				.AutoHeight()
+				[
+					SNew(SHorizontalBox)
+					+SHorizontalBox::Slot()
+					.HAlign(HAlign_Left)
+					.VAlign(VAlign_Center)
+					[
+						SNew(STextBlock)
+						.Text(LOCTEXT("SelectedLiveSession", "Selected Live Session"))
+						.Font(FCoreStyle::GetDefaultFontStyle("Regular", 12))
+					]
+					+SHorizontalBox::Slot()
+					.HAlign(HAlign_Right)
+					.VAlign(VAlign_Center)
+					[
+						SAssignNew(NamePickerWidget, SModalSessionNameListPicker)
+						.OnNameSleceted_Raw(this, &SModalSessionBrowser::HandleSessionNameSelected)
+					]	
+				]
+			]
+		]
+		+SVerticalBox::Slot()
+		.AutoHeight()
+		.HAlign(HAlign_Right)
+		.Padding(5)
+		[
+			SNew(SUniformGridPanel)
+			.SlotPadding(FAppStyle::GetMargin("StandardDialog.SlotPadding"))
+			.MinDesiredSlotWidth(FAppStyle::GetFloat("StandardDialog.MinDesiredSlotWidth"))
+			.MinDesiredSlotHeight(FAppStyle::GetFloat("StandardDialog.MinDesiredSlotHeight"))
+			+SUniformGridPanel::Slot(0, 0)
+			[
+				SNew(SButton)
+				.HAlign(HAlign_Center)
+				.ContentPadding(FAppStyle::GetMargin("StandardDialog.ContentPadding"))
+				.Text(LOCTEXT("ConnectToSession", "Connect to Session"))
+				.IsEnabled_Raw(this, &SModalSessionBrowser::CanOpenSession)
+				.OnClicked(this, &SModalSessionBrowser::OnButtonClick, EAppReturnType::Ok)
+			]
+			+SUniformGridPanel::Slot(1, 0)
+			[
+				SNew(SButton)
+				.HAlign(HAlign_Center)
+				.ContentPadding(FAppStyle::GetMargin("StandardDialog.ContentPadding"))
+				.Text(LOCTEXT("Cancel", "Cancel"))
+				.OnClicked(this, &SModalSessionBrowser::OnButtonClick, EAppReturnType::Cancel)
+			]
+		]
+	]);
+
+	UpdateCurrentSessionInfoMap();
+
+	if (FSlateApplication::IsInitialized())
+	{
+		ModalTickHandle = FSlateApplication::Get().GetOnModalLoopTickEvent().AddRaw(this, &SModalSessionBrowser::ModalTick);
+	}
+}
+
+EAppReturnType::Type SModalSessionBrowser::ShowModal()
+{
+	GEditor->EditorAddModalWindow(SharedThis(this));
+	return UserResponse;
+}
+
+SModalSessionBrowser::FTraceSessionInfo SModalSessionBrowser::GetSelectedTraceInfo()
+{
+	if (ensure(CurrentTraceSessionSelected.IsValid()))
+	{
+		if (const FTraceSessionInfo* SessionInfoPtr = CurrentSessionInfosMap.Find(CurrentTraceSessionSelected))
+		{
+			return *SessionInfoPtr;
+		}
+		else
+		{
+			UE_LOG(LogSessionBrowser, Error, TEXT("Selected session [%s] is no longer available"), *CurrentTraceSessionSelected.ToString())
+		}
+	}
+	else
+	{
+		UE_LOG(LogSessionBrowser, Error, TEXT("Selected session [%s] is invalid"), *CurrentTraceSessionSelected.ToString())
+	}
+
+	return FTraceSessionInfo();
+}
+
+void SModalSessionBrowser::ModalTick(float InDeltaTime)
+{
+	AccumulatedTimeBetweenTicks += InDeltaTime;
+
+	// Update the sessions at least once per desired interval
+	constexpr float PollingInterval = 1.0f;
+	if (AccumulatedTimeBetweenTicks > PollingInterval)
+	{
+		UpdateCurrentSessionInfoMap();
+		AccumulatedTimeBetweenTicks = 0.0f;
+	}
+}
+
+bool SModalSessionBrowser::CanOpenSession() const
+{
+	return !CurrentTraceSessionSelected.IsNone();
+}
+
+void SModalSessionBrowser::UpdateCurrentSessionInfoMap()
+{
+	using namespace UE::Trace;
+	const FStoreClient* StoreClient = FStoreClient::Connect(*CurrentTraceStoreAddress);
+
+	if (!StoreClient)
+	{
+		UE_LOG(LogSessionBrowser, Error, TEXT("Failed to connect to trace store at [%s]"), *CurrentTraceStoreAddress)
+		return;
+	}
+
+	const uint32 SessionCount = StoreClient->GetSessionCount();
+	CurrentSessionInfosMap.Empty(SessionCount);
+
+	for (uint32 SessionIndex = 0; SessionIndex < SessionCount; SessionIndex++)
+	{
+		if (const FStoreClient::FSessionInfo* SessionInfo = StoreClient->GetSessionInfo(SessionIndex))
+		{
+			//TODO: We need a way to know if it is an IPV6
+			FCVDIPv4Address IP(SessionInfo->GetIpAddress());
+			
+			FStringFormatOrderedArguments Args {IP.ToString(), FString::FromInt(SessionInfo->GetId()) };
+			FString SessionDisplayNameString = FString::Format(TEXT("IP: {0} + {1}"), Args);
+
+			FTraceSessionInfo CVDSessionInfo { SessionInfo->GetTraceId(), SessionInfo->GetIpAddress(), SessionInfo->GetControlPort(), true };
+
+			CurrentSessionInfosMap.Add(FName(SessionDisplayNameString), CVDSessionInfo);
+		}	
+	}
+
+	TArray<TSharedPtr<FName>> NewNameList;
+	NewNameList.Reserve(CurrentSessionInfosMap.Num());
+
+	Algo::Transform(CurrentSessionInfosMap, NewNameList, [](const TPair<FName, FTraceSessionInfo>& TransformData){ return MakeShared<FName>(TransformData.Key);});
+	
+	NamePickerWidget->UpdateNameList(MoveTemp(NewNameList));
+}
+
+FReply SModalSessionBrowser::OnButtonClick(EAppReturnType::Type ButtonID)
+{
+	UserResponse = ButtonID;
+	RequestDestroyWindow();
+	return FReply::Handled();
+}
+
+void SModalSessionBrowser::HandleSessionNameSelected(TSharedPtr<FName> SelectedName)
+{
+	CurrentTraceSessionSelected = SelectedName.IsValid() ? *SelectedName : FName();
+}
+
+void SModalSessionBrowser::OnTraceStoreAddressUpdated(const FText& InText, ETextCommit::Type CommitType)
+{
+	if (CommitType == ETextCommit::OnEnter || CommitType == ETextCommit::Default)
+	{
+		CurrentTraceStoreAddress = InText.ToString();
+	}
+
+	UpdateCurrentSessionInfoMap();
+}
+
+#undef LOCTEXT_NAMESPACE
diff --git a/Engine/Plugins/Animation/GameplayInsights/Source/RewindDebugger/Private/SModalSessionBrowser.h b/Engine/Plugins/Animation/GameplayInsights/Source/RewindDebugger/Private/SModalSessionBrowser.h
new file mode 100644
--- /dev/null
+++ b/Engine/Plugins/Animation/GameplayInsights/Source/RewindDebugger/Private/SModalSessionBrowser.h
@@ -0,0 +1,83 @@
+// Copyright Epic Games, Inc. All Rights Reserved.
+
+#pragma once
+
+#include "Containers/Map.h"
+#include "Templates/SharedPointer.h"
+#include "Widgets/SWindow.h"
+
+class FName;
+class FText;
+class FLiveSessionTracker;
+class SModalSessionNameListPicker;
+class SEditableTextBox;
+
+namespace EAppReturnType { enum Type; }
+
+/** Structure containing info about a Trace Session */
+
+/**
+ * Modal window used to find, show and select active Trace Sessions to be used
+ * with Chaos Visual Debugger
+ */
+class SModalSessionBrowser : public SWindow
+{
+public:
+	struct FTraceSessionInfo
+	{
+		uint32 TraceID = 0;
+		uint32 IPAddress = 0;
+		uint32 ControlPort = 0;
+		bool bIsValid = false;
+	};
+
+	SLATE_BEGIN_ARGS(SModalSessionBrowser)
+		{
+		}
+	SLATE_END_ARGS()
+
+	virtual ~SModalSessionBrowser() override;
+
+	void Construct(const FArguments& InArgs);
+
+	/** Displays the dialog in a blocking fashion */
+	EAppReturnType::Type ShowModal();
+
+	/** Response that triggered the closing of this modal */
+	EAppReturnType::Type GetResponse() const { return UserResponse; }
+
+	/** Returns information about the active trace session selected */
+	FTraceSessionInfo GetSelectedTraceInfo();
+
+	/** Returns the address of the trace store selected while looking for active Trace Sessions */
+	FString GetSelectedTraceStoreAddress() const { return CurrentTraceStoreAddress; };
+
+	void ModalTick(float InDeltaTime);
+
+protected:
+
+	bool CanOpenSession() const;
+	
+	void UpdateCurrentSessionInfoMap();
+
+	FReply OnButtonClick(EAppReturnType::Type ButtonID);
+	
+	void HandleSessionNameSelected(TSharedPtr<FName> SelectedName);
+	
+	void OnTraceStoreAddressUpdated(const FText& InText, ETextCommit::Type CommitType);
+
+	TSharedPtr<SModalSessionNameListPicker> NamePickerWidget;
+	
+	TSharedPtr<SEditableTextBox> TraceStoreAddressWidget;
+	
+	TMap<FName, FTraceSessionInfo> CurrentSessionInfosMap;
+
+	EAppReturnType::Type UserResponse = EAppReturnType::Cancel;
+
+	FString CurrentTraceStoreAddress;
+
+	FName CurrentTraceSessionSelected;
+
+	FDelegateHandle ModalTickHandle;
+	float AccumulatedTimeBetweenTicks = 0.0f;
+};
diff --git a/Engine/Plugins/Animation/GameplayInsights/Source/RewindDebugger/Private/SModalSessionNameListPicker.cpp b/Engine/Plugins/Animation/GameplayInsights/Source/RewindDebugger/Private/SModalSessionNameListPicker.cpp
new file mode 100644
--- /dev/null
+++ b/Engine/Plugins/Animation/GameplayInsights/Source/RewindDebugger/Private/SModalSessionNameListPicker.cpp
@@ -0,0 +1,145 @@
+// Copyright Epic Games, Inc. All Rights Reserved.
+
+#include "SModalSessionNameListPicker.h"
+#include "Widgets/Input/SComboButton.h"
+#include "Widgets/Layout/SExpandableArea.h"
+#include "Widgets/Text/STextBlock.h"
+#include "Widgets/Views/STableRow.h"
+#include "Widgets/Views/SListView.h"
+#include "Widgets/SBoxPanel.h"
+
+#define LOCTEXT_NAMESPACE "ModalSessionNameListPicker"
+
+void SModalSessionNameListPicker::Construct(const FArguments& InArgs)
+{
+	NameSelectedDelegate = InArgs._OnNameSleceted;
+
+	// Dropdown Button
+	 SAssignNew(DropDownButtonLabelTextBox, STextBlock)
+		.Font(FAppStyle::GetFontStyle(TEXT("PropertyWindow.NormalFont")))
+		.Text(this, &SModalSessionNameListPicker::GetCurrentDropDownButtonLabel);
+
+	// Name List
+	SAssignNew(NameListWidget, SListView<TSharedPtr<FName>>)
+		.ListItemsSource(&CachedNameList)
+		.SelectionMode( ESelectionMode::Single )
+		.OnMouseButtonClick(this, &SModalSessionNameListPicker::OnClickItem)
+		.ListViewStyle(&FAppStyle::Get().GetWidgetStyle<FTableViewStyle>("SimpleListView"))
+		.OnGenerateRow(this, &SModalSessionNameListPicker::MakeNameItemForList);
+
+	// Dropdown List content
+	SAssignNew(NamesListDropdown, SListViewSelectorDropdownMenu< TSharedPtr<FName> >, nullptr, NameListWidget)
+	[
+		SNew(SBorder)
+		.BorderImage(FAppStyle::GetBrush("Menu.Background"))
+		.Padding(2)
+		[
+			SNew(SBox)
+			[				
+				SNew(SVerticalBox)
+				+SVerticalBox::Slot()
+				.MaxHeight(200.0f)
+				[
+					NameListWidget.ToSharedRef()
+				]
+			]
+		]
+	];
+	
+	ChildSlot
+	[
+		// Combo button that summons the dropdown menu
+		SAssignNew(PickerComboButton, SComboButton)
+			.IsEnabled_Raw(this, &SModalSessionNameListPicker::HasElements)
+			.ButtonContent()
+			[
+				DropDownButtonLabelTextBox.ToSharedRef()
+			]
+			.MenuContent()
+			[
+				NamesListDropdown.ToSharedRef()
+			]
+			.IsFocusable(true)
+			.ContentPadding(2.0f)
+			.OnComboBoxOpened(this, &SModalSessionNameListPicker::OnListOpened)
+	];
+}
+
+void SModalSessionNameListPicker::UpdateNameList(TArray<TSharedPtr<FName>>&& NewNameList)
+{
+	// TODO: It might be worth to check is the array is different before doing this, but it might be more expensive
+	CachedNameList = MoveTemp(NewNameList);
+
+	if(!HasElements())
+	{
+		CurrentSelectedName = nullptr;
+	}
+
+	NameListWidget->RebuildList();
+}
+
+void SModalSessionNameListPicker::SelectName(const TSharedPtr<FName>& NameToSelect, ESelectInfo::Type SelectionInfo)
+{
+	NameListWidget->SetSelection(NameToSelect, ESelectInfo::OnMouseClick);
+	NameListWidget->RequestListRefresh();
+	OnClickItem(NameToSelect);
+}
+
+bool SModalSessionNameListPicker::HasElements() const
+{
+	return CachedNameList.Num() > 0;
+}
+
+FText SModalSessionNameListPicker::GetCurrentDropDownButtonLabel() const
+{
+	if (CurrentSelectedName.IsValid())
+	{
+		return FText::FromName(*CurrentSelectedName);
+	}
+	else
+	{
+		return HasElements() ? LOCTEXT("NameList", "Pick an element...") : LOCTEXT("EmptyList", "Empty...");
+	}
+}
+
+void SModalSessionNameListPicker::OnClickItem(TSharedPtr<FName> Item)
+{
+	CurrentSelectedName = Item;
+
+	if (PickerComboButton.IsValid())
+	{
+		PickerComboButton->SetIsOpen(false);
+	}
+	
+	NameSelectedDelegate.ExecuteIfBound(Item);
+}
+
+void SModalSessionNameListPicker::OnListOpened()
+{
+	if (NameListWidget.IsValid())
+	{
+		NameListWidget->SetSelection(CurrentSelectedName, ESelectInfo::OnKeyPress);
+		NameListWidget->RequestScrollIntoView(CurrentSelectedName);
+	}
+}
+
+TSharedRef<ITableRow> SModalSessionNameListPicker::MakeNameItemForList(TSharedPtr<FName> Name, const TSharedRef<STableViewBase>& OwnerTable)
+{
+	TSharedPtr<STableRow<TSharedPtr<FString>>> ListTableRow;
+
+	SAssignNew(ListTableRow, STableRow<TSharedPtr<FString>>, OwnerTable)
+		[
+			SNew(SVerticalBox)
+				+SVerticalBox::Slot()
+				.AutoHeight()
+				.Padding(20.0f, 1.0f, 20.0f, 1.0f)
+				[
+					SNew(STextBlock)
+					.Text(FText::FromName(Name.IsValid() ? *Name.Get() : FName("Invalid")))
+				]
+		];
+
+	return ListTableRow.ToSharedRef();
+}
+
+#undef LOCTEXT_NAMESPACE
diff --git a/Engine/Plugins/Animation/GameplayInsights/Source/RewindDebugger/Private/SModalSessionNameListPicker.h b/Engine/Plugins/Animation/GameplayInsights/Source/RewindDebugger/Private/SModalSessionNameListPicker.h
new file mode 100644
--- /dev/null
+++ b/Engine/Plugins/Animation/GameplayInsights/Source/RewindDebugger/Private/SModalSessionNameListPicker.h
@@ -0,0 +1,58 @@
+// Copyright Epic Games, Inc. All Rights Reserved.
+
+#pragma once
+#include "SListViewSelectorDropdownMenu.h"
+#include "Widgets/SCompoundWidget.h"
+#include "Containers/Array.h"
+#include "Templates/SharedPointer.h"
+
+class SComboButton;
+class STextBlock;
+
+DECLARE_DELEGATE_OneParam(FModalSessionNameListSelected, TSharedPtr<FName>)
+
+/** Widget that takes a list of String Views and creates a drop down menu for it for the Chaos Visual Debugger Tool.
+ */
+class SModalSessionNameListPicker : public SCompoundWidget
+{
+public:
+	SLATE_BEGIN_ARGS(SModalSessionNameListPicker){}
+	SLATE_EVENT(FModalSessionNameListSelected, OnNameSleceted)
+	SLATE_END_ARGS()
+	
+	void Construct(const FArguments& InArgs);
+
+	/**
+	 * Replaces the names in the list view
+	 * @param NewNameList Array of Name to represent.
+	 */
+	void UpdateNameList(TArray<TSharedPtr<FName>>&& NewNameList);
+	
+	void SelectName(const TSharedPtr<FName>& NameToSelect, ESelectInfo::Type SelectionInfo = ESelectInfo::Direct);
+	
+protected:
+
+	bool HasElements() const;
+
+	FText GetCurrentDropDownButtonLabel() const;
+	
+	void OnClickItem(TSharedPtr<FName> Item);
+
+	void OnListOpened();
+
+	TSharedRef<ITableRow> MakeNameItemForList(TSharedPtr<FName> Name, const TSharedRef<STableViewBase>& OwnerTable);
+
+	TSharedPtr<SComboButton> PickerComboButton;
+
+	TSharedPtr<STextBlock> DropDownButtonLabelTextBox;
+	
+	TSharedPtr<SListView<TSharedPtr<FName>>> NameListWidget;
+
+	TSharedPtr<SListViewSelectorDropdownMenu<TSharedPtr<FName>>> NamesListDropdown;
+
+	TArray<TSharedPtr<FName>> CachedNameList;
+
+	TSharedPtr<FName> CurrentSelectedName;
+
+	FModalSessionNameListSelected NameSelectedDelegate;
+};
