optimizations/bug fixes for rewind debugger

diff --git a/Engine/Plugins/Animation/GameplayInsights/Source/GameplayInsights/Private/AnimationProvider.cpp b/Engine/Plugins/Animation/GameplayInsights/Source/GameplayInsights/Private/AnimationProvider.cpp
--- a/Engine/Plugins/Animation/GameplayInsights/Source/GameplayInsights/Private/AnimationProvider.cpp
+++ b/Engine/Plugins/Animation/GameplayInsights/Source/GameplayInsights/Private/AnimationProvider.cpp
@@ -28,13 +28,10 @@ void FAnimationProvider::EnumerateSkeletalMeshPoseTimelines(TFunctionRef<void(ui
 	
 	for(auto& IndexMapping : ObjectIdToSkeletalMeshPoseTimelines)
 	{
-		for (const TSharedRef<TraceServices::TIntervalTimeline<FAnimGraphMessage>>& Timeline : AnimGraphTimelines)
+		const TSharedPtr<FSkeletalMeshTimelineStorage>& TimelineStorage = SkeletalMeshPoseTimelineStorage[IndexMapping.Value];
+		if (TimelineStorage->Timeline.IsValid())
 		{
-			const TSharedPtr<FSkeletalMeshTimelineStorage>& TimelineStorage = SkeletalMeshPoseTimelineStorage[IndexMapping.Value];
-			if (TimelineStorage->Timeline.IsValid())
-			{
-				Callback(IndexMapping.Key, *TimelineStorage->Timeline);
-			}
+			Callback(IndexMapping.Key, *TimelineStorage->Timeline);
 		}
 	}
 }
diff --git a/Engine/Plugins/Animation/GameplayInsights/Source/RewindDebugger/Private/RewindDebuggerAnimation.cpp b/Engine/Plugins/Animation/GameplayInsights/Source/RewindDebugger/Private/RewindDebuggerAnimation.cpp
--- a/Engine/Plugins/Animation/GameplayInsights/Source/RewindDebugger/Private/RewindDebuggerAnimation.cpp
+++ b/Engine/Plugins/Animation/GameplayInsights/Source/RewindDebugger/Private/RewindDebuggerAnimation.cpp
@@ -178,6 +178,7 @@ void FRewindDebuggerAnimation::ApplyPoseToMesh(const IAnimationProvider* Animati
 			MeshComponent->UpdateLODStatus();
 			MeshComponent->UpdateChildTransforms(EUpdateTransformFlags::None, ETeleportType::TeleportPhysics);
 			MeshComponent->SetVisibility(PoseMessage->bIsVisible);
+			MeshComponent->MarkRenderStateDirty();
 		}
 	}
 }
