Fix rewind debugger crash with auto-recording if pie is started and then the rewind debugger tab is opened

diff --git a/Engine/Plugins/Animation/GameplayInsights/Source/RewindDebugger/Private/RewindDebugger.cpp b/Engine/Plugins/Animation/GameplayInsights/Source/RewindDebugger/Private/RewindDebugger.cpp
--- a/Engine/Plugins/Animation/GameplayInsights/Source/RewindDebugger/Private/RewindDebugger.cpp
+++ b/Engine/Plugins/Animation/GameplayInsights/Source/RewindDebugger/Private/RewindDebugger.cpp
@@ -185,7 +185,7 @@ void FRewindDebugger::OnPIEStarted(bool bSimulating)
 
 	if (ShouldAutoRecordOnPIE())
 	{
-		StartRecording();
+		bQueueStartRecording = true;
 	}
 }
 
@@ -1031,6 +1031,13 @@ void FRewindDebugger::Tick(float DeltaTime)
 {
 	TRACE_CPUPROFILER_EVENT_SCOPE(FRewindDebugger::Tick);
 
+
+	if (bQueueStartRecording)
+	{
+		StartRecording();
+		bQueueStartRecording = false;
+	}
+
 	if (const TraceServices::IAnalysisSession* Session = GetAnalysisSession())
 	{
 		const IAnimationProvider* AnimationProvider = Session->ReadProvider<IAnimationProvider>("AnimationProvider");
diff --git a/Engine/Plugins/Animation/GameplayInsights/Source/RewindDebugger/Private/RewindDebugger.h b/Engine/Plugins/Animation/GameplayInsights/Source/RewindDebugger/Private/RewindDebugger.h
--- a/Engine/Plugins/Animation/GameplayInsights/Source/RewindDebugger/Private/RewindDebugger.h
+++ b/Engine/Plugins/Animation/GameplayInsights/Source/RewindDebugger/Private/RewindDebugger.h
@@ -179,6 +179,7 @@ private:
 	FOnComponentListChanged ComponentListChangedDelegate;
 	FOnTrackCursor TrackCursorDelegate;
 
+	bool bQueueStartRecording = false;
 	bool bTraceJustConnected = false;
 	bool bPIEStarted = false;
 	bool bPIESimulating = false;
