Rewind Debugger: Add support to specify trace filename or host in RewindDebugger.StartRecording

diff --git a/Engine/Plugins/Animation/GameplayInsights/Source/RewindDebuggerRuntime/Private/RewindDebuggerRuntime.cpp b/Engine/Plugins/Animation/GameplayInsights/Source/RewindDebuggerRuntime/Private/RewindDebuggerRuntime.cpp
--- a/Engine/Plugins/Animation/GameplayInsights/Source/RewindDebuggerRuntime/Private/RewindDebuggerRuntime.cpp
+++ b/Engine/Plugins/Animation/GameplayInsights/Source/RewindDebuggerRuntime/Private/RewindDebuggerRuntime.cpp
@@ -6,6 +6,7 @@
 #include "EngineUtils.h"
 #include "Engine/World.h"
 #include "GameFramework/Pawn.h"
+#include "Misc/FileHelper.h"
 #include "ProfilingDebugging/TraceAuxiliary.h"
 #include "UObject/UObjectIterator.h"
 #include "RewindDebuggerRuntimeInterface/IRewindDebuggerRuntimeExtension.h"
@@ -55,7 +56,65 @@ namespace RewindDebugger
 		, nullptr);
 	}
 	
+	void FRewindDebuggerRuntime::StartRecordingWithArgs(const TArray<FString>& Args)
+	{
+		if (Args.Num() > 0)
+		{
+			FTraceAuxiliary::EConnectionType TraceType = FTraceAuxiliary::EConnectionType::None;
+			FString TraceDestination;
+
+			for (const FString& Arg : Args)
+			{
+				if (Arg.StartsWith(TEXT("-tracefile"), ESearchCase::IgnoreCase))
+				{
+					ensureMsgf(TraceType == FTraceAuxiliary::EConnectionType::None, TEXT("RewindDebugger.StartRecording: Specifying more than 1 trace destination is not supported. Received: %s"), *FString::Join(Args, TEXT(" ")));
+				
+					TraceType = FTraceAuxiliary::EConnectionType::File;
+					
+					// Try to extract filename.
+					if (FParse::Value(*Arg, TEXT("-tracefile="), TraceDestination))
+					{
+						// Make sure it's a valid filename
+						FText FilenameError;
+						if (!FFileHelper::IsFilenameValidForSaving(TraceDestination, FilenameError))
+						{
+							ensureMsgf(TraceType == FTraceAuxiliary::EConnectionType::None, TEXT("RewindDebugger.StartRecording: Specified filename is not supported: %s"), *FilenameError.ToString());
+							TraceDestination = "";
+						}
+					}
+				}
+				else if (Arg.StartsWith(TEXT("-tracehost"), ESearchCase::IgnoreCase))
+				{
+					ensureMsgf(TraceType == FTraceAuxiliary::EConnectionType::None, TEXT("RewindDebugger.StartRecording: Specifying more than 1 trace destination is not supported. Received: %s"), *FString::Join(Args, TEXT(" ")));
+
+					if (FParse::Value(*Arg, TEXT("-tracehost="), TraceDestination))
+					{
+						// Should we validate that TraceDestination is valid ip address?
+					}
+				}
+				else
+				{
+					ensureMsgf(false, TEXT("RewindDebugger.StartRecording: Received unknown argument: %s"), *Arg);
+				}
+			}
+
+			if (TraceType != FTraceAuxiliary::EConnectionType::None)
+			{
+				StartRecording(TraceType, *TraceDestination);
+				return;
+			}
+		}
+
+		// No destination was specified, just use the default.
+		StartRecording();
+	}
+	
 	void FRewindDebuggerRuntime::StartRecording()
+	{
+		StartRecording(FTraceAuxiliary::EConnectionType::Network, TEXT("127.0.0.1"));
+	}
+
+	void FRewindDebuggerRuntime::StartRecording(FTraceAuxiliary::EConnectionType TraceType, const TCHAR* TraceDestination)
 	{
 #if OBJECT_TRACE_ENABLED
 		// Clear caches  (maybe move these to plugins?)
@@ -76,7 +135,7 @@ namespace RewindDebugger
 	
 		// FTraceAuxiliary::OnConnection.AddRaw(this, &FRewindDebugger::OnConnection); 
 	
-		FTraceAuxiliary::Start(FTraceAuxiliary::EConnectionType::Network, TEXT("127.0.0.1"), TEXT(""), &Options, LogRewindDebuggerRuntime);
+		FTraceAuxiliary::Start(TraceType, TraceDestination, TEXT(""), &Options, LogRewindDebuggerRuntime);
 		
 		UE::Trace::ToggleChannel(TEXT("Object"), true);
 		UE::Trace::ToggleChannel(TEXT("ObjectProperties"), true);
diff --git a/Engine/Plugins/Animation/GameplayInsights/Source/RewindDebuggerRuntime/Private/RewindDebuggerRuntimeModule.cpp b/Engine/Plugins/Animation/GameplayInsights/Source/RewindDebuggerRuntime/Private/RewindDebuggerRuntimeModule.cpp
--- a/Engine/Plugins/Animation/GameplayInsights/Source/RewindDebuggerRuntime/Private/RewindDebuggerRuntimeModule.cpp
+++ b/Engine/Plugins/Animation/GameplayInsights/Source/RewindDebuggerRuntime/Private/RewindDebuggerRuntimeModule.cpp
@@ -18,7 +18,7 @@ void FRewindDebuggerRuntimeModule::StartupModule()
 	ConsoleObjects.Add(IConsoleManager::Get().RegisterConsoleCommand(
 		TEXT("RewindDebugger.StartRecording"),
 		TEXT("Starts making a rewind debugger recording."),
-		FConsoleCommandDelegate::CreateRaw(RewindDebugger::FRewindDebuggerRuntime::Instance(), &RewindDebugger::FRewindDebuggerRuntime::StartRecording),
+		FConsoleCommandWithArgsDelegate::CreateRaw(RewindDebugger::FRewindDebuggerRuntime::Instance(), &RewindDebugger::FRewindDebuggerRuntime::StartRecordingWithArgs),
 		ECVF_Default
 	));
 	
diff --git a/Engine/Plugins/Animation/GameplayInsights/Source/RewindDebuggerRuntime/Public/RewindDebuggerRuntime/RewindDebuggerRuntime.h b/Engine/Plugins/Animation/GameplayInsights/Source/RewindDebuggerRuntime/Public/RewindDebuggerRuntime/RewindDebuggerRuntime.h
--- a/Engine/Plugins/Animation/GameplayInsights/Source/RewindDebuggerRuntime/Public/RewindDebuggerRuntime/RewindDebuggerRuntime.h
+++ b/Engine/Plugins/Animation/GameplayInsights/Source/RewindDebuggerRuntime/Public/RewindDebuggerRuntime/RewindDebuggerRuntime.h
@@ -3,6 +3,7 @@
 #pragma once
 
 #include "CoreMinimal.h"
+#include "ProfilingDebugging/TraceAuxiliary.h" // for FTraceAuxiliary::EConnectionType
 
 REWINDDEBUGGERRUNTIME_API DECLARE_LOG_CATEGORY_EXTERN(LogRewindDebuggerRuntime, Log, All);
 
@@ -20,12 +21,16 @@ namespace RewindDebugger
 		void StartRecording();
 		void StopRecording();
 
+		void StartRecordingWithArgs(const TArray<FString>& Args);
+
 		bool IsRecording() const { return bIsRecording; }
 
 		FSimpleMulticastDelegate RecordingStarted;
 		FSimpleMulticastDelegate ClearRecording;
 		FSimpleMulticastDelegate RecordingStopped;
 	private:
+		void StartRecording(FTraceAuxiliary::EConnectionType TraceType, const TCHAR* TraceDestination);
+
 		bool bIsRecording = false;
 		static FRewindDebuggerRuntime* InternalInstance;
 	};
