Fix ensure conditions in RewindDebuggerRuntime.cpp

diff --git a/Engine/Plugins/Animation/GameplayInsights/Source/RewindDebuggerRuntime/Private/RewindDebuggerRuntime.cpp b/Engine/Plugins/Animation/GameplayInsights/Source/RewindDebuggerRuntime/Private/RewindDebuggerRuntime.cpp
--- a/Engine/Plugins/Animation/GameplayInsights/Source/RewindDebuggerRuntime/Private/RewindDebuggerRuntime.cpp
+++ b/Engine/Plugins/Animation/GameplayInsights/Source/RewindDebuggerRuntime/Private/RewindDebuggerRuntime.cpp
@@ -78,7 +78,7 @@ namespace RewindDebugger
 						FText FilenameError;
 						if (!FFileHelper::IsFilenameValidForSaving(TraceDestination, FilenameError))
 						{
-							ensureMsgf(TraceType == FTraceAuxiliary::EConnectionType::None, TEXT("RewindDebugger.StartRecording: Specified filename is not supported: %s"), *FilenameError.ToString());
+							ensureMsgf(false, TEXT("RewindDebugger.StartRecording: Specified filename is not supported: %s"), *FilenameError.ToString());
 							TraceDestination = "";
 						}
 					}
@@ -87,6 +87,8 @@ namespace RewindDebugger
 				{
 					ensureMsgf(TraceType == FTraceAuxiliary::EConnectionType::None, TEXT("RewindDebugger.StartRecording: Specifying more than 1 trace destination is not supported. Received: %s"), *FString::Join(Args, TEXT(" ")));
 
+					TraceType = FTraceAuxiliary::EConnectionType::Network;
+
 					if (FParse::Value(*Arg, TEXT("-tracehost="), TraceDestination))
 					{
 						// Should we validate that TraceDestination is valid ip address?
