Fix text rendering in Rewind debugger Visual Logger integration

diff --git a/Engine/Plugins/Animation/GameplayInsights/Source/RewindDebuggerVLog/Private/RewindDebuggerVLog.cpp b/Engine/Plugins/Animation/GameplayInsights/Source/RewindDebuggerVLog/Private/RewindDebuggerVLog.cpp
--- a/Engine/Plugins/Animation/GameplayInsights/Source/RewindDebuggerVLog/Private/RewindDebuggerVLog.cpp
+++ b/Engine/Plugins/Animation/GameplayInsights/Source/RewindDebuggerVLog/Private/RewindDebuggerVLog.cpp
@@ -10,7 +10,9 @@
 #include "RewindDebuggerVLogSettings.h"
 #include "ToolMenus.h"
 #include "VisualLogEntryRenderer.h"
+#include "Debug/DebugDrawService.h"
 #include "Editor/EditorEngine.h"
+#include "Engine/Canvas.h"
 #include "Engine/World.h"
 #include "HAL/IConsoleManager.h"
 #include "TraceServices/Model/Frames.h"
@@ -25,6 +27,55 @@ FRewindDebuggerVLog::FRewindDebuggerVLog()
 
 }
 
+void FRewindDebuggerVLog::OnShowDebugInfo(UCanvas* Canvas, APlayerController* Player)
+{
+	ScreenTextY = 100;
+	if (IRewindDebugger* RewindDebugger = IRewindDebugger::Instance())
+	{
+		if (RewindDebugger->IsPIESimulating())
+		{
+			// make sure this is the primary view, when we are playing in PIE, so we don't clear ImmediateRenderQueue when this has been called on some other editor view.
+			if (Canvas->SceneView->ViewActor)
+			{
+				for (FVisualLogEntry& Entry : ImmediateRenderQueue)
+				{
+					RenderLogEntry(Entry, Canvas);
+				}
+
+				ImmediateRenderQueue.SetNum(0);
+			}
+		}
+		else
+		{
+			ObjectsVisited.Empty();
+			if (const TraceServices::IAnalysisSession* Session = RewindDebugger->GetAnalysisSession())
+			{
+				TraceServices::FAnalysisSessionReadScope SessionReadScope(*Session);
+				double CurrentTraceTime = RewindDebugger->CurrentTraceTime();
+
+				const TraceServices::IFrameProvider& FrameProvider = TraceServices::ReadFrameProvider(*Session);
+				TraceServices::FFrame CurrentFrame;
+
+				if(FrameProvider.GetFrameFromTime(ETraceFrameType::TraceFrameType_Game, CurrentTraceTime, CurrentFrame))
+				{
+					if (const IVisualLoggerProvider* VisualLoggerProvider = Session->ReadProvider<IVisualLoggerProvider>("VisualLoggerProvider"))
+					{
+						AddLogEntries(RewindDebugger->GetDebugComponents(), CurrentFrame.StartTime, CurrentFrame.EndTime, VisualLoggerProvider, Canvas);
+					}
+				}
+			}
+		}
+	}
+
+
+	
+}
+
+FRewindDebuggerVLog::~FRewindDebuggerVLog()
+{
+	UDebugDrawService::Unregister(DelegateHandle);
+}
+
 void FRewindDebuggerVLog::Initialize()
 {
 	UToolMenu* Menu = UToolMenus::Get()->FindMenu("RewindDebugger.ToolBar");
@@ -49,6 +100,11 @@ void FRewindDebuggerVLog::Initialize()
 		
 	FVisualLoggerTraceDevice& TraceDevice = FVisualLoggerTraceDevice::Get();
 	TraceDevice.ImmediateRenderDelegate.BindRaw(this, &FRewindDebuggerVLog::ImmediateRender);
+
+	DelegateHandle = UDebugDrawService::Register(TEXT("VirtualTextureResidency")/*TEXT("VisLog")*/, FDebugDrawDelegate::CreateRaw(this, &FRewindDebuggerVLog::OnShowDebugInfo));
+
+	FTopLevelAssetPath MonospaceFontPath = FTopLevelAssetPath("/Engine/EngineFonts/DroidSansMono.DroidSansMono");
+	MonospaceFont = LoadObject<UFont>(nullptr, *MonospaceFontPath.ToString(), nullptr, LOAD_None, nullptr);
 }
 
 bool ContainsObject(TArray<TSharedPtr<FDebugObjectInfo>>& Components, uint64 ObjectId)
@@ -76,7 +132,7 @@ bool MatchCategoryFilters(const FName& CategoryName, ELogVerbosity::Type Verbosi
 	return Settings.DisplayCategories.Contains(CategoryName) && Verbosity <= Settings.DisplayVerbosity;
 }
 
-void FRewindDebuggerVLog::RenderLogEntry(const FVisualLogEntry& Entry)
+void FRewindDebuggerVLog::RenderLogEntry(const FVisualLogEntry& Entry, UCanvas* Canvas)
 {
 	if (CVarRewindDebuggerVLogUseActor.GetValueOnAnyThread())
 	{
@@ -89,9 +145,8 @@ void FRewindDebuggerVLog::RenderLogEntry(const FVisualLogEntry& Entry)
 	else
 	{
 		UWorld* World = IRewindDebugger::Instance()->GetWorldToVisualize();
-		const FVisualLogShapeElement* ElementToDraw = Entry.ElementsToDraw.GetData();
-		const int32 ElementsCount = Entry.ElementsToDraw.Num();
-		FVisualLogEntryRenderer::RenderLogEntry(World,Entry, &MatchCategoryFilters);
+		FVisualLogEntryRenderer::RenderLogEntry(World,Entry, &MatchCategoryFilters, Canvas, MonospaceFont, ScreenTextY);
+
 	}
 }
 
@@ -103,7 +158,7 @@ void FRewindDebuggerVLog::ImmediateRender(const UObject* Object, const FVisualLo
 		uint64 ObjectId = FObjectTrace::GetObjectId(Object);
 		if (ContainsObject(RewindDebugger->GetDebugComponents(), ObjectId))
 		{
-			RenderLogEntry(Entry);
+			ImmediateRenderQueue.Add(Entry);
 		}
 	}
 #endif
@@ -182,20 +237,27 @@ void FRewindDebuggerVLog::MakeCategoriesMenu(UToolMenu* Menu)
 	}
 }
 
-void FRewindDebuggerVLog::AddLogEntries(const TArray<TSharedPtr<FDebugObjectInfo>>& Components, float StartTime, float EndTime, const IVisualLoggerProvider* VisualLoggerProvider)
+void FRewindDebuggerVLog::AddLogEntries(const TArray<TSharedPtr<FDebugObjectInfo>>& Components, float StartTime, float EndTime, const IVisualLoggerProvider* VisualLoggerProvider, UCanvas* Canvas)
 {
 	for(const TSharedPtr<FDebugObjectInfo>& ComponentInfo : Components)
 	{
-		VisualLoggerProvider->ReadVisualLogEntryTimeline(ComponentInfo->ObjectId, [this,StartTime, EndTime](const IVisualLoggerProvider::VisualLogEntryTimeline &TimelineData)
+		if (!ObjectsVisited.Contains(ComponentInfo->ObjectId))
 		{
-			TimelineData.EnumerateEvents(StartTime, EndTime, [this](double InStartTime, double InEndTime, uint32 InDepth, const FVisualLogEntry& LogEntry)
+			ObjectsVisited.Add(ComponentInfo->ObjectId);
+			VisualLoggerProvider->ReadVisualLogEntryTimeline(ComponentInfo->ObjectId, [this,StartTime, EndTime, Canvas](const IVisualLoggerProvider::VisualLogEntryTimeline &TimelineData)
 			{
-				RenderLogEntry(LogEntry);
-				return TraceServices::EEventEnumerate::Continue;
+				TimelineData.EnumerateEvents(StartTime, EndTime, [this, StartTime, EndTime, Canvas](double InStartTime, double InEndTime, uint32 InDepth, const FVisualLogEntry& LogEntry)
+				{
+					if (InStartTime >= StartTime && InStartTime <= EndTime)
+					{
+						RenderLogEntry(LogEntry, Canvas);
+					}
+					return TraceServices::EEventEnumerate::Continue;
+				});
 			});
-		});
+		}
 
-		AddLogEntries(ComponentInfo->Children, StartTime, EndTime, VisualLoggerProvider);
+		AddLogEntries(ComponentInfo->Children, StartTime, EndTime, VisualLoggerProvider, Canvas);
 	}
 }
 
@@ -216,25 +278,6 @@ AVLogRenderingActor* FRewindDebuggerVLog::GetRenderingActor()
 
 void FRewindDebuggerVLog::Update(float DeltaTime, IRewindDebugger* RewindDebugger)
 {
-	if (!RewindDebugger->IsPIESimulating())
-	{
-		if (const TraceServices::IAnalysisSession* Session = RewindDebugger->GetAnalysisSession())
-		{
-			TraceServices::FAnalysisSessionReadScope SessionReadScope(*Session);
-			double CurrentTraceTime = RewindDebugger->CurrentTraceTime();
-
-			const TraceServices::IFrameProvider& FrameProvider = TraceServices::ReadFrameProvider(*Session);
-			TraceServices::FFrame CurrentFrame;
-
-			if(FrameProvider.GetFrameFromTime(ETraceFrameType::TraceFrameType_Game, CurrentTraceTime, CurrentFrame))
-			{
-				if (const IVisualLoggerProvider* VisualLoggerProvider = Session->ReadProvider<IVisualLoggerProvider>("VisualLoggerProvider"))
-				{
-					AddLogEntries(RewindDebugger->GetDebugComponents(), CurrentFrame.StartTime, CurrentFrame.EndTime, VisualLoggerProvider);
-				}
-			}
-		}
-	}
 }
 
 #undef LOCTEXT_NAMESPACE
\ No newline at end of file
diff --git a/Engine/Plugins/Animation/GameplayInsights/Source/RewindDebuggerVLog/Private/RewindDebuggerVLog.h b/Engine/Plugins/Animation/GameplayInsights/Source/RewindDebuggerVLog/Private/RewindDebuggerVLog.h
--- a/Engine/Plugins/Animation/GameplayInsights/Source/RewindDebuggerVLog/Private/RewindDebuggerVLog.h
+++ b/Engine/Plugins/Animation/GameplayInsights/Source/RewindDebuggerVLog/Private/RewindDebuggerVLog.h
@@ -14,7 +14,8 @@ class FRewindDebuggerVLog : public IRewindDebuggerExtension
 {
 public:
 	FRewindDebuggerVLog();
-	virtual ~FRewindDebuggerVLog() {};
+	void OnShowDebugInfo(UCanvas* Canvas, APlayerController* Player);
+	virtual ~FRewindDebuggerVLog();
 
 	void Initialize();
 	void MakeCategoriesMenu(UToolMenu* Menu);
@@ -28,12 +29,20 @@ public:
 	virtual void Update(float DeltaTime, IRewindDebugger* RewindDebugger) override;
 
 private:
-	void AddLogEntries(const TArray<TSharedPtr<FDebugObjectInfo>>& Components, float StartTime, float EndTime, const class IVisualLoggerProvider* Provider);
+	void AddLogEntries(const TArray<TSharedPtr<FDebugObjectInfo>>& Components, float StartTime, float EndTime, const class IVisualLoggerProvider* Provider, UCanvas* Canvas);
 	void ImmediateRender(const UObject* Object, const FVisualLogEntry& Entry);
-	void RenderLogEntry(const FVisualLogEntry& Entry);
+	void RenderLogEntry(const FVisualLogEntry& Entry, UCanvas* Canvas);
 
 	AVLogRenderingActor* GetRenderingActor();
 
-	TWeakObjectPtr<AVLogRenderingActor> VLogActor; 
+	TWeakObjectPtr<AVLogRenderingActor> VLogActor;
+
+	TSet<uint64> ObjectsVisited;
+	int32 ScreenTextY;
+
+	FDelegateHandle DelegateHandle;
+	UFont* MonospaceFont = nullptr;
+
+	TArray<FVisualLogEntry> ImmediateRenderQueue;
 };
 
diff --git a/Engine/Plugins/Animation/GameplayInsights/Source/RewindDebuggerVLog/Private/VisualLogEntryRenderer.cpp b/Engine/Plugins/Animation/GameplayInsights/Source/RewindDebuggerVLog/Private/VisualLogEntryRenderer.cpp
--- a/Engine/Plugins/Animation/GameplayInsights/Source/RewindDebuggerVLog/Private/VisualLogEntryRenderer.cpp
+++ b/Engine/Plugins/Animation/GameplayInsights/Source/RewindDebuggerVLog/Private/VisualLogEntryRenderer.cpp
@@ -7,22 +7,40 @@
 #include "DrawDebugHelpers.h"
 #include "GeomTools.h"
 #include "IndexTypes.h"
+#include "Engine/Canvas.h"
 #include "Engine/World.h"
 #include "VisualLogger/VisualLogger.h"
+#include "UObject/Package.h"
 
 namespace
 {
 	static const int CircleSegments = 24;
-
 	
 	// utility wrapper for drawing shape descriptions
-	inline void RenderDescription(UWorld* World, const FVisualLogShapeElement* ElementToDraw, const FVector& Position, const FColor& Color, int Index, int Count)
+	inline void RenderDescription(UWorld* World, UCanvas* Canvas, UFont* Font, const FVisualLogShapeElement* ElementToDraw, const FVector& Position, const FColor& Color, int Index, int Count)
 	{
-		if (ElementToDraw->Description.IsEmpty())
+		if (Canvas)
 		{
-			const FString PrintString = Count == 1 ? ElementToDraw->Description : FString::Printf(TEXT("%s_%d"), *ElementToDraw->Description, Index);
-			// todo: we should swap this out with something that displays when paused
-			DrawDebugString(World, Position, PrintString, nullptr, Color);
+			if (!ElementToDraw->Description.IsEmpty())
+			{
+				const FSceneView* View = Canvas->SceneView;
+				if (View->CullingFrustum.IntersectPoint(Position))
+				{
+					const FPlane Proj = View->Project(Position);
+					if (Proj.W > 0.f)
+					{
+						const int32 HalfX = Canvas->SizeX / 2;
+						const int32 HalfY = Canvas->SizeY / 2;
+
+						const int32 XPos = HalfX + static_cast<int32>(HalfX * Proj.X);
+						const int32 YPos = HalfY + static_cast<int32>(HalfY * (Proj.Y * -1));
+					
+						const FString PrintString = Count == 1 ? ElementToDraw->Description : FString::Printf(TEXT("%s_%d"), *ElementToDraw->Description, Index);
+
+						Canvas->Canvas->DrawShadowedString(XPos, YPos, PrintString, Font, Color);
+					}
+				}
+			}
 		}
 	}
 
@@ -96,8 +114,17 @@ namespace
 }
 
 
-void FVisualLogEntryRenderer::RenderLogEntry(class UWorld* World, const FVisualLogEntry& Entry, TFunctionRef<bool (const FName&, ELogVerbosity::Type)> MatchCategoryFilters)
+void FVisualLogEntryRenderer::RenderLogEntry(class UWorld* World, const FVisualLogEntry& Entry, TFunctionRef<bool (const FName&, ELogVerbosity::Type)> MatchCategoryFilters,
+	UCanvas* Canvas, UFont* Font, int32& ScreenTextY)
 {
+	for(const FVisualLogLine& LogLine : Entry.LogLines)
+	{
+		if (MatchCategoryFilters(LogLine.Category, LogLine.Verbosity))
+		{
+			Canvas->Canvas->DrawShadowedString(100, ScreenTextY+=20, LogLine.Line, Font, FColor::White);
+		}
+	}
+	
 	const FVisualLogShapeElement* ElementToDraw = Entry.ElementsToDraw.GetData();
 	const int32 ElementsCount = Entry.ElementsToDraw.Num();
 
@@ -155,7 +182,7 @@ void FVisualLogEntryRenderer::RenderLogEntry(class UWorld* World, const FVisualL
 					}
 				}
 
-				RenderDescription(World, ElementToDraw, Point, Color, Index, NumPoints);
+				RenderDescription(World, Canvas, Font, ElementToDraw, Point, Color, Index, NumPoints);
 			}
 		}
 		break;
@@ -208,7 +235,7 @@ void FVisualLogEntryRenderer::RenderLogEntry(class UWorld* World, const FVisualL
 			for (int32 Index = 0; Index + 1 < NumPoints; Index += 2, Location += 2)
 			{
 				DrawDebugLine(World, Location[0], Location[1], Color, false, -1, DepthPriority, Thickness);
-				RenderDescription(World, ElementToDraw, (Location[0] + Location[1]), Color, Index/2, NumPoints/2);
+				RenderDescription(World, Canvas, Font, ElementToDraw, (Location[0] + Location[1])/2, Color, Index/2, NumPoints/2);
 			}
 		}
 			break;
@@ -238,7 +265,7 @@ void FVisualLogEntryRenderer::RenderLogEntry(class UWorld* World, const FVisualL
 				const FBox Box = FBox(*BoxExtent, *(BoxExtent + 1));
 				// todo: wireframe, 
 				DrawDebugSolidBox(World, Box, Color, Transform, false, -1.f, DepthPriority);
-				RenderDescription(World, ElementToDraw, Transform.TransformPosition(Box.GetCenter()), Color, Index/2, NumPoints/2);
+				RenderDescription(World, Canvas, Font, ElementToDraw, Transform.TransformPosition(Box.GetCenter()), Color, Index/2, NumPoints/2);
 			}
 		}
 			break;
@@ -258,7 +285,7 @@ void FVisualLogEntryRenderer::RenderLogEntry(class UWorld* World, const FVisualL
 				Direction.FindBestAxisVectors(YAxis, ZAxis);
 
 				DrawDebugCone(World, Origin, Direction, Length, Angles.Y, Angles.Z, CircleSegments, Color, false, -1.f, DepthPriority, Thickness);
-				RenderDescription(World, ElementToDraw, Origin, Color, Index/3, ElementToDraw->Points.Num()/3);
+				RenderDescription(World, Canvas, Font, ElementToDraw, Origin, Color, Index/3, ElementToDraw->Points.Num()/3);
 			}
 		}
 			break;
@@ -275,7 +302,7 @@ void FVisualLogEntryRenderer::RenderLogEntry(class UWorld* World, const FVisualL
 				const FVector Center = 0.5 * (Start + End);
 
 				DrawDebugCylinder(World, Start, End, static_cast<float>(OtherData.X), CircleSegments, Color, false, -1.f, DepthPriority, Thickness);
-				RenderDescription(World, ElementToDraw, Center, Color, Index/3, ElementToDraw->Points.Num()/3);
+				RenderDescription(World, Canvas, Font, ElementToDraw, Center, Color, Index/3, ElementToDraw->Points.Num()/3);
 			}
 		}
 			break;
@@ -294,7 +321,7 @@ void FVisualLogEntryRenderer::RenderLogEntry(class UWorld* World, const FVisualL
 				const FQuat Rotation = FQuat(FirstData.Z, SecondData.X, SecondData.Y, SecondData.Z);
 		
 				DrawDebugCapsule(World, Base, HalfHeight, Radius, Rotation, Color, false, -1.f, DepthPriority, Thickness);
-				RenderDescription(World, ElementToDraw, Base, Color, Index/3, ElementToDraw->Points.Num()/3);
+				RenderDescription(World, Canvas, Font, ElementToDraw, Base, Color, Index/3, ElementToDraw->Points.Num()/3);
 			}
 		}
 			break;
@@ -372,7 +399,7 @@ void FVisualLogEntryRenderer::RenderLogEntry(class UWorld* World, const FVisualL
 			for (int32 Index = 0; Index + 1 < NumPoints; Index += 2, Location += 2)
 			{
 				DrawDebugDirectionalArrow(World, Location[0], Location[1], 40, Color, false, -1, DepthPriority, 1.5);
-				RenderDescription(World, ElementToDraw, (Location[0] + Location[1]) / 2, Color, Index/2, ElementToDraw->Points.Num()/2);
+				RenderDescription(World, Canvas, Font, ElementToDraw, (Location[0] + Location[1]) / 2, Color, Index/2, ElementToDraw->Points.Num()/2);
 			}
 		}
 		break;
@@ -391,7 +418,7 @@ void FVisualLogEntryRenderer::RenderLogEntry(class UWorld* World, const FVisualL
 				const float Thickness = float(ElementToDraw->Thicknes);
 
 				RenderCircle(World, Center, UpAxis, Radius, Thickness, Color, DepthPriority);
-				RenderDescription(World, ElementToDraw, Center, Color, Index/3, ElementToDraw->Points.Num()/3);
+				RenderDescription(World, Canvas, Font, ElementToDraw, Center, Color, Index/3, ElementToDraw->Points.Num()/3);
 			}
 		}
 			break;
diff --git a/Engine/Plugins/Animation/GameplayInsights/Source/RewindDebuggerVLog/Private/VisualLogEntryRenderer.h b/Engine/Plugins/Animation/GameplayInsights/Source/RewindDebuggerVLog/Private/VisualLogEntryRenderer.h
--- a/Engine/Plugins/Animation/GameplayInsights/Source/RewindDebuggerVLog/Private/VisualLogEntryRenderer.h
+++ b/Engine/Plugins/Animation/GameplayInsights/Source/RewindDebuggerVLog/Private/VisualLogEntryRenderer.h
@@ -6,5 +6,6 @@
 class FVisualLogEntryRenderer
 {
 public:
-	static void RenderLogEntry(class UWorld* World, const struct FVisualLogEntry& Entry, TFunctionRef<bool(const FName&, ELogVerbosity::Type)> MatchCategoryFilters);
+	static void RenderLogEntry(class UWorld* World, const struct FVisualLogEntry& Entry, TFunctionRef<bool(const FName&, ELogVerbosity::Type)> MatchCategoryFilters,
+		class UCanvas* Canvas, class UFont* Font, int32& ScreenTextY);
 };
diff --git a/Engine/Plugins/Animation/GameplayInsights/Source/RewindDebuggerVLog/Private/VisualLoggerProvider.cpp b/Engine/Plugins/Animation/GameplayInsights/Source/RewindDebuggerVLog/Private/VisualLoggerProvider.cpp
--- a/Engine/Plugins/Animation/GameplayInsights/Source/RewindDebuggerVLog/Private/VisualLoggerProvider.cpp
+++ b/Engine/Plugins/Animation/GameplayInsights/Source/RewindDebuggerVLog/Private/VisualLoggerProvider.cpp
@@ -57,6 +57,11 @@ void  FVisualLoggerProvider::AppendVisualLogEntry(uint64 InObjectId, double InTi
 
 	Timeline->AppendEvent(InTime, Entry);
 	
+	for (const FVisualLogLine& Line : Entry.LogLines)
+	{
+		Categories.AddUnique(Line.Category);
+	}
+	
 	for (const FVisualLogShapeElement& Shape : Entry.ElementsToDraw)
 	{
 		Categories.AddUnique(Shape.Category);
