Fix Mac compilation issue

diff --git a/Engine/Plugins/Animation/GameplayInsights/Source/GameplayInsights/Public/SEventTimelineView.h b/Engine/Plugins/Animation/GameplayInsights/Source/GameplayInsights/Public/SEventTimelineView.h
--- a/Engine/Plugins/Animation/GameplayInsights/Source/GameplayInsights/Public/SEventTimelineView.h
+++ b/Engine/Plugins/Animation/GameplayInsights/Source/GameplayInsights/Public/SEventTimelineView.h
@@ -18,6 +18,14 @@ public:
 	public:
 		struct EventPoint
 		{
+			EventPoint(const double InTime, const FText& InType, const FText& InDescription, const FLinearColor& InColor)
+				: Time(InTime)
+				, Type(InType)
+				, Description(InDescription)
+				, Color(InColor)
+			{
+			}
+
 			double Time;
 			FText Type;
 			FText Description;
@@ -26,6 +34,15 @@ public:
 		
 		struct EventWindow
 		{
+			EventWindow(const double InTimeStart, const double InTimeEnd, const FText& InType, const FText& InDescription, const FLinearColor& InColor)
+				: TimeStart(InTimeStart)
+				, TimeEnd(InTimeEnd)
+				, Type(InType)
+				, Description(InDescription)
+				, Color(InColor)
+			{
+			}
+
 			double TimeStart;
 			double TimeEnd;
 			FText Type;
