add world filter to rewind debugger

diff --git a/Engine/Plugins/Animation/GameplayInsights/Source/GameplayInsights/Private/GameplayProvider.cpp b/Engine/Plugins/Animation/GameplayInsights/Source/GameplayInsights/Private/GameplayProvider.cpp
--- a/Engine/Plugins/Animation/GameplayInsights/Source/GameplayInsights/Private/GameplayProvider.cpp
+++ b/Engine/Plugins/Animation/GameplayInsights/Source/GameplayInsights/Private/GameplayProvider.cpp
@@ -227,6 +227,16 @@ void FGameplayProvider::EnumerateObjects(double StartTime, double EndTime, TFunc
 		});
 }
 
+void FGameplayProvider::EnumerateWorlds(TFunctionRef<void(const FWorldInfo&)> Callback) const
+{
+	Session.ReadAccessCheck();
+    
+	for(const FWorldInfo& WorldInfo : WorldInfos)
+	{
+		Callback(WorldInfo);
+	}
+}
+
 const FObjectPropertyValue* FGameplayProvider::FindPropertyValueFromStorageIndex(uint64 InObjectId, int64 InStorageIndex) const
 {
 	Session.ReadAccessCheck();
diff --git a/Engine/Plugins/Animation/GameplayInsights/Source/GameplayInsights/Private/GameplayProvider.h b/Engine/Plugins/Animation/GameplayInsights/Source/GameplayInsights/Private/GameplayProvider.h
--- a/Engine/Plugins/Animation/GameplayInsights/Source/GameplayInsights/Private/GameplayProvider.h
+++ b/Engine/Plugins/Animation/GameplayInsights/Source/GameplayInsights/Private/GameplayProvider.h
@@ -26,6 +26,7 @@ public:
 	virtual void EnumerateObjectPropertyValues(uint64 InObjectId, const FObjectPropertiesMessage& InMessage, TFunctionRef<void(const FObjectPropertyValue&)> Callback) const override;
 	virtual void EnumerateObjects(TFunctionRef<void(const FObjectInfo&)> Callback) const override;
 	virtual void EnumerateObjects(double StartTime, double EndTime, TFunctionRef<void(const FObjectInfo&)> Callback) const override;
+	virtual void EnumerateWorlds(TFunctionRef<void(const FWorldInfo&)> Callback) const override;
 	virtual void EnumerateSubobjects(uint64 ObjectId, TFunctionRef<void(uint64 SubobjectId)> Callback) const override;
 	virtual const FObjectPropertyValue* FindPropertyValueFromStorageIndex(uint64 InObjectId, int64 InStorageIndex) const override;
 	virtual const FClassInfo* FindClassInfo(uint64 InClassId) const override;
diff --git a/Engine/Plugins/Animation/GameplayInsights/Source/GameplayInsights/Public/IGameplayProvider.h b/Engine/Plugins/Animation/GameplayInsights/Source/GameplayInsights/Public/IGameplayProvider.h
--- a/Engine/Plugins/Animation/GameplayInsights/Source/GameplayInsights/Public/IGameplayProvider.h
+++ b/Engine/Plugins/Animation/GameplayInsights/Source/GameplayInsights/Public/IGameplayProvider.h
@@ -118,6 +118,7 @@ public:
 	virtual void EnumerateObjectPropertyValues(uint64 InObjectId, const FObjectPropertiesMessage& InMessage, TFunctionRef<void(const FObjectPropertyValue&)> Callback) const = 0;
 	virtual void EnumerateObjects(TFunctionRef<void(const FObjectInfo&)> Callback) const = 0;
 	virtual void EnumerateObjects(double StartTime, double EndTime, TFunctionRef<void(const FObjectInfo&)> Callback) const = 0;
+	virtual void EnumerateWorlds(TFunctionRef<void(const FWorldInfo&)> Callback) const = 0;
 	virtual void EnumerateSubobjects(uint64 ObjectId, TFunctionRef<void(uint64 SubobjectId)> Callback) const = 0;
 	virtual const FObjectPropertyValue * FindPropertyValueFromStorageIndex(uint64 InObjectId, int64 InStorageIndex) const = 0;
 	virtual const FClassInfo* FindClassInfo(uint64 InClassId) const = 0;
diff --git a/Engine/Plugins/Animation/GameplayInsights/Source/RewindDebugger/Private/RewindDebugger.cpp b/Engine/Plugins/Animation/GameplayInsights/Source/RewindDebugger/Private/RewindDebugger.cpp
--- a/Engine/Plugins/Animation/GameplayInsights/Source/RewindDebugger/Private/RewindDebugger.cpp
+++ b/Engine/Plugins/Animation/GameplayInsights/Source/RewindDebugger/Private/RewindDebugger.cpp
@@ -254,6 +254,7 @@ void FRewindDebugger::OnPIEStopped(bool bSimulating)
 
 	StopRecording();
 	
+	bDisplayWorldIdValid = false;
 }
 
 bool FRewindDebugger::GetTargetActorPosition(FVector& OutPosition) const
@@ -493,6 +494,8 @@ void FRewindDebugger::OpenTrace(const FString& FilePath)
 { 
 	ClearTrace();
 
+	bDisplayWorldIdValid = false;
+	
 	IUnrealInsightsModule& TraceInsightsModule = FModuleManager::LoadModuleChecked<IUnrealInsightsModule>("TraceInsights");
 	TraceInsightsModule.StartAnalysisForTraceFile(*FilePath);
 
@@ -1031,7 +1034,6 @@ void FRewindDebugger::Tick(float DeltaTime)
 {
 	TRACE_CPUPROFILER_EVENT_SCOPE(FRewindDebugger::Tick);
 
-
 	if (bQueueStartRecording)
 	{
 		StartRecording();
@@ -1047,6 +1049,32 @@ void FRewindDebugger::Tick(float DeltaTime)
 		{
 			TraceServices::FAnalysisSessionReadScope SessionReadScope(*Session);
 
+			// set a default display world when loading a trace (first client/standalone world)
+			if (IsTraceFileLoaded() && !bDisplayWorldIdValid)
+			{
+				GameplayProvider->EnumerateWorlds([this](const FWorldInfo& WorldInfo)
+           		{
+					if (WorldInfo.Type == FWorldInfo::EType::PIE)
+					{
+						if(WorldInfo.NetMode == FWorldInfo::ENetMode::Client && WorldInfo.PIEInstanceId == 1)
+						{
+							DisplayWorldId = WorldInfo.Id;
+							bDisplayWorldIdValid = true;
+						}
+						if(WorldInfo.NetMode == FWorldInfo::ENetMode::Standalone && WorldInfo.PIEInstanceId == 0)
+						{
+							DisplayWorldId = WorldInfo.Id;
+							bDisplayWorldIdValid = true;
+						}
+					}
+					else if (WorldInfo.Type == FWorldInfo::EType::Game)
+					{
+						DisplayWorldId = WorldInfo.Id;
+						bDisplayWorldIdValid = true;
+					}
+           		});
+			}
+
 			double RecordingDurationValue = GameplayProvider->GetRecordingDuration();
 			if (IsTraceFileLoaded() && RecordingDurationValue > RecordingDuration.Get())
 			{
@@ -1226,6 +1254,152 @@ void FRewindDebugger::RegisterComponentContextMenu()
 	}));
 }
 
+void FRewindDebugger::MakeOtherWorldsMenu(UToolMenu* Menu)
+{
+	FRewindDebugger* RewindDebugger = FRewindDebugger::Instance();
+	
+	FToolMenuSection& Section = Menu->AddSection("Other Worlds", LOCTEXT("Other Worlds", "Other Worlds"));
+
+	if (const TraceServices::IAnalysisSession* Session = RewindDebugger->GetAnalysisSession())
+	{
+		TraceServices::FAnalysisSessionReadScope SessionReadScope(*Session);
+		const IGameplayProvider* GameplayProvider = Session->ReadProvider<IGameplayProvider>("GameplayProvider");
+
+		GameplayProvider->EnumerateWorlds([GameplayProvider, &Section](const FWorldInfo& WorldInfo)
+		{
+			const FObjectInfo* ObjectInfo = GameplayProvider->FindObjectInfo(WorldInfo.Id);
+			FString Name = ObjectInfo->Name;
+
+			if(WorldInfo.NetMode == FWorldInfo::ENetMode::DedicatedServer)
+			{
+				return;
+			}
+			else if (WorldInfo.Type == FWorldInfo::EType::Game || WorldInfo.Type == FWorldInfo::EType::PIE)
+			{
+				return;
+			}
+			else
+			{
+				if (WorldInfo.Type == FWorldInfo::EType::Editor)
+				{
+					Name = Name + " (Editor)";
+				}
+				else if (WorldInfo.Type == FWorldInfo::EType::Inactive)
+				{
+					Name = Name + " (Editor)";
+				}
+				else if (WorldInfo.Type == FWorldInfo::EType::EditorPreview)
+				{
+					Name = Name + " (Editor Preview)";
+				}
+				else if (WorldInfo.Type == FWorldInfo::EType::GamePreview)
+				{
+					Name = Name + " (Game Preview)";
+				}
+				else if (WorldInfo.Type == FWorldInfo::EType::GameRPC)
+				{
+					Name = Name + " (Game RPC)";
+				}
+			}
+		
+			Section.AddMenuEntry(FName(ObjectInfo->Name,WorldInfo.Id),
+								FText::FromString(Name),
+								FText(),
+								FSlateIcon(),
+								FUIAction( FExecuteAction::CreateLambda([World = WorldInfo.Id]()
+								{
+									FRewindDebugger::Instance()->SetDisplayWorld(World);
+								}),
+								FCanExecuteAction(),
+								FIsActionChecked::CreateLambda([World = WorldInfo.Id]()
+								{
+									return FRewindDebugger::Instance()->DisplayWorldId == World;
+								})),
+								EUserInterfaceActionType::Check
+							);
+		
+		});
+	}
+}
+
+void FRewindDebugger::SetDisplayWorld(uint64 WorldId)
+{
+	DisplayWorldId = WorldId;
+	
+	IterateExtensions([this](IRewindDebuggerExtension* Extension)
+	{
+		Extension->Clear(this);
+		Extension->Update(0.0,this);
+	});
+}
+void FRewindDebugger::MakeWorldsMenu(UToolMenu* Menu)
+{
+	FRewindDebugger* RewindDebugger = FRewindDebugger::Instance();
+	
+	FToolMenuSection& ServerWorldsSection = Menu->AddSection("Server Worlds", LOCTEXT("Server", "Server"));
+	FToolMenuSection& GameWorldsSection = Menu->AddSection("Game Worlds", LOCTEXT("Game Worlds", "Game Worlds"));
+	FToolMenuSection& OtherWorldsSection = Menu->AddSection("Other Worlds", LOCTEXT("Other Worlds", "Other Worlds"));
+
+	OtherWorldsSection.AddSubMenu("Other Worlds",
+		LOCTEXT("Other Worlds", "Other Worlds"),
+		LOCTEXT("Other Worlds Tooltip", "Additional worlds such as  Editor Preview worlds"),
+		FNewToolMenuChoice(
+			FNewToolMenuDelegate::CreateStatic(FRewindDebugger::MakeOtherWorldsMenu)
+			));
+	
+	if (const TraceServices::IAnalysisSession* Session = RewindDebugger->GetAnalysisSession())
+	{
+		TraceServices::FAnalysisSessionReadScope SessionReadScope(*Session);
+		const IGameplayProvider* GameplayProvider = Session->ReadProvider<IGameplayProvider>("GameplayProvider");
+
+		GameplayProvider->EnumerateWorlds([GameplayProvider,&GameWorldsSection, &OtherWorldsSection, &ServerWorldsSection](const FWorldInfo& WorldInfo)
+		{
+			const FObjectInfo* ObjectInfo = GameplayProvider->FindObjectInfo(WorldInfo.Id);
+			FString Name = ObjectInfo->Name;
+
+			FToolMenuSection* Section = &OtherWorldsSection;
+
+			if(WorldInfo.NetMode == FWorldInfo::ENetMode::DedicatedServer)
+			{
+				Section = &ServerWorldsSection;
+				Name = Name + " (Server)";
+			}
+			else if (WorldInfo.Type == FWorldInfo::EType::Game || WorldInfo.Type == FWorldInfo::EType::PIE)
+			{
+				Section = &GameWorldsSection;
+				if(WorldInfo.NetMode == FWorldInfo::ENetMode::Client && WorldInfo.PIEInstanceId >= 0)
+				{
+					Name = Name + " (Client " + FString::FromInt(WorldInfo.PIEInstanceId) + ")";
+				}
+				if(WorldInfo.NetMode == FWorldInfo::ENetMode::Standalone && WorldInfo.PIEInstanceId >= 0)
+				{
+					Name = Name + " (Standalone " + FString::FromInt(WorldInfo.PIEInstanceId) + ")";
+				}
+			}
+			else
+			{
+				return;
+			}
+			
+			Section->AddMenuEntry(FName(ObjectInfo->Name,WorldInfo.Id),
+								FText::FromString(Name),
+								FText(),
+								FSlateIcon(),
+								FUIAction( FExecuteAction::CreateLambda([World = WorldInfo.Id]()
+								{
+									FRewindDebugger::Instance()->SetDisplayWorld(World);
+								}),
+								FCanExecuteAction(),
+								FIsActionChecked::CreateLambda([World = WorldInfo.Id]()
+								{
+									return FRewindDebugger::Instance()->DisplayWorldId == World;
+								})),
+								EUserInterfaceActionType::Check
+							);
+		});
+	}
+}
+
 void FRewindDebugger::RegisterToolBar()
 {
 	UToolMenu* Menu = UToolMenus::Get()->RegisterMenu("RewindDebugger.ToolBar", NAME_None, EMultiBoxType::ToolBar);
@@ -1449,7 +1623,18 @@ void FRewindDebugger::RegisterToolBar()
 				TAttribute<FText>(),
 				FSlateIcon("RewindDebuggerStyle", "RewindDebugger.AutoRecord")));
 	
-
+	Section.AddSeparator("NAME_None");
+
+	Section.AddEntry(FToolMenuEntry::InitComboButton(
+		"Display World",
+		FUIAction(
+			FExecuteAction(),
+			FCanExecuteAction::CreateLambda([](){ return FRewindDebugger::Instance()->IsTraceFileLoaded(); })
+			),
+		FNewToolMenuDelegate::CreateStatic(&FRewindDebugger::MakeWorldsMenu),
+		LOCTEXT("Display World", "Display World"),
+		LOCTEXT("Display World Tooltip", "When loading trace files, only the objects (Such as Skeletal Meshes) from the world selected here will be spawned for preview")
+		));
 	
 	Menu->SetStyleSet(&FAppStyle::Get());
 	Menu->StyleName = "PaletteToolBar";
diff --git a/Engine/Plugins/Animation/GameplayInsights/Source/RewindDebugger/Private/RewindDebugger.h b/Engine/Plugins/Animation/GameplayInsights/Source/RewindDebugger/Private/RewindDebugger.h
--- a/Engine/Plugins/Animation/GameplayInsights/Source/RewindDebugger/Private/RewindDebugger.h
+++ b/Engine/Plugins/Animation/GameplayInsights/Source/RewindDebugger/Private/RewindDebugger.h
@@ -51,6 +51,7 @@ public:
 	virtual TSharedPtr<FDebugObjectInfo> GetSelectedComponent() const override;
 	virtual TSharedPtr<RewindDebugger::FRewindDebuggerTrack> GetSelectedTrack() const override;
 	virtual TArray<TSharedPtr<FDebugObjectInfo>>& GetDebugComponents() override;
+	virtual bool ShouldDisplayWorld(uint64 WorldId) override { return DisplayWorldId == WorldId; } 
 
 	void OnConnection();
 
@@ -135,6 +136,9 @@ public:
 	
 	void UpdateDetailsPanel(TSharedRef<SDockTab> DetailsTab);
 	static void RegisterComponentContextMenu();
+	static void MakeOtherWorldsMenu(class UToolMenu* Menu);
+	void SetDisplayWorld(uint64 WorldId);
+	static void MakeWorldsMenu(class UToolMenu* Menu);
 	static void RegisterToolBar();
 	
 	DECLARE_DELEGATE_OneParam( FOnTrackCursor, bool)
@@ -223,4 +227,7 @@ private:
 	TArray<RewindDebugger::FRewindDebuggerTrackType> TrackTypes;
 
 	bool bIsDetailsPanelOpen = true;
+
+	uint64 DisplayWorldId = 0;
+	bool bDisplayWorldIdValid = false;
 };
diff --git a/Engine/Plugins/Animation/GameplayInsights/Source/RewindDebugger/Private/RewindDebuggerAnimation.cpp b/Engine/Plugins/Animation/GameplayInsights/Source/RewindDebugger/Private/RewindDebuggerAnimation.cpp
--- a/Engine/Plugins/Animation/GameplayInsights/Source/RewindDebugger/Private/RewindDebuggerAnimation.cpp
+++ b/Engine/Plugins/Animation/GameplayInsights/Source/RewindDebugger/Private/RewindDebuggerAnimation.cpp
@@ -329,7 +329,16 @@ void FRewindDebuggerAnimation::Update(float DeltaTime, IRewindDebugger* RewindDe
 									
 								if (MeshComponentInfo == nullptr)
 								{
-									MeshComponentInfo = SpawnMesh(ObjectId, GameplayProvider);
+									if (RewindDebugger->IsTraceFileLoaded())
+									{
+										if (const FWorldInfo* ObjectWorldInfo = GameplayProvider->FindWorldInfoFromObject(ObjectId))
+										{
+											if (RewindDebugger->ShouldDisplayWorld(ObjectWorldInfo->Id))
+											{
+												MeshComponentInfo = SpawnMesh(ObjectId, GameplayProvider);
+											}
+										}
+									}
 									bLoadMesh = true;
 									LastScrubTime = -1;
 								}
diff --git a/Engine/Source/Editor/RewindDebuggerInterface/Public/IRewindDebugger.h b/Engine/Source/Editor/RewindDebuggerInterface/Public/IRewindDebugger.h
--- a/Engine/Source/Editor/RewindDebuggerInterface/Public/IRewindDebugger.h
+++ b/Engine/Source/Editor/RewindDebuggerInterface/Public/IRewindDebugger.h
@@ -104,6 +104,8 @@ public:
 	
 	virtual const struct FObjectInfo* FindOwningActorInfo(const class IGameplayProvider* GameplayProvider, uint64 ObjectId) const = 0;
 
+	virtual bool ShouldDisplayWorld(uint64 WorldId) = 0;
+
 	// get the current IRewindDebugger instance
 	static IRewindDebugger* Instance();
 	
