Fix for follow actor mode in rewind debugger when server components are present

diff --git a/Engine/Plugins/Animation/GameplayInsights/Source/RewindDebugger/Private/RewindDebugger.cpp b/Engine/Plugins/Animation/GameplayInsights/Source/RewindDebugger/Private/RewindDebugger.cpp
--- a/Engine/Plugins/Animation/GameplayInsights/Source/RewindDebugger/Private/RewindDebugger.cpp
+++ b/Engine/Plugins/Animation/GameplayInsights/Source/RewindDebugger/Private/RewindDebugger.cpp
@@ -1133,18 +1133,23 @@ void FRewindDebugger::Tick(float DeltaTime)
 							{
 								AnimationProvider->EnumerateSkeletalMeshPoseTimelines([this, &bNewActor, GameplayProvider](uint64 ObjectId, const IAnimationProvider::SkeletalMeshPoseTimeline& TimelineData)
 								{
-									// until we have actor transforms traced out, the first skeletal mesh component transform on the target actor be used as as the actor position
+									// until we have actor transforms traced out, the first (from a non-server) skeletal mesh component transform on the target actor be used as as the actor position
 
-									if (const FObjectInfo* ActorInfo = FindOwningActorInfo(GameplayProvider, ObjectId))
+									if (const FWorldInfo* WorldInfo = GameplayProvider->FindWorldInfoFromObject(ObjectId))
 									{
-										if (TargetObjectIds.Contains(ActorInfo->Id))
+										if (WorldInfo->NetMode != FWorldInfo::ENetMode::DedicatedServer)
 										{
-											bNewActor = true;
-											TargetActorIdForMesh = ActorInfo->Id;
-											TargetActorMeshId = ObjectId;
+											if (const FObjectInfo* ActorInfo = FindOwningActorInfo(GameplayProvider, ObjectId))
+											{
+												if (TargetObjectIds.Contains(ActorInfo->Id))
+												{
+													bNewActor = true;
+													TargetActorIdForMesh = ActorInfo->Id;
+													TargetActorMeshId = ObjectId;
+												}
+											}
 										}
 									}
-									
 								});
 							}
 						
