Fixed RewindDebugger drawing debug capsules at the wrong location. Previously it passed in the capsule's base coordinate as the center when calling DrawDebugCapsule.

diff --git a/Engine/Plugins/Animation/GameplayInsights/Source/RewindDebuggerVLog/Private/VisualLogEntryRenderer.cpp b/Engine/Plugins/Animation/GameplayInsights/Source/RewindDebuggerVLog/Private/VisualLogEntryRenderer.cpp
--- a/Engine/Plugins/Animation/GameplayInsights/Source/RewindDebuggerVLog/Private/VisualLogEntryRenderer.cpp
+++ b/Engine/Plugins/Animation/GameplayInsights/Source/RewindDebuggerVLog/Private/VisualLogEntryRenderer.cpp
@@ -319,9 +319,10 @@ void FVisualLogEntryRenderer::RenderLogEntry(class UWorld* World, const FVisualL
 				const float HalfHeight = static_cast<float>(FirstData.X);
 				const float Radius = static_cast<float>(FirstData.Y);
 				const FQuat Rotation = FQuat(FirstData.Z, SecondData.X, SecondData.Y, SecondData.Z);
-		
-				DrawDebugCapsule(World, Base, HalfHeight, Radius, Rotation, Color, false, -1.f, DepthPriority, Thickness);
-				RenderDescription(World, Canvas, Font, ElementToDraw, Base, Color, Index/3, ElementToDraw->Points.Num()/3);
+
+				const FVector Center = Base + Rotation.RotateVector(FVector(0, 0, HalfHeight));
+				DrawDebugCapsule(World, Center, HalfHeight, Radius, Rotation, Color, false, -1.f, DepthPriority, Thickness);
+				RenderDescription(World, Canvas, Font, ElementToDraw, Center, Color, Index/3, ElementToDraw->Points.Num()/3);
 			}
 		}
 			break;
