Adding visual logger logline support in rewind debugger.

diff --git a/Engine/Plugins/Animation/GameplayInsights/Source/RewindDebuggerVLog/Private/VisualLogTrack.cpp b/Engine/Plugins/Animation/GameplayInsights/Source/RewindDebuggerVLog/Private/VisualLogTrack.cpp
--- a/Engine/Plugins/Animation/GameplayInsights/Source/RewindDebuggerVLog/Private/VisualLogTrack.cpp
+++ b/Engine/Plugins/Animation/GameplayInsights/Source/RewindDebuggerVLog/Private/VisualLogTrack.cpp
@@ -97,10 +97,15 @@ bool FVisualLogCategoryTrack::UpdateInternal()
 			{
 				InTimeline.EnumerateEvents(StartTime, EndTime, [this, StartTime, EndTime, VisLogProvider, AnalysisSession](double InStartTime, double InEndTime, uint32 InDepth, const FVisualLogEntry& InMessage)
 				{
-					for(const FVisualLogShapeElement& Element : InMessage.ElementsToDraw)
-					{
-						EventData->Points.Add({InMessage.TimeStamp,FText::FromName(Element.Category), FText::FromString(Element.Description),Element.GetFColor()});
-					}
+                	for(const FVisualLogShapeElement& Element : InMessage.ElementsToDraw)
+                	{
+                		EventData->Points.Add({InMessage.TimeStamp,FText::FromName(Element.Category), FText::FromString(Element.Description),Element.GetFColor()});
+                	}
+                	for(const FVisualLogLine& Line : InMessage.LogLines)
+                	{
+                		// No Description from log lines.
+                		EventData->Points.Add({InMessage.TimeStamp, FText::FromName(Line.Category), FText::FromString(""), Line.Color});
+                	}
 					return TraceServices::EEventEnumerate::Continue;
 				});
 			});
@@ -135,7 +140,6 @@ bool FVisualLogCategoryTrack::UpdateInternal()
 						{
 							DetailsObject->VisualLogDetails.Add({Element.Category, Element.Description});
 						}
-						
 						for(const FVisualLogLine& Line : InMessage.LogLines)
 						{
 							DetailsObject->VisualLogDetails.Add({Line.Category, Line.Line});
@@ -219,6 +223,13 @@ bool FVisualLogTrack::UpdateInternal()
 		{
 			InTimeline.EnumerateEvents(StartTime, EndTime, [this, StartTime, EndTime, VisLogProvider, AnalysisSession, &UniqueTrackIds](double InStartTime, double InEndTime, uint32 InDepth, const FVisualLogEntry& InMessage)
 			{
+				if(InMessage.LogLines.Num() > 0)
+				{
+					for(const FVisualLogLine& Line : InMessage.LogLines)
+					{
+						UniqueTrackIds.AddUnique(Line.Category);
+					}
+				}
 				for(const FVisualLogShapeElement& Element : InMessage.ElementsToDraw)
 				{
 					UniqueTrackIds.AddUnique(Element.Category);
