Remove GameplayInsights extender from the Memory Insights tab.

diff --git a/Engine/Plugins/Animation/GameplayInsights/Source/GameplayInsights/Private/GameplayTimingViewExtender.cpp b/Engine/Plugins/Animation/GameplayInsights/Source/GameplayInsights/Private/GameplayTimingViewExtender.cpp
--- a/Engine/Plugins/Animation/GameplayInsights/Source/GameplayInsights/Private/GameplayTimingViewExtender.cpp
+++ b/Engine/Plugins/Animation/GameplayInsights/Source/GameplayInsights/Private/GameplayTimingViewExtender.cpp
@@ -2,6 +2,7 @@
 
 #include "GameplayTimingViewExtender.h"
 #include "Insights/ITimingViewSession.h"
+#include "Insights/IUnrealInsightsModule.h"
 #include "GameplaySharedData.h"
 #include "AnimationSharedData.h"
 #include "UObject/WeakObjectPtr.h"
@@ -21,6 +22,11 @@
 
 void FGameplayTimingViewExtender::OnBeginSession(UE::Insights::Timing::ITimingViewSession& InSession)
 {
+	if (InSession.GetName() != FInsightsManagerTabs::TimingProfilerTabId)
+	{
+		return;
+	}
+
 	FPerSessionData* PerSessionData = PerSessionDataMap.Find(&InSession);
 	if(PerSessionData == nullptr)
 	{
@@ -57,6 +63,11 @@ void FGameplayTimingViewExtender::OnBeginSession(UE::Insights::Timing::ITimingVi
 
 void FGameplayTimingViewExtender::OnEndSession(UE::Insights::Timing::ITimingViewSession& InSession)
 {
+	if (InSession.GetName() != FInsightsManagerTabs::TimingProfilerTabId)
+	{
+		return;
+	}
+
 	FPerSessionData* PerSessionData = PerSessionDataMap.Find(&InSession);
 	if(PerSessionData != nullptr)
 	{
@@ -77,6 +88,11 @@ void FGameplayTimingViewExtender::OnEndSession(UE::Insights::Timing::ITimingView
 
 void FGameplayTimingViewExtender::Tick(UE::Insights::Timing::ITimingViewSession& InSession, const TraceServices::IAnalysisSession& InAnalysisSession)
 {
+	if (InSession.GetName() != FInsightsManagerTabs::TimingProfilerTabId)
+	{
+		return;
+	}
+
 	FPerSessionData* PerSessionData = PerSessionDataMap.Find(&InSession);
 	if(PerSessionData != nullptr)
 	{
@@ -87,6 +103,11 @@ void FGameplayTimingViewExtender::Tick(UE::Insights::Timing::ITimingViewSession&
 
 void FGameplayTimingViewExtender::ExtendFilterMenu(UE::Insights::Timing::ITimingViewSession& InSession, FMenuBuilder& InMenuBuilder)
 {
+	if (InSession.GetName() != FInsightsManagerTabs::TimingProfilerTabId)
+	{
+		return;
+	}
+
 	FPerSessionData* PerSessionData = PerSessionDataMap.Find(&InSession);
 	if(PerSessionData != nullptr)
 	{
