Fix filtering of events in VisualLogTrack

diff --git a/Engine/Plugins/Animation/GameplayInsights/Source/RewindDebuggerVLog/Private/VisualLogTrack.cpp b/Engine/Plugins/Animation/GameplayInsights/Source/RewindDebuggerVLog/Private/VisualLogTrack.cpp
--- a/Engine/Plugins/Animation/GameplayInsights/Source/RewindDebuggerVLog/Private/VisualLogTrack.cpp
+++ b/Engine/Plugins/Animation/GameplayInsights/Source/RewindDebuggerVLog/Private/VisualLogTrack.cpp
@@ -99,12 +99,18 @@ bool FVisualLogCategoryTrack::UpdateInternal()
 				{
                 	for(const FVisualLogShapeElement& Element : InMessage.ElementsToDraw)
                 	{
-                		EventData->Points.Add({InMessage.TimeStamp,FText::FromName(Element.Category), FText::FromString(Element.Description),Element.GetFColor()});
+						if (Element.Category == Category)
+						{
+							EventData->Points.Add({InMessage.TimeStamp,FText::FromName(Element.Category), FText::FromString(Element.Description),Element.GetFColor()});
+						}
                 	}
                 	for(const FVisualLogLine& Line : InMessage.LogLines)
                 	{
-                		// No Description from log lines.
-                		EventData->Points.Add({InMessage.TimeStamp, FText::FromName(Line.Category), FText::FromString(""), Line.Color});
+						if (Line.Category == Category)
+						{
+							// No Description from log lines.
+							EventData->Points.Add({InMessage.TimeStamp, FText::FromName(Line.Category), FText::FromString(""), Line.Color});
+						}
                 	}
 					return TraceServices::EEventEnumerate::Continue;
 				});
@@ -138,11 +144,17 @@ bool FVisualLogCategoryTrack::UpdateInternal()
 					{
 						for(const FVisualLogShapeElement& Element : InMessage.ElementsToDraw)
 						{
-							DetailsObject->VisualLogDetails.Add({Element.Category, Element.Description});
+							if (Element.Category == Category)
+							{
+								DetailsObject->VisualLogDetails.Add({Element.Category, Element.Description});
+							}
 						}
 						for(const FVisualLogLine& Line : InMessage.LogLines)
 						{
-							DetailsObject->VisualLogDetails.Add({Line.Category, Line.Line});
+							if (Line.Category == Category)
+							{
+								DetailsObject->VisualLogDetails.Add({Line.Category, Line.Line});
+							}
 						}
 						return TraceServices::EEventEnumerate::Continue;
 					});
