Fix on screen rendering of recorded vlog text for high DPI displays

diff --git a/Engine/Plugins/Animation/GameplayInsights/Source/RewindDebuggerVLog/Private/VisualLogEntryRenderer.cpp b/Engine/Plugins/Animation/GameplayInsights/Source/RewindDebuggerVLog/Private/VisualLogEntryRenderer.cpp
--- a/Engine/Plugins/Animation/GameplayInsights/Source/RewindDebuggerVLog/Private/VisualLogEntryRenderer.cpp
+++ b/Engine/Plugins/Animation/GameplayInsights/Source/RewindDebuggerVLog/Private/VisualLogEntryRenderer.cpp
@@ -26,19 +26,10 @@ namespace
 				const FSceneView* View = Canvas->SceneView;
 				if (View->GetCullingFrustum().IntersectPoint(Position))
 				{
-					const FPlane Proj = View->Project(Position);
-					if (Proj.W > 0.f)
-					{
-						const int32 HalfX = Canvas->SizeX / 2;
-						const int32 HalfY = Canvas->SizeY / 2;
-
-						const int32 XPos = HalfX + static_cast<int32>(HalfX * Proj.X);
-						const int32 YPos = HalfY + static_cast<int32>(HalfY * (Proj.Y * -1));
-					
-						const FString PrintString = Count == 1 ? ElementToDraw->Description : FString::Printf(TEXT("%s_%d"), *ElementToDraw->Description, Index);
-
-						Canvas->Canvas->DrawShadowedString(XPos, YPos, PrintString, Font, Color);
-					}
+					const FVector ScreenLoc = Canvas->Project(Position);
+					Canvas->SetDrawColor(Color);
+					const FString PrintString = Count == 1 ? ElementToDraw->Description : FString::Printf(TEXT("%s_%d"), *ElementToDraw->Description, Index);
+					Canvas->Canvas->DrawShadowedString(ScreenLoc.X, ScreenLoc.Y, PrintString, Font, Color);
 				}
 			}
 		}
