Fix some issues with rewind debugger text rendering

diff --git a/Engine/Plugins/Animation/GameplayInsights/Source/RewindDebuggerVLog/Private/VisualLogEntryRenderer.cpp b/Engine/Plugins/Animation/GameplayInsights/Source/RewindDebuggerVLog/Private/VisualLogEntryRenderer.cpp
--- a/Engine/Plugins/Animation/GameplayInsights/Source/RewindDebuggerVLog/Private/VisualLogEntryRenderer.cpp
+++ b/Engine/Plugins/Animation/GameplayInsights/Source/RewindDebuggerVLog/Private/VisualLogEntryRenderer.cpp
@@ -108,11 +108,13 @@ namespace
 void FVisualLogEntryRenderer::RenderLogEntry(class UWorld* World, const FVisualLogEntry& Entry, TFunctionRef<bool (const FName&, ELogVerbosity::Type)> MatchCategoryFilters,
 	UCanvas* Canvas, UFont* Font, UFont* MonospaceFont, int32& ScreenTextY)
 {
+	const float InverseDPI = 1.f / Canvas->GetDPIScale();
+	
 	for(const FVisualLogLine& LogLine : Entry.LogLines)
 	{
 		if (MatchCategoryFilters(LogLine.Category, LogLine.Verbosity))
 		{
-			Canvas->Canvas->DrawShadowedString(20, ScreenTextY+=16, LogLine.Line, (LogLine.bMonospace && MonospaceFont) ? MonospaceFont : Font, LogLine.Color);
+			ScreenTextY += Canvas->Canvas->DrawShadowedString(20, ScreenTextY, LogLine.Line, (LogLine.bMonospace && MonospaceFont) ? MonospaceFont : Font, LogLine.Color) * InverseDPI;
 		}
 	}
 	
