Add tracing of SceneComponent transform changes when occur and update Rewind Debugger to use those transforms when scrubbing to set the mesh transforms

diff --git a/Engine/Plugins/Animation/GameplayInsights/Source/GameplayInsights/Private/GameplayAnalyzer.cpp b/Engine/Plugins/Animation/GameplayInsights/Source/GameplayInsights/Private/GameplayAnalyzer.cpp
--- a/Engine/Plugins/Animation/GameplayInsights/Source/GameplayInsights/Private/GameplayAnalyzer.cpp
+++ b/Engine/Plugins/Animation/GameplayInsights/Source/GameplayInsights/Private/GameplayAnalyzer.cpp
@@ -20,6 +20,7 @@ void FGameplayAnalyzer::OnAnalysisBegin(const FOnAnalysisContext& Context)
 	Builder.RouteEvent(RouteId_RecordingInfo, "Object", "RecordingInfo");
 	Builder.RouteEvent(RouteId_Class, "Object", "Class");
 	Builder.RouteEvent(RouteId_Object, "Object", "Object");
+	Builder.RouteEvent(RouteId_ObjectTransform, "Object", "ObjectTransform");
 	Builder.RouteEvent(RouteId_ObjectEvent, "Object", "ObjectEvent");
 	Builder.RouteEvent(RouteId_ObjectLifetimeBegin, "Object", "ObjectLifetimeBegin");
 	Builder.RouteEvent(RouteId_ObjectLifetimeBegin2, "Object", "ObjectLifetimeBegin2");
@@ -105,6 +106,17 @@ bool FGameplayAnalyzer::OnEvent(uint16 RouteId, EStyle Style, const FOnEventCont
 		GameplayProvider.AppendObject(Id, OuterId, ClassId, *ObjectName, *ObjectPathName);
 		break;
 	}
+	case RouteId_ObjectTransform:
+	{
+		uint64 Cycle = EventData.GetValue<uint64>("Cycle");
+		double RecordingTime = EventData.GetValue<double>("RecordingTime");
+		uint64 Id = EventData.GetValue<uint64>("Id");
+		TArrayView<const double> TransformArray = EventData.GetArrayView<double>("Transform");
+		FTransform Transform;
+	 	FMemory::Memcpy(&Transform, &TransformArray[0], sizeof(FTransform));
+		GameplayProvider.AppendObjectTransform(Id, Context.EventTime.AsSeconds(Cycle), RecordingTime, Transform);
+		break;
+	}
 	case RouteId_ObjectEvent:
 	{
 		uint64 Cycle = EventData.GetValue<uint64>("Cycle");
diff --git a/Engine/Plugins/Animation/GameplayInsights/Source/GameplayInsights/Private/GameplayAnalyzer.h b/Engine/Plugins/Animation/GameplayInsights/Source/GameplayInsights/Private/GameplayAnalyzer.h
--- a/Engine/Plugins/Animation/GameplayInsights/Source/GameplayInsights/Private/GameplayAnalyzer.h
+++ b/Engine/Plugins/Animation/GameplayInsights/Source/GameplayInsights/Private/GameplayAnalyzer.h
@@ -22,6 +22,7 @@ private:
 		RouteId_World,
 		RouteId_Class,
 		RouteId_Object,
+		RouteId_ObjectTransform,
 		RouteId_ObjectEvent,
 		RouteId_ObjectLifetimeBegin,
 		RouteId_ObjectLifetimeBegin2,
diff --git a/Engine/Plugins/Animation/GameplayInsights/Source/GameplayInsights/Private/GameplayProvider.cpp b/Engine/Plugins/Animation/GameplayInsights/Source/GameplayInsights/Private/GameplayProvider.cpp
--- a/Engine/Plugins/Animation/GameplayInsights/Source/GameplayInsights/Private/GameplayProvider.cpp
+++ b/Engine/Plugins/Animation/GameplayInsights/Source/GameplayInsights/Private/GameplayProvider.cpp
@@ -109,6 +109,7 @@ static UClass* FindClassByPathName(const FString& RawTargePathName, bool bCheckB
 
 FGameplayProvider::FGameplayProvider(TraceServices::IAnalysisSession& InSession)
 	: Session(InSession)
+	, TransformTimeline(InSession.GetLinearAllocator())
 	, EndPlayEvent(nullptr)
 	, PawnPossession(InSession.GetLinearAllocator())
 	, ObjectLifetimes(InSession.GetLinearAllocator())
@@ -496,6 +497,20 @@ void FGameplayProvider::AppendObject(uint64 InObjectId, uint64 InOuterId, uint64
 	}
 }
 
+void FGameplayProvider::AppendObjectTransform(uint64 InObjectId, double InProfileTime, double InRecordingTime, const FTransform& Transform)
+{
+	Session.WriteAccessCheck();
+
+	bHasAnyData = true;
+
+	FObjectTransformMessage Message;
+	Message.ObjectId = InObjectId;
+	Message.RecordingTime = InRecordingTime;
+	Message.Transform = Transform;
+
+	TransformTimeline.AppendEvent(InProfileTime, Message);
+}
+
 void FGameplayProvider::AppendObjectLifetimeBegin(uint64 InObjectId, double InProfileTime, double InRecordingTime)
 {
 	Session.WriteAccessCheck();
@@ -614,6 +629,22 @@ uint64 FGameplayProvider::FindPossessingController(uint64 PawnId, double Time) c
 	return ControllerId;
 }
 
+bool FGameplayProvider::GetObjectTransform(uint64 InObjectId, double InStartTime, double InEndTime, FTransform& OutTransform) const
+{
+	bool bSuccess = false;
+	TransformTimeline.EnumerateEvents(InStartTime, InEndTime, [&OutTransform, InObjectId, &bSuccess](bool bStart, double Time, const FObjectTransformMessage& TransformMessage)
+	{
+		if (TransformMessage.ObjectId == InObjectId)
+		{
+			bSuccess = true;
+			OutTransform = TransformMessage.Transform;
+		}
+		
+		return TraceServices::EEventEnumerate::Continue;
+	});
+	return bSuccess;
+}
+
 TRange<double> FGameplayProvider::GetObjectTraceLifetime(uint64 ObjectId) const
 {
 	Session.ReadAccessCheck();
diff --git a/Engine/Plugins/Animation/GameplayInsights/Source/GameplayInsights/Private/GameplayProvider.h b/Engine/Plugins/Animation/GameplayInsights/Source/GameplayInsights/Private/GameplayProvider.h
--- a/Engine/Plugins/Animation/GameplayInsights/Source/GameplayInsights/Private/GameplayProvider.h
+++ b/Engine/Plugins/Animation/GameplayInsights/Source/GameplayInsights/Private/GameplayProvider.h
@@ -54,6 +54,10 @@ public:
 	/** Add an object message */
 	void AppendObject(uint64 InObjectId, uint64 InOuterId, uint64 InClassId, const TCHAR* InObjectName, const TCHAR* InObjectPathName);
 
+	
+	/** Add an object transform */
+	void AppendObjectTransform(uint64 InObjectId, double InProfileTime, double InRecordingTime, const FTransform& Transform);
+	
 	/** Add an object create message */
 	void AppendObjectLifetimeBegin(uint64 InObjectId, double InProfileTime, double InRecordingTime);
 	
@@ -95,7 +99,9 @@ public:
 
 	/** Search PawnPossession timeline to find the Controller Object Id for a Pawn (or 0 if none)*/
 	virtual uint64 FindPossessingController(uint64 Pawn, double Time) const override;
-	
+
+	/** find the Transform of an object within a frame that Starts at InStartTime, and ends at InEndTime */
+	virtual bool GetObjectTransform(uint64 InObjectId, double InStartTime, double InEndTime, FTransform& OutTransform) const override;
 
 	/** find the Trace time range for which an object existed */
 	virtual TRange<double> GetObjectTraceLifetime(uint64 ObjectId) const override;
@@ -126,6 +132,7 @@ private:
 
 	/** Map from object Ids to timeline index */
 	TMap<uint64, uint32> ObjectIdToEventTimelines;
+	TMap<uint64, uint32> ObjectIdToTransformTimeline;
 	TMap<uint64, uint32> ObjectIdToPropertiesStorage;
 
 	struct FPawnPossessMessage
@@ -134,6 +141,13 @@ private:
 		uint64 PawnId = 0;
 	};
 	
+	struct FObjectTransformMessage
+   	{
+   		uint64 ObjectId = 0;
+		double RecordingTime;
+   		FTransform Transform;
+   	};
+	
 	struct FObjectExistsMessage
 	{
 		uint64 ObjectId = 0;
@@ -150,6 +164,7 @@ private:
 
 	/** Message storage */
 	TArray<TSharedRef<TraceServices::TPointTimeline<FObjectEventMessage>>> EventTimelines;
+	TraceServices::TPointTimeline<FObjectTransformMessage> TransformTimeline;
 	TArray<TSharedRef<FObjectPropertiesStorage>> PropertiesStorage;
 	TSharedPtr<TraceServices::TPointTimeline<FViewMessage>> ViewTimeline;
 
@@ -177,7 +192,6 @@ private:
 	// Timeline containing intervals where an object exists - In Rewind Debugger recording time units, to avoid excessive conversions
 	TraceServices::TIntervalTimeline<FObjectExistsMessage> ObjectRecordingLifetimes;
 
-
 	// map from controller id to index in the PawnPossession timeline, for lookup when ending events
 	TMap<uint64, uint64> ActivePawnPossession;
 
diff --git a/Engine/Plugins/Animation/GameplayInsights/Source/GameplayInsights/Public/IGameplayProvider.h b/Engine/Plugins/Animation/GameplayInsights/Source/GameplayInsights/Public/IGameplayProvider.h
--- a/Engine/Plugins/Animation/GameplayInsights/Source/GameplayInsights/Public/IGameplayProvider.h
+++ b/Engine/Plugins/Animation/GameplayInsights/Source/GameplayInsights/Public/IGameplayProvider.h
@@ -141,4 +141,5 @@ public:
 	virtual TRange<double> GetObjectRecordingLifetime(uint64 ObjectId) const = 0;
 	virtual double GetRecordingDuration() const = 0;
 	virtual FSlateIcon FindIconForClass(uint64 ClassId) const = 0;
+	virtual bool GetObjectTransform(uint64 InObjectId, double InStartTime, double InEndTime, FTransform& OutTransform) const = 0;
 };
\ No newline at end of file
diff --git a/Engine/Plugins/Animation/GameplayInsights/Source/RewindDebugger/Private/RewindDebuggerAnimation.cpp b/Engine/Plugins/Animation/GameplayInsights/Source/RewindDebugger/Private/RewindDebuggerAnimation.cpp
--- a/Engine/Plugins/Animation/GameplayInsights/Source/RewindDebugger/Private/RewindDebuggerAnimation.cpp
+++ b/Engine/Plugins/Animation/GameplayInsights/Source/RewindDebugger/Private/RewindDebuggerAnimation.cpp
@@ -128,7 +128,6 @@ void FRewindDebuggerAnimation::ApplyPoseToMesh(const IAnimationProvider* Animati
 	// Update mesh based on pose
 	if (PoseMessage)
 	{
-		FTransform ComponentWorldTransform;
 		if (const FSkeletalMeshInfo* SkeletalMeshInfo = AnimationProvider->FindSkeletalMeshInfo(PoseMessage->MeshId))
 		{
 			if (bApplyMesh)
@@ -157,8 +156,11 @@ void FRewindDebuggerAnimation::ApplyPoseToMesh(const IAnimationProvider* Animati
 				}
 			}
 			
+			FTransform ComponentWorldTransform;
 			AnimationProvider->GetSkeletalMeshComponentSpacePose(*PoseMessage, *SkeletalMeshInfo, ComponentWorldTransform, MeshComponent->GetEditableComponentSpaceTransforms());
 			MeshComponent->ApplyEditedComponentSpaceTransforms();
+
+			GameplayProvider->GetObjectTransform(FObjectTrace::GetObjectId(MeshComponent)/*ObjectId*/, Frame.StartTime, Frame.EndTime, ComponentWorldTransform);
 			
 			if (bQueueForReset)
 			{
@@ -174,6 +176,7 @@ void FRewindDebuggerAnimation::ApplyPoseToMesh(const IAnimationProvider* Animati
 			}
 
 			MeshComponent->SetWorldTransform(ComponentWorldTransform, false, nullptr, ETeleportType::TeleportPhysics);
+			
 			MeshComponent->SetForcedLOD(PoseMessage->LodIndex + 1);
 			MeshComponent->UpdateLODStatus();
 			MeshComponent->UpdateChildTransforms(EUpdateTransformFlags::None, ETeleportType::TeleportPhysics);
diff --git a/Engine/Source/Runtime/Engine/Private/Components/SceneComponent.cpp b/Engine/Source/Runtime/Engine/Private/Components/SceneComponent.cpp
--- a/Engine/Source/Runtime/Engine/Private/Components/SceneComponent.cpp
+++ b/Engine/Source/Runtime/Engine/Private/Components/SceneComponent.cpp
@@ -23,6 +23,7 @@
 #include "Logging/MessageLog.h"
 #include "Net/UnrealNetwork.h"
 #include "ComponentUtils.h"
+#include "ObjectTrace.h"
 #include "Framework/Notifications/NotificationManager.h"
 #include "UObject/UObjectAnnotation.h"
 #include "Widgets/Notifications/SNotificationList.h"
@@ -921,6 +922,8 @@ void USceneComponent::PropagateTransformUpdate(bool bTransformChanged, EUpdateTr
 	FPlatformMisc::Prefetch(AttachedChildren.GetData());
 	if (bTransformChanged)
 	{
+		TRACE_OBJECT_TRANSFORM(this, GetComponentToWorld(), Teleport);
+		
 		//QUICK_SCOPE_CYCLE_COUNTER(STAT_USceneComponent_PropagateTransformUpdate_TransformChanged);
 		{
 			// Then update bounds
diff --git a/Engine/Source/Runtime/Engine/Private/Level.cpp b/Engine/Source/Runtime/Engine/Private/Level.cpp
--- a/Engine/Source/Runtime/Engine/Private/Level.cpp
+++ b/Engine/Source/Runtime/Engine/Private/Level.cpp
@@ -3789,8 +3789,6 @@ int32 ULevel::GetEstimatedAddToWorldWorkUnitsRemaining() const
 
 void ULevel::RouteActorInitialize(int32 NumActorsToProcess)
 {
-	TRACE_OBJECT_EVENT(this, RouteActorInitialize);
-
 	const bool bFullProcessing = (NumActorsToProcess <= 0);
 	switch (RouteActorInitializationState)
 	{
diff --git a/Engine/Source/Runtime/Engine/Private/ObjectTrace.cpp b/Engine/Source/Runtime/Engine/Private/ObjectTrace.cpp
--- a/Engine/Source/Runtime/Engine/Private/ObjectTrace.cpp
+++ b/Engine/Source/Runtime/Engine/Private/ObjectTrace.cpp
@@ -38,6 +38,15 @@ UE_TRACE_EVENT_BEGIN(Object, ObjectLifetimeBegin2)
 	UE_TRACE_EVENT_FIELD(uint64, Id)
 UE_TRACE_EVENT_END()
 
+
+UE_TRACE_EVENT_BEGIN(Object, ObjectTransform)
+	UE_TRACE_EVENT_FIELD(uint64, Cycle)
+	UE_TRACE_EVENT_FIELD(double, RecordingTime)
+	UE_TRACE_EVENT_FIELD(uint64, Id)
+	UE_TRACE_EVENT_FIELD(uint8, TeleportType)
+	UE_TRACE_EVENT_FIELD(double[], Transform)
+UE_TRACE_EVENT_END()
+
 UE_TRACE_EVENT_BEGIN(Object, ObjectLifetimeEnd2)
 	UE_TRACE_EVENT_FIELD(uint64, Cycle)
 	UE_TRACE_EVENT_FIELD(double, RecordingTime)
@@ -424,6 +433,34 @@ void FObjectTrace::OutputObject(const UObject* InObject)
 		<< ObjectLifetimeBegin2.Id(GetObjectId(InObject));
 }
 
+void FObjectTrace::OutputObjectTransform(const UObject* InObject, const FTransform& InTransform, ETeleportType InTeleportType)
+{
+	bool bChannelEnabled = UE_TRACE_CHANNELEXPR_IS_ENABLED(ObjectChannel);
+	if (!bChannelEnabled || InObject == nullptr)
+	{
+		return;
+	}
+
+	if(InObject->HasAnyFlags(RF_ClassDefaultObject))
+	{
+		return;
+	}
+
+	if (CANNOT_TRACE_OBJECT(InObject->GetWorld()))
+	{
+		return;
+	}
+
+	TRACE_OBJECT(InObject);
+
+	UE_TRACE_LOG(Object, ObjectTransform, ObjectChannel)
+		<< ObjectTransform.Cycle(FPlatformTime::Cycles64())
+		<< ObjectTransform.RecordingTime(FObjectTrace::GetWorldElapsedTime(InObject->GetWorld()))
+		<< ObjectTransform.Id(GetObjectId(InObject))
+		<< ObjectTransform.TeleportType(static_cast<uint8>(InTeleportType))
+		<< ObjectTransform.Transform(reinterpret_cast<const double*>(&InTransform), sizeof(FTransform) / sizeof(double));
+}
+
 void FObjectTrace::OutputObjectEvent(const UObject* InObject, const TCHAR* InEvent)
 {
 	bool bChannelEnabled = UE_TRACE_CHANNELEXPR_IS_ENABLED(ObjectChannel);
diff --git a/Engine/Source/Runtime/Engine/Public/ObjectTrace.h b/Engine/Source/Runtime/Engine/Public/ObjectTrace.h
--- a/Engine/Source/Runtime/Engine/Public/ObjectTrace.h
+++ b/Engine/Source/Runtime/Engine/Public/ObjectTrace.h
@@ -64,6 +64,9 @@ struct FObjectTrace
 
 	/** Helper function to output an object */
 	ENGINE_API static void OutputObject(const UObject* InObject);
+
+	/** Helper function to output an object transform */
+	ENGINE_API static void OutputObjectTransform(const UObject* InObject, const FTransform& InTransform, ETeleportType InTeleportType);
 	
 	/** Helper function to output object create event */
 	ENGINE_API static void OutputObjectLifetimeBegin(const UObject* InObject);
@@ -125,6 +128,9 @@ struct FObjectTrace
 
 #define TRACE_OBJECT_EVENT(Object, Event) \
 	if(CAN_TRACE_OBJECT(Object)) { UNCONDITIONAL_TRACE_OBJECT_EVENT(Object, Event); }
+
+#define TRACE_OBJECT_TRANSFORM(Object, Transform, Teleport) \
+	if(CAN_TRACE_OBJECT(Object)) { FObjectTrace::OutputObjectTransform(Object, Transform, Teleport); }
 	
 #define TRACE_OBJECT_LIFETIME_BEGIN(Object) \
 	if(CAN_TRACE_OBJECT(Object)) { FObjectTrace::OutputObjectLifetimeBegin(Object); }
@@ -168,6 +174,7 @@ struct FObjectTrace
 
 #define TRACE_CLASS(Class)
 #define TRACE_OBJECT(Object)
+#define TRACE_OBJECT_TRANSFORM(Object, Transform, Teleport)
 #define TRACE_OBJECT_EVENT(Object, Event)
 #define TRACE_WORLD(World)
 #define TRACE_PAWN_POSSESS(Controller, Pawn)
