Fix for editor final builds

diff --git a/Engine/Plugins/Animation/GameplayInsights/Source/RewindDebugger/Private/RewindDebuggerAnimation.cpp b/Engine/Plugins/Animation/GameplayInsights/Source/RewindDebugger/Private/RewindDebuggerAnimation.cpp
--- a/Engine/Plugins/Animation/GameplayInsights/Source/RewindDebugger/Private/RewindDebuggerAnimation.cpp
+++ b/Engine/Plugins/Animation/GameplayInsights/Source/RewindDebugger/Private/RewindDebuggerAnimation.cpp
@@ -160,7 +160,7 @@ void FRewindDebuggerAnimation::ApplyPoseToMesh(const IAnimationProvider* Animati
 			AnimationProvider->GetSkeletalMeshComponentSpacePose(*PoseMessage, *SkeletalMeshInfo, ComponentWorldTransform, MeshComponent->GetEditableComponentSpaceTransforms());
 			MeshComponent->ApplyEditedComponentSpaceTransforms();
 
-			GameplayProvider->GetObjectTransform(FObjectTrace::GetObjectId(MeshComponent)/*ObjectId*/, Frame.StartTime, Frame.EndTime, ComponentWorldTransform);
+			GameplayProvider->GetObjectTransform(ObjectId, Frame.StartTime, Frame.EndTime, ComponentWorldTransform);
 			
 			if (bQueueForReset)
 			{
