[Optim] Vlog rewind debugger. Use bespoke method to find if object should be debugged

FRewindDebuggerVLog::ImmediateRender Avg 729us ->  121us
FVisualLogger::Flush 39ms -> 7ms (called every frame)

diff --git a/Engine/Plugins/Animation/GameplayInsights/Source/RewindDebugger/Private/RewindDebugger.cpp b/Engine/Plugins/Animation/GameplayInsights/Source/RewindDebugger/Private/RewindDebugger.cpp
--- a/Engine/Plugins/Animation/GameplayInsights/Source/RewindDebugger/Private/RewindDebugger.cpp
+++ b/Engine/Plugins/Animation/GameplayInsights/Source/RewindDebugger/Private/RewindDebugger.cpp
@@ -1714,7 +1714,7 @@ TSharedPtr<RewindDebugger::FRewindDebuggerTrack> FRewindDebugger::GetSelectedTra
 // build a component tree that's compatible with the public api from 5.0 for GetDebugComponents.
 void FRewindDebugger::RefreshDebugComponents(TArray<TSharedPtr<RewindDebugger::FRewindDebuggerTrack>>& InTracks, TArray<TSharedPtr<FDebugObjectInfo>>& OutComponents)
 {
-	OutComponents.SetNum(0);
+	OutComponents.SetNum(0, EAllowShrinking::No);
 	for(TSharedPtr<RewindDebugger::FRewindDebuggerTrack>& Track : InTracks)
 	{
 		int Index = OutComponents.Num();
@@ -1731,4 +1731,31 @@ TArray<TSharedPtr<FDebugObjectInfo>>& FRewindDebugger::GetDebugComponents()
 	return DebugComponents;
 }
 
+bool FRewindDebugger::IsContainedByDebugComponent(uint64 ObjectId) const
+{
+	for(auto Track : DebugTracks)
+	{
+		if (Track->GetObjectId() == ObjectId)
+		{
+			return true;
+		}
+
+		bool Found = false;
+		Track->IterateSubTracks( [ObjectId, &Found] (TSharedPtr<RewindDebugger::FRewindDebuggerTrack> Child)
+		{
+			if (Child->GetObjectId() == ObjectId)
+			{
+				Found = true;
+				// Todo: want to stop iteration here
+			}
+		});
+
+		if (Found)
+		{
+			return true;
+		}
+	}
+	return false;
+}
+
 #undef LOCTEXT_NAMESPACE
diff --git a/Engine/Plugins/Animation/GameplayInsights/Source/RewindDebugger/Private/RewindDebugger.h b/Engine/Plugins/Animation/GameplayInsights/Source/RewindDebugger/Private/RewindDebugger.h
--- a/Engine/Plugins/Animation/GameplayInsights/Source/RewindDebugger/Private/RewindDebugger.h
+++ b/Engine/Plugins/Animation/GameplayInsights/Source/RewindDebugger/Private/RewindDebugger.h
@@ -51,6 +51,7 @@ public:
 	virtual TSharedPtr<FDebugObjectInfo> GetSelectedComponent() const override;
 	virtual TSharedPtr<RewindDebugger::FRewindDebuggerTrack> GetSelectedTrack() const override;
 	virtual TArray<TSharedPtr<FDebugObjectInfo>>& GetDebugComponents() override;
+	virtual bool IsContainedByDebugComponent(uint64 ObjectId) const override;
 	virtual bool ShouldDisplayWorld(uint64 WorldId) override { return DisplayWorldId == WorldId; } 
 
 	void OnConnection();
diff --git a/Engine/Plugins/Animation/GameplayInsights/Source/RewindDebuggerVLog/Private/RewindDebuggerVLog.cpp b/Engine/Plugins/Animation/GameplayInsights/Source/RewindDebuggerVLog/Private/RewindDebuggerVLog.cpp
--- a/Engine/Plugins/Animation/GameplayInsights/Source/RewindDebuggerVLog/Private/RewindDebuggerVLog.cpp
+++ b/Engine/Plugins/Animation/GameplayInsights/Source/RewindDebuggerVLog/Private/RewindDebuggerVLog.cpp
@@ -107,25 +107,6 @@ void FRewindDebuggerVLog::Initialize()
 	MonospaceFont = LoadObject<UFont>(nullptr, *MonospaceFontPath.ToString(), nullptr, LOAD_None, nullptr);
 }
 
-bool ContainsObject(TArray<TSharedPtr<FDebugObjectInfo>>& Components, uint64 ObjectId)
-{
-	for(TSharedPtr<FDebugObjectInfo>& Component : Components)
-	{
-		if (Component->ObjectId == ObjectId)
-		{
-			return true;
-		}
-		if (!Component->Children.IsEmpty())
-		{
-			if (ContainsObject(Component->Children, ObjectId))
-			{
-				return true;
-			}
-		}
-	}
-	return false;
-}
-
 bool MatchCategoryFilters(const FName& CategoryName, ELogVerbosity::Type Verbosity)
 {
 	URewindDebuggerVLogSettings& Settings = URewindDebuggerVLogSettings::Get();
@@ -156,7 +137,7 @@ void FRewindDebuggerVLog::ImmediateRender(const UObject* Object, const FVisualLo
 	if (IRewindDebugger* RewindDebugger = IRewindDebugger::Instance())
 	{
 		uint64 ObjectId = FObjectTrace::GetObjectId(Object);
-		if (ContainsObject(RewindDebugger->GetDebugComponents(), ObjectId))
+		if (RewindDebugger->IsContainedByDebugComponent(ObjectId))
 		{
 			ImmediateRenderQueue.Add(Entry);
 		}
diff --git a/Engine/Source/Editor/RewindDebuggerInterface/Public/IRewindDebugger.h b/Engine/Source/Editor/RewindDebuggerInterface/Public/IRewindDebugger.h
--- a/Engine/Source/Editor/RewindDebuggerInterface/Public/IRewindDebugger.h
+++ b/Engine/Source/Editor/RewindDebuggerInterface/Public/IRewindDebugger.h
@@ -74,6 +74,9 @@ public:
 
 	// get a list of all components of the selected target actor (with the actor as the first element in the list)
 	virtual TArray<TSharedPtr<FDebugObjectInfo>>& GetDebugComponents() = 0;
+
+	// check if the given object id is a debug component, or a child or a debug component
+	virtual bool IsContainedByDebugComponent(uint64 ObjectId) const = 0;
 	
 	// returns the currently selected debug component
 	virtual TSharedPtr<FDebugObjectInfo> GetSelectedComponent() const = 0;
