Fix for anim next variable details not showing up

diff --git a/Engine/Plugins/Animation/GameplayInsights/Source/RewindDebuggerRuntime/Private/RewindDebuggerAnimationRuntime.h b/Engine/Plugins/Animation/GameplayInsights/Source/RewindDebuggerRuntime/Private/RewindDebuggerAnimationRuntime.h
--- a/Engine/Plugins/Animation/GameplayInsights/Source/RewindDebuggerRuntime/Private/RewindDebuggerAnimationRuntime.h
+++ b/Engine/Plugins/Animation/GameplayInsights/Source/RewindDebuggerRuntime/Private/RewindDebuggerAnimationRuntime.h
@@ -3,6 +3,7 @@
 #pragma once
 
 #include "CoreMinimal.h"
+#include "Animation/AnimTrace.h"
 #include "RewindDebuggerRuntimeInterface/IRewindDebuggerRuntimeExtension.h"
 
 
@@ -21,5 +22,12 @@ namespace RewindDebugger
 		{
 			UE::Trace::ToggleChannel(TEXT("Animation"), false);
 		}
+
+		virtual void Clear() override
+		{
+#if ANIM_TRACE_ENABLED
+			FAnimTrace::Reset();
+#endif
+		}
 	};
 }
diff --git a/Engine/Plugins/Animation/GameplayInsights/Source/RewindDebuggerRuntime/Private/RewindDebuggerRuntime.cpp b/Engine/Plugins/Animation/GameplayInsights/Source/RewindDebuggerRuntime/Private/RewindDebuggerRuntime.cpp
--- a/Engine/Plugins/Animation/GameplayInsights/Source/RewindDebuggerRuntime/Private/RewindDebuggerRuntime.cpp
+++ b/Engine/Plugins/Animation/GameplayInsights/Source/RewindDebuggerRuntime/Private/RewindDebuggerRuntime.cpp
@@ -2,7 +2,7 @@
 
 #include "RewindDebuggerRuntime/RewindDebuggerRuntime.h"
 
-#include "Animation/AnimTrace.h"
+#include "ObjectTrace.h"
 #include "EngineUtils.h"
 #include "Engine/World.h"
 #include "GameFramework/Pawn.h"
@@ -119,13 +119,16 @@ namespace RewindDebugger
 	void FRewindDebuggerRuntime::StartRecording(FTraceAuxiliary::EConnectionType TraceType, const TCHAR* TraceDestination)
 	{
 #if OBJECT_TRACE_ENABLED
-		// Clear caches  (maybe move these to plugins?)
 		FObjectTrace::Reset();
-#if ANIM_TRACE_ENABLED
-		FAnimTrace::Reset();
-#endif
 
 		ClearRecording.Broadcast();
+
+		// update extensions
+		IterateExtensions([](IRewindDebuggerRuntimeExtension* Extension)
+			{
+				Extension->Clear();
+			}
+		);
         	
         // Disable all trace channels, and then enable only the ones needed by RewindDebugger
         // for systems with RewindDebugger integration, they should enable their channel(s) in an Extension in "RecordingStarted"
