[RewindDebugger] renamed actor/component related code to more generic names in preparation of changes to support more types
- Header cleanup
- applied some coding standard
- fixed formatting

diff --git a/Engine/Plugins/Animation/GameplayInsights/Source/GameplayInsights/Private/GameplayProvider.cpp b/Engine/Plugins/Animation/GameplayInsights/Source/GameplayInsights/Private/GameplayProvider.cpp
--- a/Engine/Plugins/Animation/GameplayInsights/Source/GameplayInsights/Private/GameplayProvider.cpp
+++ b/Engine/Plugins/Animation/GameplayInsights/Source/GameplayInsights/Private/GameplayProvider.cpp
@@ -38,27 +38,27 @@ static UClass* FindBlueprintClass(const FString& TargetNameRaw)
 	// from other blueprints.
 	UClass* FoundClass = nullptr;
 	AssetRegistry.EnumerateAssets(Filter, [&FoundClass, TargetName](const FAssetData& AssetData)
-	{
-		if ((AssetData.AssetName.ToString() == TargetName) || (AssetData.GetObjectPathString() == TargetName))
 		{
-			if (UBlueprint* BP = Cast<UBlueprint>(AssetData.GetAsset()))
+			if ((AssetData.AssetName.ToString() == TargetName) || (AssetData.GetObjectPathString() == TargetName))
 			{
-				FoundClass = BP->GeneratedClass;
-				return false;
+				if (UBlueprint* BP = Cast<UBlueprint>(AssetData.GetAsset()))
+				{
+					FoundClass = BP->GeneratedClass;
+					return false;
+				}
 			}
-		}
 
-		return true;
-	});
+			return true;
+		});
 
 	return FoundClass;
 }
 
 // This is copied from EditorUtilitySubsystem (where it is not public)
 // Should probably be somewhere shared
-static UClass* FindClassByPathName(const FString& RawTargePathName, bool bCheckBlueprintClasses)
+static UClass* FindClassByPathName(const FString& RawTargetPathName, const bool bCheckBlueprintClasses)
 {
-	FString TargetName = RawTargePathName;
+	FString TargetName = RawTargetPathName;
 
 	// Check native classes and loaded assets first before resorting to the asset registry
 	bool bIsValidClassName = true;
@@ -77,7 +77,7 @@ static UClass* FindClassByPathName(const FString& RawTargePathName, bool bCheckB
 			FString ObjectName;
 			TargetName.Split(TEXT("."), &PackageName, &ObjectName);
 
-			const bool bIncludeReadOnlyRoots = true;
+			constexpr bool bIncludeReadOnlyRoots = true;
 			FText Reason;
 			if (!FPackageName::IsValidLongPackageName(PackageName, bIncludeReadOnlyRoots, &Reason))
 			{
@@ -114,8 +114,6 @@ FGameplayProvider::FGameplayProvider(TraceServices::IAnalysisSession& InSession)
 	, PawnPossession(InSession.GetLinearAllocator())
 	, ObjectLifetimes(InSession.GetLinearAllocator())
 	, ObjectRecordingLifetimes(InSession.GetLinearAllocator())
-	, bHasAnyData(false)
-	, bHasObjectProperties(false)
 {
 }
 
@@ -124,9 +122,9 @@ bool FGameplayProvider::ReadObjectEventsTimeline(uint64 InObjectId, TFunctionRef
 	Session.ReadAccessCheck();
 
 	const uint32* IndexPtr = ObjectIdToEventTimelines.Find(InObjectId);
-	if(IndexPtr != nullptr)
+	if (IndexPtr != nullptr)
 	{
-		if (*IndexPtr < uint32(EventTimelines.Num()))
+		if (*IndexPtr < static_cast<uint32>(EventTimelines.Num()))
 		{
 			Callback(*EventTimelines[*IndexPtr]);
 			return true;
@@ -141,22 +139,22 @@ bool FGameplayProvider::ReadObjectEvent(uint64 InObjectId, uint64 InMessageId, T
 	Session.ReadAccessCheck();
 
 	return ReadObjectEventsTimeline(InObjectId, [&Callback, &InMessageId](const ObjectEventsTimeline& InTimeline)
-	{
-		if(InMessageId < InTimeline.GetEventCount())
 		{
-			Callback(InTimeline.GetEvent(InMessageId));
-		}
-	});
+			if (InMessageId < InTimeline.GetEventCount())
+			{
+				Callback(InTimeline.GetEvent(InMessageId));
+			}
+		});
 }
 
-bool FGameplayProvider::ReadObjectPropertiesTimeline(uint64 InObjectId, TFunctionRef<void(const ObjectPropertiesTimeline&)> Callback) const
+bool FGameplayProvider::ReadObjectPropertiesTimeline(const uint64 InObjectId, TFunctionRef<void(const ObjectPropertiesTimeline&)> Callback) const
 {
 	Session.ReadAccessCheck();
 
 	const uint32* IndexPtr = ObjectIdToPropertiesStorage.Find(InObjectId);
-	if(IndexPtr != nullptr)
+	if (IndexPtr != nullptr)
 	{
-		if (*IndexPtr < uint32(ObjectIdToPropertiesStorage.Num()))
+		if (*IndexPtr < static_cast<uint32>(ObjectIdToPropertiesStorage.Num()))
 		{
 			Callback(*PropertiesStorage[*IndexPtr]->Timeline);
 			return true;
@@ -166,14 +164,14 @@ bool FGameplayProvider::ReadObjectPropertiesTimeline(uint64 InObjectId, TFunctio
 	return false;
 }
 
-bool FGameplayProvider::ReadObjectPropertiesStorage(uint64 InObjectId, const FObjectPropertiesMessage& InMessage, TFunctionRef<void(const TConstArrayView<FObjectPropertyValue>&)> Callback) const
+bool FGameplayProvider::ReadObjectPropertiesStorage(const uint64 InObjectId, const FObjectPropertiesMessage& InMessage, TFunctionRef<void(const TConstArrayView<FObjectPropertyValue>&)> Callback) const
 {
 	Session.ReadAccessCheck();
 
 	const uint32* IndexPtr = ObjectIdToPropertiesStorage.Find(InObjectId);
-	if(IndexPtr != nullptr)
+	if (IndexPtr != nullptr)
 	{
-		if (*IndexPtr < uint32(ObjectIdToPropertiesStorage.Num()))
+		if (*IndexPtr < static_cast<uint32>(ObjectIdToPropertiesStorage.Num()))
 		{
 			const TSharedRef<FObjectPropertiesStorage> Storage = PropertiesStorage[*IndexPtr];
 			Callback(MakeArrayView(Storage->Values.GetData() + InMessage.PropertyValueStartIndex, InMessage.PropertyValueEndIndex - InMessage.PropertyValueStartIndex));
@@ -184,17 +182,17 @@ bool FGameplayProvider::ReadObjectPropertiesStorage(uint64 InObjectId, const FOb
 	return false;
 }
 
-void FGameplayProvider::EnumerateObjectPropertyValues(uint64 InObjectId, const FObjectPropertiesMessage& InMessage, TFunctionRef<void(const FObjectPropertyValue&)> Callback) const
+void FGameplayProvider::EnumerateObjectPropertyValues(const uint64 InObjectId, const FObjectPropertiesMessage& InMessage, TFunctionRef<void(const FObjectPropertyValue&)> Callback) const
 {
 	Session.ReadAccessCheck();
 
 	const uint32* IndexPtr = ObjectIdToPropertiesStorage.Find(InObjectId);
-	if(IndexPtr != nullptr)
+	if (IndexPtr != nullptr)
 	{
-		if (*IndexPtr < uint32(ObjectIdToPropertiesStorage.Num()))
+		if (*IndexPtr < static_cast<uint32>(ObjectIdToPropertiesStorage.Num()))
 		{
 			TSharedRef<FObjectPropertiesStorage> Storage = PropertiesStorage[*IndexPtr];
-			for(int64 ValueIndex = InMessage.PropertyValueStartIndex; ValueIndex < InMessage.PropertyValueEndIndex; ++ValueIndex)
+			for (int64 ValueIndex = InMessage.PropertyValueStartIndex; ValueIndex < InMessage.PropertyValueEndIndex; ++ValueIndex)
 			{
 				Callback(Storage->Values[ValueIndex]);
 			}
@@ -206,7 +204,7 @@ void FGameplayProvider::EnumerateObjects(TFunctionRef<void(const FObjectInfo&)>
 {
 	Session.ReadAccessCheck();
 
-	for(const FObjectInfo& ObjectInfo : ObjectInfos)
+	for (const FObjectInfo& ObjectInfo : ObjectInfos)
 	{
 		Callback(ObjectInfo);
 	}
@@ -231,19 +229,19 @@ void FGameplayProvider::EnumerateObjects(double StartTime, double EndTime, TFunc
 void FGameplayProvider::EnumerateWorlds(TFunctionRef<void(const FWorldInfo&)> Callback) const
 {
 	Session.ReadAccessCheck();
-    
-	for(const FWorldInfo& WorldInfo : WorldInfos)
+
+	for (const FWorldInfo& WorldInfo : WorldInfos)
 	{
 		Callback(WorldInfo);
 	}
 }
 
-const FObjectPropertyValue* FGameplayProvider::FindPropertyValueFromStorageIndex(uint64 InObjectId, int64 InStorageIndex) const
+const FObjectPropertyValue* FGameplayProvider::FindPropertyValueFromStorageIndex(const uint64 InObjectId, int64 InStorageIndex) const
 {
 	Session.ReadAccessCheck();
 
 	const uint32* IndexPtr = ObjectIdToPropertiesStorage.Find(InObjectId);
-	if(IndexPtr != nullptr)
+	if (IndexPtr != nullptr)
 	{
 		if (*IndexPtr < static_cast<uint32>(ObjectIdToPropertiesStorage.Num()))
 		{
@@ -276,7 +274,7 @@ const FClassInfo* FGameplayProvider::FindClassInfo(uint64 InClassId) const
 	Session.ReadAccessCheck();
 
 	const int32* ClassIndex = ClassIdToIndexMap.Find(InClassId);
-	if(ClassIndex != nullptr)
+	if (ClassIndex != nullptr)
 	{
 		return &ClassInfos[*ClassIndex];
 	}
@@ -315,7 +313,7 @@ bool FGameplayProvider::IsSubClassOf(uint64 InSubClassId, uint64 InParentClassId
 	Session.ReadAccessCheck();
 
 	uint64 ClassId = InSubClassId;
-	while(true)
+	while (true)
 	{
 		if (ClassId == InParentClassId)
 		{
@@ -341,7 +339,7 @@ const FObjectInfo* FGameplayProvider::FindObjectInfo(uint64 InObjectId) const
 	Session.ReadAccessCheck();
 
 	const int32* ObjectIndex = ObjectIdToIndexMap.Find(InObjectId);
-	if(ObjectIndex != nullptr)
+	if (ObjectIndex != nullptr)
 	{
 		return &ObjectInfos[*ObjectIndex];
 	}
@@ -349,7 +347,7 @@ const FObjectInfo* FGameplayProvider::FindObjectInfo(uint64 InObjectId) const
 	return nullptr;
 }
 
-const FWorldInfo* FGameplayProvider::FindWorldInfo(uint64 InObjectId) const
+const FWorldInfo* FGameplayProvider::FindWorldInfo(const uint64 InObjectId) const
 {
 	Session.ReadAccessCheck();
 
@@ -362,10 +360,9 @@ const FWorldInfo* FGameplayProvider::FindWorldInfo(uint64 InObjectId) const
 	return nullptr;
 }
 
-const FWorldInfo* FGameplayProvider::FindWorldInfoFromObject(uint64 InObjectId) const
+const FWorldInfo* FGameplayProvider::FindWorldInfoFromObject(const uint64 InObjectId) const
 {
-	const FClassInfo* WorldClass = FindClassInfo(TEXT("/Script/Engine.World"));
-	if(WorldClass)
+	if (const FClassInfo* WorldClass = FindClassInfo(TEXT("/Script/Engine.World")))
 	{
 		// Traverse outer chain until we find a world
 		const FObjectInfo* ObjectInfo = FindObjectInfo(InObjectId);
@@ -383,10 +380,9 @@ const FWorldInfo* FGameplayProvider::FindWorldInfoFromObject(uint64 InObjectId)
 	return nullptr;
 }
 
-bool FGameplayProvider::IsWorld(uint64 InObjectId) const
+bool FGameplayProvider::IsWorld(const uint64 InObjectId) const
 {
-	const FClassInfo* WorldClass = FindClassInfo(TEXT("/Script/Engine.World"));
-	if (WorldClass)
+	if (const FClassInfo* WorldClass = FindClassInfo(TEXT("/Script/Engine.World")))
 	{
 		const FObjectInfo* ObjectInfo = FindObjectInfo(InObjectId);
 		return ObjectInfo->ClassId == WorldClass->Id;
@@ -397,8 +393,7 @@ bool FGameplayProvider::IsWorld(uint64 InObjectId) const
 
 const FClassInfo& FGameplayProvider::GetClassInfo(uint64 InClassId) const
 {
-	const FClassInfo* ClassInfo = FindClassInfo(InClassId);
-	if(ClassInfo)
+	if (const FClassInfo* ClassInfo = FindClassInfo(InClassId))
 	{
 		return *ClassInfo;
 	}
@@ -408,13 +403,11 @@ const FClassInfo& FGameplayProvider::GetClassInfo(uint64 InClassId) const
 	return DefaultClassInfo;
 }
 
-const FClassInfo& FGameplayProvider::GetClassInfoFromObject(uint64 InObjectId) const
+const FClassInfo& FGameplayProvider::GetClassInfoFromObject(const uint64 InObjectId) const
 {
-	const FObjectInfo* ObjectInfo = FindObjectInfo(InObjectId);
-	if(ObjectInfo)
+	if (const FObjectInfo* ObjectInfo = FindObjectInfo(InObjectId))
 	{
-		const FClassInfo* ClassInfo = FindClassInfo(ObjectInfo->ClassId);
-		if(ClassInfo)
+		if (const FClassInfo* ClassInfo = FindClassInfo(ObjectInfo->ClassId))
 		{
 			return *ClassInfo;
 		}
@@ -427,8 +420,7 @@ const FClassInfo& FGameplayProvider::GetClassInfoFromObject(uint64 InObjectId) c
 
 const FObjectInfo& FGameplayProvider::GetObjectInfo(uint64 InObjectId) const
 {
-	const FObjectInfo* ObjectInfo = FindObjectInfo(InObjectId);
-	if(ObjectInfo)
+	if (const FObjectInfo* ObjectInfo = FindObjectInfo(InObjectId))
 	{
 		return *ObjectInfo;
 	}
@@ -440,7 +432,7 @@ const FObjectInfo& FGameplayProvider::GetObjectInfo(uint64 InObjectId) const
 
 const TCHAR* FGameplayProvider::GetPropertyName(uint32 InPropertyStringId) const
 {
-	if(const TCHAR*const* FoundString = PropertyStrings.Find(InPropertyStringId))
+	if (const TCHAR* const* FoundString = PropertyStrings.Find(InPropertyStringId))
 	{
 		return *FoundString;
 	}
@@ -455,7 +447,7 @@ void FGameplayProvider::AppendClass(uint64 InClassId, uint64 InSuperId, const TC
 
 	bHasAnyData = true;
 
-	if(ClassIdToIndexMap.Find(InClassId) == nullptr)
+	if (ClassIdToIndexMap.Find(InClassId) == nullptr)
 	{
 		const TCHAR* NewClassName = Session.StoreString(InClassName);
 		const TCHAR* NewClassPathName = Session.StoreString(InClassPathName);
@@ -478,7 +470,7 @@ void FGameplayProvider::AppendObject(uint64 InObjectId, uint64 InOuterId, uint64
 
 	bHasAnyData = true;
 
-	if(ObjectIdToIndexMap.Find(InObjectId) == nullptr)
+	if (ObjectIdToIndexMap.Find(InObjectId) == nullptr)
 	{
 		const TCHAR* NewObjectName = Session.StoreString(InObjectName);
 		const TCHAR* NewObjectPathName = Session.StoreString(InObjectPathName);
@@ -530,22 +522,22 @@ void FGameplayProvider::AppendObjectLifetimeEnd(uint64 InObjectId, double InProf
 	Session.WriteAccessCheck();
 	bHasAnyData = true;
 
-	if (const uint64 *FoundIndex = ActiveObjectLifetimes.Find(InObjectId))
+	if (const uint64* FoundIndex = ActiveObjectLifetimes.Find(InObjectId))
 	{
 		ObjectLifetimes.EndEvent(*FoundIndex, InProfileTime);
 		ActiveObjectLifetimes.Remove(InObjectId);
 	}
 
-	if (const uint64 *FoundIndex = ActiveObjectRecordingLifetimes.Find(InObjectId))
+	if (const uint64* FoundIndex = ActiveObjectRecordingLifetimes.Find(InObjectId))
 	{
-		if(ObjectRecordingLifetimes.GetEventEndTime(*FoundIndex) > InRecordingTime)
+		if (ObjectRecordingLifetimes.GetEventEndTime(*FoundIndex) > InRecordingTime)
 		{
 			ObjectRecordingLifetimes.EndEvent(*FoundIndex, InRecordingTime);
 		}
 		// do not remove from ActiveObjectRecordingLifetimes - lifetimes can be queried by Object Id in GetObjectRecordingLifetime
 	}
 
-	if(int32* ObjectInfoIndex = ObjectIdToIndexMap.Find(InObjectId))
+	if (int32* ObjectInfoIndex = ObjectIdToIndexMap.Find(InObjectId))
 	{
 		OnObjectEndPlayDelegate.Broadcast(InObjectId, InProfileTime, ObjectInfos[*ObjectInfoIndex]);
 	}
@@ -558,14 +550,14 @@ void FGameplayProvider::AppendObjectEvent(uint64 InObjectId, double InTime, cons
 	bHasAnyData = true;
 
 	// Important events need some extra routing
-	if(EndPlayEvent == nullptr)
+	if (EndPlayEvent == nullptr)
 	{
 		EndPlayEvent = Session.StoreString(TEXT("EndPlay"));
 	}
 
 	TSharedPtr<TraceServices::TPointTimeline<FObjectEventMessage>> Timeline;
 	uint32* IndexPtr = ObjectIdToEventTimelines.Find(InObjectId);
-	if(IndexPtr != nullptr)
+	if (IndexPtr != nullptr)
 	{
 		Timeline = EventTimelines[*IndexPtr];
 	}
@@ -580,9 +572,9 @@ void FGameplayProvider::AppendObjectEvent(uint64 InObjectId, double InTime, cons
 	Message.Id = InObjectId;
 	Message.Name = Session.StoreString(InEventName);
 
-	if(Message.Name == EndPlayEvent)
+	if (Message.Name == EndPlayEvent)
 	{
-		if(int32* ObjectInfoIndex = ObjectIdToIndexMap.Find(InObjectId))
+		if (int32* ObjectInfoIndex = ObjectIdToIndexMap.Find(InObjectId))
 		{
 			OnObjectEndPlayDelegate.Broadcast(InObjectId, InTime, ObjectInfos[*ObjectInfoIndex]);
 		}
@@ -599,12 +591,12 @@ void FGameplayProvider::AppendPawnPossess(uint64 InControllerId, uint64 InPawnId
 	bHasAnyData = true;
 
 	// End any active controller attachment interval for this controller
-	if (const uint64 *FoundIndex = ActivePawnPossession.Find(InControllerId))
+	if (const uint64* FoundIndex = ActivePawnPossession.Find(InControllerId))
 	{
 		PawnPossession.EndEvent(*FoundIndex, InTime);
 		ActivePawnPossession.Remove(InControllerId);
 	}
-	
+
 	if (InPawnId)
 	{
 		FPawnPossessMessage Message;
@@ -617,15 +609,15 @@ void FGameplayProvider::AppendPawnPossess(uint64 InControllerId, uint64 InPawnId
 uint64 FGameplayProvider::FindPossessingController(uint64 PawnId, double Time) const
 {
 	uint64 ControllerId = 0;
-	PawnPossession.EnumerateEvents(Time,Time,[&ControllerId, PawnId](double StartTime,double EndTime, const FPawnPossessMessage Message)
-	{
-		if (Message.PawnId == PawnId)
+	PawnPossession.EnumerateEvents(Time, Time, [&ControllerId, PawnId](double StartTime, double EndTime, const FPawnPossessMessage Message)
 		{
-			ControllerId = Message.ControllerId;
-			return TraceServices::EEventEnumerate::Stop;
-		}
-		return TraceServices::EEventEnumerate::Continue;
-	});
+			if (Message.PawnId == PawnId)
+			{
+				ControllerId = Message.ControllerId;
+				return TraceServices::EEventEnumerate::Stop;
+			}
+			return TraceServices::EEventEnumerate::Continue;
+		});
 	return ControllerId;
 }
 
@@ -633,15 +625,15 @@ bool FGameplayProvider::GetObjectTransform(uint64 InObjectId, double InStartTime
 {
 	bool bSuccess = false;
 	TransformTimeline.EnumerateEvents(InStartTime, InEndTime, [&OutTransform, InObjectId, &bSuccess](bool bStart, double Time, const FObjectTransformMessage& TransformMessage)
-	{
-		if (TransformMessage.ObjectId == InObjectId)
 		{
-			bSuccess = true;
-			OutTransform = TransformMessage.Transform;
-		}
-		
-		return TraceServices::EEventEnumerate::Continue;
-	});
+			if (TransformMessage.ObjectId == InObjectId)
+			{
+				bSuccess = true;
+				OutTransform = TransformMessage.Transform;
+			}
+
+			return TraceServices::EEventEnumerate::Continue;
+		});
 	return bSuccess;
 }
 
@@ -664,13 +656,13 @@ TRange<double> FGameplayProvider::GetObjectRecordingLifetime(uint64 ObjectId) co
 	TRACE_CPUPROFILER_EVENT_SCOPE(FGameplayProvider::GetObjectRecordingLifetime);
 	Session.ReadAccessCheck();
 
-	if (const uint64 *FoundIndex = ActiveObjectRecordingLifetimes.Find(ObjectId))
+	if (const uint64* FoundIndex = ActiveObjectRecordingLifetimes.Find(ObjectId))
 	{
 		return TRange<double>(ObjectRecordingLifetimes.GetEventStartTime(*FoundIndex), ObjectRecordingLifetimes.GetEventEndTime(*FoundIndex));
 	}
 	else
 	{
-		return TRange<double>(0,0);
+		return TRange<double>(0, 0);
 	}
 }
 
@@ -695,8 +687,8 @@ FSlateIcon FGameplayProvider::FindIconForClass(uint64 ClassId) const
 #if WITH_EDITOR
 	// traverse class hierarchy until we find a class that exists, and return the icon for that class
 	// loaded trace data may contain classes from different projects than this editor is compiled against or loading
-	// but we will still find a Character icon for example for subclasses of Character
-	while(ClassId)
+	// , but we will still find a Character icon for example for subclasses of Character
+	while (ClassId)
 	{
 		// don't search blueprints because it's very slow and icons are associated with native classes
 		if (const UClass* FoundClass = FindClass(ClassId, false))
@@ -704,11 +696,11 @@ FSlateIcon FGameplayProvider::FindIconForClass(uint64 ClassId) const
 			Icon = FSlateIconFinder::FindIconForClass(FoundClass);
 			break;
 		}
-		
+
 		const FClassInfo& ClassInfo = GetClassInfo(ClassId);
 		ClassId = ClassInfo.SuperId;
 	}
-	
+
 	if (!Icon.IsSet())
 	{
 		Icon = FSlateIconFinder::FindIconForClass(UObject::StaticClass());
@@ -740,23 +732,23 @@ void FGameplayProvider::AppendView(uint64 InPlayerId, double InTime, const FVect
 	Session.UpdateDurationSeconds(InTime);
 }
 
-void FGameplayProvider::AppendWorld(uint64 InObjectId, int32 InPIEInstanceId, uint8 InType, uint8 InNetMode, bool bInIsSimulating)
+void FGameplayProvider::AppendWorld(uint64 InWorldId, int32 InPIEInstanceId, uint8 InType, uint8 InNetMode, bool bInIsSimulating)
 {
 	Session.WriteAccessCheck();
 
 	bHasAnyData = true;
 
-	if (WorldIdToIndexMap.Find(InObjectId) == nullptr)
+	if (WorldIdToIndexMap.Find(InWorldId) == nullptr)
 	{
 		FWorldInfo NewWorldInfo;
-		NewWorldInfo.Id = InObjectId;
+		NewWorldInfo.Id = InWorldId;
 		NewWorldInfo.PIEInstanceId = InPIEInstanceId;
-		NewWorldInfo.Type = (FWorldInfo::EType)InType;
-		NewWorldInfo.NetMode = (FWorldInfo::ENetMode)InNetMode;
+		NewWorldInfo.Type = static_cast<FWorldInfo::EType>(InType);
+		NewWorldInfo.NetMode = static_cast<FWorldInfo::ENetMode>(InNetMode);
 		NewWorldInfo.bIsSimulating = bInIsSimulating;
 
 		int32 NewWorldInfoIndex = WorldInfos.Add(NewWorldInfo);
-		WorldIdToIndexMap.Add(InObjectId, NewWorldInfoIndex);
+		WorldIdToIndexMap.Add(InWorldId, NewWorldInfoIndex);
 	}
 }
 
@@ -769,7 +761,7 @@ void FGameplayProvider::AppendRecordingInfo(uint64 InWorldId, double InProfileTi
 		if (WorldInfos.IsValidIndex(*WorldIndex))
 		{
 			// Try to only keep track of the client world for now.
-			// We should put each world into its own timeline in future, to support frame stepping in sync with server frames
+			// We should put each world into its own timeline in the future, to support frame stepping in sync with server frames
 			if (WorldInfos[*WorldIndex].NetMode == FWorldInfo::ENetMode::Standalone || WorldInfos[*WorldIndex].NetMode == FWorldInfo::ENetMode::Client)
 			{
 				FRecordingInfoMessage NewRecordingInfo;
@@ -778,7 +770,7 @@ void FGameplayProvider::AppendRecordingInfo(uint64 InWorldId, double InProfileTi
 				NewRecordingInfo.RecordingIndex = InRecordingIndex;
 				NewRecordingInfo.FrameIndex = InFrameIndex;
 				NewRecordingInfo.ElapsedTime = InElapsedTime;
-				RecordingDuration = FMath::Max(RecordingDuration,InElapsedTime);
+				RecordingDuration = FMath::Max(RecordingDuration, InElapsedTime);
 
 				if (TSharedRef<TraceServices::TPointTimeline<FRecordingInfoMessage>>* ExistingRecording = Recordings.Find(InRecordingIndex))
 				{
@@ -796,11 +788,11 @@ void FGameplayProvider::AppendRecordingInfo(uint64 InWorldId, double InProfileTi
 }
 
 
-const FGameplayProvider::RecordingInfoTimeline* FGameplayProvider::GetRecordingInfo(uint32 RecordingId) const  
+const FGameplayProvider::RecordingInfoTimeline* FGameplayProvider::GetRecordingInfo(uint32 RecordingId) const
 {
 	Session.ReadAccessCheck();
 
-	if(const TSharedRef<TraceServices::TPointTimeline<FRecordingInfoMessage>>* Recording = Recordings.Find(RecordingId))
+	if (const TSharedRef<TraceServices::TPointTimeline<FRecordingInfoMessage>>* Recording = Recordings.Find(RecordingId))
 	{
 		return &(*Recording).Get();
 	}
@@ -819,7 +811,7 @@ void FGameplayProvider::AppendClassPropertyStringId(uint32 InStringId, const FSt
 	PropertyStrings.Add(InStringId, StoredString);
 }
 
-void FGameplayProvider::AppendPropertiesStart(uint64 InObjectId, double InTime, uint64 InEventId, double InRecordingTime)
+void FGameplayProvider::AppendPropertiesStart(const uint64 InObjectId, double InTime, uint64 InEventId, double InRecordingTime)
 {
 	Session.WriteAccessCheck();
 
@@ -828,7 +820,7 @@ void FGameplayProvider::AppendPropertiesStart(uint64 InObjectId, double InTime,
 
 	TSharedPtr<FObjectPropertiesStorage> Storage;
 	uint32* IndexPtr = ObjectIdToPropertiesStorage.Find(InObjectId);
-	if(IndexPtr != nullptr)
+	if (IndexPtr != nullptr)
 	{
 		Storage = PropertiesStorage[*IndexPtr];
 	}
@@ -850,7 +842,7 @@ void FGameplayProvider::AppendPropertiesStart(uint64 InObjectId, double InTime,
 	Message.ElapsedTime = InRecordingTime;
 }
 
-void FGameplayProvider::AppendPropertiesEnd(uint64 InObjectId, double InTime)
+void FGameplayProvider::AppendPropertiesEnd(const uint64 InObjectId, double InTime)
 {
 	Session.WriteAccessCheck();
 
@@ -859,7 +851,7 @@ void FGameplayProvider::AppendPropertiesEnd(uint64 InObjectId, double InTime)
 
 	TSharedPtr<FObjectPropertiesStorage> Storage;
 	uint32* IndexPtr = ObjectIdToPropertiesStorage.Find(InObjectId);
-	if(IndexPtr != nullptr)
+	if (IndexPtr != nullptr)
 	{
 		Storage = PropertiesStorage[*IndexPtr];
 	}
@@ -871,7 +863,7 @@ void FGameplayProvider::AppendPropertiesEnd(uint64 InObjectId, double InTime)
 		PropertiesStorage.Add(Storage.ToSharedRef());
 	}
 
-	if(Storage->OpenEventId != 0)
+	if (Storage->OpenEventId != 0)
 	{
 		uint64 EventIndex = Storage->Timeline->AppendBeginEvent(Storage->OpenStartTime, Storage->OpenEvent);
 		Storage->Timeline->EndEvent(EventIndex, InTime);
@@ -880,7 +872,7 @@ void FGameplayProvider::AppendPropertiesEnd(uint64 InObjectId, double InTime)
 	}
 }
 
-void FGameplayProvider::AppendPropertyValue(uint64 InObjectId, double InTime, uint64 InEventId, int32 InParentId, uint32 InTypeStringId, uint32 InNameId, uint32 InParentNameId, const FStringView& InValue)
+void FGameplayProvider::AppendPropertyValue(const uint64 InObjectId, double InTime, uint64 InEventId, int32 InParentId, uint32 InTypeStringId, uint32 InNameId, uint32 InParentNameId, const FStringView& InValue)
 {
 	Session.WriteAccessCheck();
 
@@ -888,8 +880,7 @@ void FGameplayProvider::AppendPropertyValue(uint64 InObjectId, double InTime, ui
 	bHasObjectProperties = true;
 
 	TSharedPtr<FObjectPropertiesStorage> Storage;
-	uint32* IndexPtr = ObjectIdToPropertiesStorage.Find(InObjectId);
-	if(IndexPtr != nullptr)
+	if (const uint32* IndexPtr = ObjectIdToPropertiesStorage.Find(InObjectId))
 	{
 		Storage = PropertiesStorage[*IndexPtr];
 	}
@@ -900,30 +891,30 @@ void FGameplayProvider::AppendPropertyValue(uint64 InObjectId, double InTime, ui
 		ObjectIdToPropertiesStorage.Add(InObjectId, PropertiesStorage.Num());
 		PropertiesStorage.Add(Storage.ToSharedRef());
 	}
-	
-	if(Storage->OpenEventId == InEventId)
+
+	if (Storage->OpenEventId == InEventId)
 	{
 		FObjectPropertyValue& Message = Storage->Values.AddDefaulted_GetRef();
-		
+
 		Message.Value = Session.StoreString(InValue);
 		Message.ValueAsFloat = FCString::Atof(Message.Value);
 		Message.ParentId = InParentId;
 		Message.TypeStringId = InTypeStringId;
 		Message.NameId = InNameId;
 		Message.ParentNameId = InParentNameId;
-		
+
 		Storage->OpenEvent.PropertyValueEndIndex = Storage->Values.Num();
 	}
 }
 
-bool FGameplayProvider::HasAnyData() const 
+bool FGameplayProvider::HasAnyData() const
 {
 	Session.ReadAccessCheck();
 
 	return bHasAnyData;
 }
 
-bool FGameplayProvider::HasObjectProperties() const 
+bool FGameplayProvider::HasObjectProperties() const
 {
 	Session.ReadAccessCheck();
 
diff --git a/Engine/Plugins/Animation/GameplayInsights/Source/GameplayInsights/Private/GameplayProvider.h b/Engine/Plugins/Animation/GameplayInsights/Source/GameplayInsights/Private/GameplayProvider.h
--- a/Engine/Plugins/Animation/GameplayInsights/Source/GameplayInsights/Private/GameplayProvider.h
+++ b/Engine/Plugins/Animation/GameplayInsights/Source/GameplayInsights/Private/GameplayProvider.h
@@ -3,11 +3,12 @@
 #pragma once
 
 #include "IGameplayProvider.h"
-#include "Model/PointTimeline.h"
 #include "Containers/Map.h"
-#include "Templates/SharedPointer.h"
 #include "Containers/StringView.h"
+#include "Math/TransformVectorized.h"
 #include "Model/IntervalTimeline.h"
+#include "Model/PointTimeline.h"
+#include "Templates/SharedPointer.h"
 
 namespace TraceServices { class IAnalysisSession; }
 
@@ -47,34 +48,32 @@ public:
 	virtual double GetRecordingDuration() const override;
 	virtual FSlateIcon FindIconForClass(uint64 ClassId) const override;
 
-
 	/** Add a class message */
 	void AppendClass(uint64 InClassId, uint64 InSuperId, const TCHAR* InClassName, const TCHAR* InClassPathName);
 
 	/** Add an object message */
 	void AppendObject(uint64 InObjectId, uint64 InOuterId, uint64 InClassId, const TCHAR* InObjectName, const TCHAR* InObjectPathName);
 
-	
 	/** Add an object transform */
 	void AppendObjectTransform(uint64 InObjectId, double InProfileTime, double InRecordingTime, const FTransform& Transform);
-	
+
 	/** Add an object create message */
 	void AppendObjectLifetimeBegin(uint64 InObjectId, double InProfileTime, double InRecordingTime);
-	
+
 	/** Add an object destroy message */
 	void AppendObjectLifetimeEnd(uint64 InObjectId, double InProfileTime, double InRecordingTime);	
 
 	/** Add an object event message */
 	void AppendObjectEvent(uint64 InObjectId, double InTime, const TCHAR* InEvent);
-	
+
 	/** Add a Controller Attach message */
 	void AppendPawnPossess(uint64 InControllerId, uint64 InPawnId, double InTime);
 
 	/** Add a view message */
-	void AppendView(uint64 InObjectId, double InTime, const FVector& InPosition, const FRotator& InRotation, float InFov, float InAspectRatio);
+	void AppendView(uint64 InPlayerId, double InTime, const FVector& InPosition, const FRotator& InRotation, float InFov, float InAspectRatio);
 
 	/** Add a world message */
-	void AppendWorld(uint64 InObjectId, int32 InPIEInstanceId, uint8 InType, uint8 InNetMode, bool bInIsSimulating);
+	void AppendWorld(uint64 InWorldId, int32 InPIEInstanceId, uint8 InType, uint8 InNetMode, bool bInIsSimulating);
 
 	/** Add a recording info message */
 	void AppendRecordingInfo(uint64 InWorldId, double InProfileTime, uint32 InRecordingIndex, uint32 InFrameIndex, double InElapsedTime);
@@ -82,10 +81,10 @@ public:
 	/** Add a class property string ID message */
 	void AppendClassPropertyStringId(uint32 InStringId, const FStringView& InString);
 
-	/** Add a properties start message */
+	/** Add a 'properties start' message */
 	void AppendPropertiesStart(uint64 InObjectId, double InTime, uint64 InEventId, double InRecordingTime);
 
-	/** Add a properties end message */
+	/** Add a 'properties end' message */
 	void AppendPropertiesEnd(uint64 InObjectId, double InTime);
 
 	/** Add a property value message */
@@ -140,14 +139,14 @@ private:
 		uint64 ControllerId = 0;
 		uint64 PawnId = 0;
 	};
-	
+
 	struct FObjectTransformMessage
-   	{
-   		uint64 ObjectId = 0;
+	{
+		uint64 ObjectId = 0;
 		double RecordingTime;
-   		FTransform Transform;
-   	};
-	
+		FTransform Transform;
+	};
+
 	struct FObjectExistsMessage
 	{
 		uint64 ObjectId = 0;
@@ -185,10 +184,10 @@ private:
 
 	// Timeline containing intervals where a controller is attached to a pawn
 	TraceServices::TIntervalTimeline<FPawnPossessMessage> PawnPossession;
-	
+
 	// Timeline containing intervals where an object exists
 	TraceServices::TIntervalTimeline<FObjectExistsMessage> ObjectLifetimes;
-	
+
 	// Timeline containing intervals where an object exists - In Rewind Debugger recording time units, to avoid excessive conversions
 	TraceServices::TIntervalTimeline<FObjectExistsMessage> ObjectRecordingLifetimes;
 
@@ -197,18 +196,18 @@ private:
 
 	// map from object id to index in ObjectLifetimes, for lookup when objects are destroyed
 	TMap<uint64, uint64> ActiveObjectLifetimes;
-	
+
 	// map from object id to index in ObjectRecordingLifetimes, for lookup
 	TMap<uint64, uint64> ActiveObjectRecordingLifetimes;
 
 	/** Whether we have any data */
-	bool bHasAnyData;
-	
+	bool bHasAnyData = false;
+
 	/** Whether we have any data */
 	double RecordingDuration = 0.0;
 
 	/** Whether we have any object properties */
-	bool bHasObjectProperties;
+	bool bHasObjectProperties = false;
 
 	/** MultiMap from parent object to sub objects, to avoid slow reverse lookups */
 	TMultiMap<uint64, uint64> ObjectHierarchy;
diff --git a/Engine/Plugins/Animation/GameplayInsights/Source/GameplayInsights/Public/IGameplayProvider.h b/Engine/Plugins/Animation/GameplayInsights/Source/GameplayInsights/Public/IGameplayProvider.h
--- a/Engine/Plugins/Animation/GameplayInsights/Source/GameplayInsights/Public/IGameplayProvider.h
+++ b/Engine/Plugins/Animation/GameplayInsights/Source/GameplayInsights/Public/IGameplayProvider.h
@@ -2,12 +2,16 @@
 
 #pragma once
 
+#if UE_ENABLE_INCLUDE_ORDER_DEPRECATED_IN_5_7
 #include "CoreMinimal.h"
+#endif // UE_ENABLE_INCLUDE_ORDER_DEPRECATED_IN_5_7
+#include "Math/Vector.h"
+#include "Math/Rotator.h"
 #include "TraceServices/Model/AnalysisSession.h"
 #include "TraceServices/Containers/Timelines.h"
 
 struct FClassInfo
-{	
+{
 	uint64 Id = 0;
 	uint64 SuperId = 0;
 	const TCHAR* Name = nullptr;
@@ -15,7 +19,7 @@ struct FClassInfo
 };
 
 struct FObjectInfo
-{	
+{
 	uint64 Id = 0;
 	uint64 OuterId = 0;
 	uint64 ClassId = 0;
diff --git a/Engine/Plugins/Animation/GameplayInsights/Source/RewindDebugger/Private/AnimInstanceHelpers.cpp b/Engine/Plugins/Animation/GameplayInsights/Source/RewindDebugger/Private/AnimInstanceHelpers.cpp
--- a/Engine/Plugins/Animation/GameplayInsights/Source/RewindDebugger/Private/AnimInstanceHelpers.cpp
+++ b/Engine/Plugins/Animation/GameplayInsights/Source/RewindDebugger/Private/AnimInstanceHelpers.cpp
@@ -11,6 +11,7 @@
 #include "RewindDebuggerAnimation.h"
 #include "Subsystems/AssetEditorSubsystem.h"
 #include "IAnimationBlueprintEditor.h"
+#include "RewindDebuggerModule.h"
 #include "Insights/IUnrealInsightsModule.h"
 #include "ToolMenus.h"
 
@@ -69,7 +70,7 @@ static bool OpenAnimBlueprintAndAttachDebugger(const TraceServices::IAnalysisSes
 
 bool FAnimInstanceDoubleClickHandler::HandleDoubleClick(IRewindDebugger* RewindDebugger)
 {
-	TSharedPtr<FDebugObjectInfo> SelectedObject = RewindDebugger->GetSelectedComponent();
+	TSharedPtr<FDebugObjectInfo> SelectedObject = RewindDebugger->GetSelectedObject();
 	if (SelectedObject.IsValid())
 	{
 		return OpenAnimBlueprintAndAttachDebugger(RewindDebugger->GetAnalysisSession(), SelectedObject->ObjectId);
@@ -86,11 +87,11 @@ FName FAnimInstanceDoubleClickHandler::GetTargetTypeName() const
 void FAnimInstanceMenu::Register()
 {
 #if OBJECT_TRACE_ENABLED
-	UToolMenu* Menu = UToolMenus::Get()->FindMenu("RewindDebugger.ComponentContextMenu");
+	UToolMenu* Menu = UToolMenus::Get()->FindMenu(FRewindDebuggerModule::TrackContextMenuName);
 	FToolMenuSection& Section = Menu->FindOrAddSection("Blueprint");
-	FToolMenuEntry& Entry = Section.AddDynamicEntry("DebugAnimInstanceEntry", FNewToolMenuSectionDelegate::CreateLambda([](FToolMenuSection& InSection)
+	Section.AddDynamicEntry("DebugAnimInstanceEntry", FNewToolMenuSectionDelegate::CreateLambda([](FToolMenuSection& InSection)
 	{
-		UComponentContextMenuContext* Context = InSection.FindContext<UComponentContextMenuContext>();
+		URewindDebuggerTrackContextMenuContext* Context = InSection.FindContext<URewindDebuggerTrackContextMenuContext>();
 		if (Context && Context->SelectedObject.IsValid() && Context->TypeHierarchy.Contains("AnimInstance"))
 		{
 			InSection.AddMenuEntry(NAME_None,
@@ -99,8 +100,8 @@ void FAnimInstanceMenu::Register()
 							FSlateIcon(),
 							FExecuteAction::CreateLambda([ObjectId = Context->SelectedObject->ObjectId]()
 							{
-								IUnrealInsightsModule& UnrealInsightsModule = FModuleManager::LoadModuleChecked<IUnrealInsightsModule>("TraceInsights");
-								TSharedPtr<const TraceServices::IAnalysisSession> Session = UnrealInsightsModule.GetAnalysisSession();
+								const IUnrealInsightsModule& UnrealInsightsModule = FModuleManager::LoadModuleChecked<IUnrealInsightsModule>("TraceInsights");
+								const TSharedPtr<const TraceServices::IAnalysisSession> Session = UnrealInsightsModule.GetAnalysisSession();
 								OpenAnimBlueprintAndAttachDebugger(Session.Get(), ObjectId);
 							}));
 		}
diff --git a/Engine/Plugins/Animation/GameplayInsights/Source/RewindDebugger/Private/BlueprintDoubleClickHandler.cpp b/Engine/Plugins/Animation/GameplayInsights/Source/RewindDebugger/Private/BlueprintDoubleClickHandler.cpp
--- a/Engine/Plugins/Animation/GameplayInsights/Source/RewindDebugger/Private/BlueprintDoubleClickHandler.cpp
+++ b/Engine/Plugins/Animation/GameplayInsights/Source/RewindDebugger/Private/BlueprintDoubleClickHandler.cpp
@@ -1,4 +1,4 @@
-﻿// Copyright Epic Games, Inc. All Rights Reserved.
+// Copyright Epic Games, Inc. All Rights Reserved.
 
 #include "BlueprintDoubleClickHandler.h"
 
@@ -47,7 +47,7 @@ static bool OpenBlueprintAndAttachDebugger(const TraceServices::IAnalysisSession
 bool FBlueprintDoubleClickHandler::HandleDoubleClick(IRewindDebugger* RewindDebugger)
 {
 #if OBJECT_TRACE_ENABLED
-	const TSharedPtr<FDebugObjectInfo> SelectedObject = RewindDebugger->GetSelectedComponent();
+	const TSharedPtr<FDebugObjectInfo> SelectedObject = RewindDebugger->GetSelectedObject();
 	if (SelectedObject.IsValid())
 	{
 		return OpenBlueprintAndAttachDebugger(RewindDebugger->GetAnalysisSession(), SelectedObject->ObjectId);
diff --git a/Engine/Plugins/Animation/GameplayInsights/Source/RewindDebugger/Private/PropertyTraceMenu.cpp b/Engine/Plugins/Animation/GameplayInsights/Source/RewindDebugger/Private/PropertyTraceMenu.cpp
--- a/Engine/Plugins/Animation/GameplayInsights/Source/RewindDebugger/Private/PropertyTraceMenu.cpp
+++ b/Engine/Plugins/Animation/GameplayInsights/Source/RewindDebugger/Private/PropertyTraceMenu.cpp
@@ -1,9 +1,10 @@
 ﻿// Copyright Epic Games, Inc. All Rights Reserved.
 
 #include "PropertyTraceMenu.h"
-#include "Modules/ModuleManager.h"
 #include "IGameplayInsightsModule.h"
+#include "Modules/ModuleManager.h"
 #include "ObjectTrace.h"
+#include "RewindDebuggerModule.h"
 #include "ToolMenus.h"
 
 #define LOCTEXT_NAMESPACE "RewindDebuggerComponentHelpers"
@@ -12,11 +13,11 @@ void FPropertyTraceMenu::Register()
 {
 #if OBJECT_TRACE_ENABLED
 
-	UToolMenu* Menu = UToolMenus::Get()->FindMenu("RewindDebugger.ComponentContextMenu");
+	UToolMenu* Menu = UToolMenus::Get()->FindMenu(FRewindDebuggerModule::TrackContextMenuName);
 	FToolMenuSection& Section = Menu->FindOrAddSection("Object");
-	FToolMenuEntry& Entry = Section.AddDynamicEntry("AssetManagerEditorViewCommands", FNewToolMenuSectionDelegate::CreateLambda([](FToolMenuSection& InSection)
+	Section.AddDynamicEntry("AssetManagerEditorViewCommands", FNewToolMenuSectionDelegate::CreateLambda([](FToolMenuSection& InSection)
 	{
-		UComponentContextMenuContext* Context = InSection.FindContext<UComponentContextMenuContext>();
+		URewindDebuggerTrackContextMenuContext* Context = InSection.FindContext<URewindDebuggerTrackContextMenuContext>();
 		if (Context && Context->SelectedObject.IsValid())
 		{
 			InSection.AddMenuEntry(
diff --git a/Engine/Plugins/Animation/GameplayInsights/Source/RewindDebugger/Private/RewindDebugger.cpp b/Engine/Plugins/Animation/GameplayInsights/Source/RewindDebugger/Private/RewindDebugger.cpp
--- a/Engine/Plugins/Animation/GameplayInsights/Source/RewindDebugger/Private/RewindDebugger.cpp
+++ b/Engine/Plugins/Animation/GameplayInsights/Source/RewindDebugger/Private/RewindDebugger.cpp
@@ -2,23 +2,24 @@
 
 #include "RewindDebugger.h"
 
-#include "DesktopPlatformModule.h"
 #include "Animation/AnimBlueprintGeneratedClass.h"
 #include "Animation/AnimTrace.h"
-#include "Components/SkeletalMeshComponent.h"
+#include "DesktopPlatformModule.h"
 #include "Editor.h"
+#include "Editor/UnrealEdEngine.h"
 #include "Engine/World.h"
 #include "EngineUtils.h"
-#include "GameFramework/Controller.h"
 #include "GameFramework/Pawn.h"
 #include "HAL/PlatformFileManager.h"
 #include "IAnimationProvider.h"
 #include "IDesktopPlatform.h"
 #include "IGameplayProvider.h"
+#include "Insights/IUnrealInsightsModule.h"
 #include "IRewindDebuggerDoubleClickHandler.h"
 #include "IRewindDebuggerExtension.h"
-#include "Insights/IUnrealInsightsModule.h"
+#include "Kismet2/DebuggerCommands.h"
 #include "LevelEditor.h"
+#include "Misc/MessageDialog.h"
 #include "Modules/ModuleManager.h"
 #include "ObjectTrace.h"
 #include "ProfilingDebugging/TraceAuxiliary.h"
@@ -26,25 +27,19 @@
 #include "RewindDebuggerModule.h"
 #include "RewindDebuggerObjectTrack.h"
 #include "RewindDebuggerPlaceholderTrack.h"
+#include "RewindDebuggerRuntime/RewindDebuggerRuntime.h"
 #include "RewindDebuggerSettings.h"
+#include "RewindDebuggerTrackCreators.h"
 #include "SLevelViewport.h"
 #include "SModalSessionBrowser.h"
 #include "ToolMenus.h"
 #include "Trace/StoreClient.h"
 #include "TraceServices/Model/Frames.h"
+#include "UnrealEdGlobals.h"
 #include "UObject/UObjectIterator.h"
 #include "Widgets/Docking/SDockTab.h"
-#include "Widgets/Layout/SSpacer.h"
-#include "RewindDebuggerCommands.h"
-#include "RewindDebuggerTrackCreators.h"
-#include "UnrealEdGlobals.h"
-#include "Editor/UnrealEdEngine.h"
-#include "Kismet2/DebuggerCommands.h"
-#include "RewindDebuggerRuntime/RewindDebuggerRuntime.h"
-#include "TraceServices/AnalysisService.h"
-#include "TraceServices/ITraceServicesModule.h"
 #include "Widgets/Input/SNumericEntryBox.h"
-#include "Misc/MessageDialog.h"
+#include "Widgets/Layout/SSpacer.h"
 
 #define LOCTEXT_NAMESPACE "RewindDebugger"
 
@@ -61,11 +56,11 @@ static void IterateExtensions(TFunction<void(IRewindDebuggerExtension* Extension
 	}
 }
 
-static void TraceSubobjects(UObject* OuterObject)
+static void TraceSubobjects(const UObject* OuterObject)
 {
 	TArray<UObject*> Subobjects;
 	GetObjectsWithOuter(OuterObject, Subobjects, true);
-	for (UObject* Subobject : Subobjects)
+	for (const UObject* Subobject : Subobjects)
 	{
 		TRACE_OBJECT_LIFETIME_BEGIN(Subobject);
 	}
@@ -84,14 +79,14 @@ FRewindDebugger::FRewindDebugger()
 		Runtime->RecordingStarted.AddRaw(this, &FRewindDebugger::OnRecordingStarted);
 		Runtime->RecordingStarted.AddRaw(this, &FRewindDebugger::OnRecordingStopped);
 	}
-	
+
 	RewindDebugger::FRewindDebuggerTrackCreators::EnumerateCreators([this](const RewindDebugger::IRewindDebuggerTrackCreator* Creator)
-    {
-		Creator->GetTrackTypes(TrackTypes);
-    });
-	
+		{
+			Creator->GetTrackTypes(TrackTypes);
+		});
+
 	RecordingDuration.Set(0);
-	
+
 	UnrealInsightsModule = &FModuleManager::LoadModuleChecked<IUnrealInsightsModule>("TraceInsights");
 
 	if (GEditor->bIsSimulatingInEditor || GEditor->PlayWorld)
@@ -105,8 +100,7 @@ FRewindDebugger::FRewindDebugger()
 	FEditorDelegates::EndPIE.AddRaw(this, &FRewindDebugger::OnPIEStopped);
 	FEditorDelegates::SingleStepPIE.AddRaw(this, &FRewindDebugger::OnPIESingleStepped);
 
-	
-	DebugTargetActor.OnPropertyChanged = DebugTargetActor.OnPropertyChanged.CreateLambda([this](FString Target)
+	SelectedObjectName.OnPropertyChanged = SelectedObjectName.OnPropertyChanged.CreateLambda([this](FString Target)
 		{
 			URewindDebuggerSettings& Settings = URewindDebuggerSettings::Get();
 			if (Settings.DebugTargetActor != Target)
@@ -115,14 +109,14 @@ FRewindDebugger::FRewindDebugger()
 				Settings.Modify();
 				Settings.SaveConfig();
 			}
-		
-			TargetObjectIds.SetNum(0);
-			GetTargetObjectIds(TargetObjectIds);
+
+			CandidateIds.SetNum(0);
+			GetTargetObjectIds(CandidateIds);
 			// make sure all the SubObjects of the target actor have been traced
 #if OBJECT_TRACE_ENABLED
-			for (uint64 TargetObjectId : TargetObjectIds)
+			for (uint64 TargetObjectId : CandidateIds)
 			{
-				if (UObject* TargetObject = FObjectTrace::GetObjectFromId(TargetObjectId))
+				if (const UObject* TargetObject = FObjectTrace::GetObjectFromId(TargetObjectId))
 				{
 					TraceSubobjects(TargetObject);
 				}
@@ -133,16 +127,16 @@ FRewindDebugger::FRewindDebugger()
 		});
 
 	TickerHandle = FTSTicker::GetCoreTicker().AddTicker(TEXT("RewindDebugger"), 0.0f, [this](float DeltaTime)
-	{
-		QUICK_SCOPE_CYCLE_COUNTER(STAT_FRewindDebuggerModule_Tick);
+		{
+			QUICK_SCOPE_CYCLE_COUNTER(STAT_FRewindDebuggerModule_Tick);
 
-		Tick(DeltaTime);
+			Tick(DeltaTime);
 
-		return true;
-	});
+			return true;
+		});
 }
 
-FRewindDebugger::~FRewindDebugger() 
+FRewindDebugger::~FRewindDebugger()
 {
 	FEditorDelegates::PostPIEStarted.RemoveAll(this);
 	FEditorDelegates::PausePIE.RemoveAll(this);
@@ -151,29 +145,29 @@ FRewindDebugger::~FRewindDebugger()
 	FEditorDelegates::SingleStepPIE.RemoveAll(this);
 
 	FTSTicker::GetCoreTicker().RemoveTicker(TickerHandle);
-	
+
 	if (RewindDebugger::FRewindDebuggerRuntime* Runtime = RewindDebugger::FRewindDebuggerRuntime::Instance())
 	{
 		Runtime->RecordingStarted.RemoveAll(this);
 	}
 }
 
-void FRewindDebugger::Initialize() 
+void FRewindDebugger::Initialize()
 {
 	InternalInstance = new FRewindDebugger;
 }
 
-void FRewindDebugger::Shutdown() 
+void FRewindDebugger::Shutdown()
 {
 	delete InternalInstance;
 }
 
-void FRewindDebugger::OnComponentListChanged(const FOnComponentListChanged& InComponentListChangedDelegate)
+void FRewindDebugger::SetTrackListChangedDelegate(const FOnTrackListChanged& InTrackListChangedDelegate)
 {
-	ComponentListChangedDelegate = InComponentListChangedDelegate;
+	TrackListChangedDelegate = InTrackListChangedDelegate;
 }
 
-void FRewindDebugger::OnTrackCursor(const FOnTrackCursor& InTrackCursorDelegate)
+void FRewindDebugger::SetTrackCursorDelegate(const FOnTrackCursor& InTrackCursorDelegate)
 {
 	TrackCursorDelegate = InTrackCursorDelegate;
 }
@@ -193,18 +187,18 @@ void FRewindDebugger::OnPIEPaused(bool bSimulating)
 {
 	bPIESimulating = false;
 	ControlState = EControlState::Pause;
-	
+
+#if OBJECT_TRACE_ENABLED
 	if (IsRecording())
 	{
-    #if OBJECT_TRACE_ENABLED
-		UWorld* World = GetWorldToVisualize();
+		const UWorld* World = GetWorldToVisualize();
 		SetCurrentScrubTime(FObjectTrace::GetWorldElapsedTime(World));
-    #endif // OBJECT_TRACE_ENABLED
 	}
-	
+#endif // OBJECT_TRACE_ENABLED
+
 	if (ShouldAutoEject() && FPlayWorldCommandCallbacks::IsInPIE())
 	{
-		bool CanEject= false;
+		bool CanEject = false;
 		for (auto It = GUnrealEd->SlatePlayInEditorMap.CreateIterator(); It; ++It)
 		{
 			CanEject = CanEject || It.Value().DestinationSlateViewport.IsValid();
@@ -229,13 +223,13 @@ void FRewindDebugger::OnPIEResumed(bool bSimulating)
 
 void FRewindDebugger::OnPIESingleStepped(bool bSimulating)
 {
+#if OBJECT_TRACE_ENABLED
 	if (IsRecording())
 	{
-    #if OBJECT_TRACE_ENABLED
-		UWorld* World = GetWorldToVisualize();
+		const UWorld* World = GetWorldToVisualize();
 		SetCurrentScrubTime(FObjectTrace::GetWorldElapsedTime(World));
-    #endif // OBJECT_TRACE_ENABLED
 	}
+#endif // OBJECT_TRACE_ENABLED
 }
 
 
@@ -244,16 +238,16 @@ void FRewindDebugger::OnPIEStopped(bool bSimulating)
 	if (IsRecording() && bPIESimulating)
 	{
 #if OBJECT_TRACE_ENABLED
-		UWorld* World = GetWorldToVisualize();
+		const UWorld* World = GetWorldToVisualize();
 		SetCurrentScrubTime(FObjectTrace::GetWorldElapsedTime(World));
 #endif // OBJECT_TRACE_ENABLED
 	}
-	
+
 	bPIEStarted = false;
 	bPIESimulating = false;
 
 	StopRecording();
-	
+
 	bDisplayWorldIdValid = false;
 }
 
@@ -265,7 +259,7 @@ bool FRewindDebugger::GetTargetActorPosition(FVector& OutPosition) const
 
 uint64 FRewindDebugger::GetTargetActorId() const
 {
-	if (DebugTargetActor.Get() == "")
+	if (SelectedObjectName.Get() == "")
 	{
 		return 0;
 	}
@@ -277,24 +271,23 @@ uint64 FRewindDebugger::GetTargetActorId() const
 		if (const IGameplayProvider* GameplayProvider = Session->ReadProvider<IGameplayProvider>("GameplayProvider"))
 		{
 			GameplayProvider->EnumerateObjects(CurrentTraceRange.GetLowerBoundValue(), CurrentTraceRange.GetUpperBoundValue(), [this, &TargetActorId](const FObjectInfo& InObjectInfo)
-			{
-				if (DebugTargetActor.Get() == InObjectInfo.Name)
 				{
-					TargetActorId = InObjectInfo.Id;
-				}
-			});
+					if (SelectedObjectName.Get() == InObjectInfo.Name)
+					{
+						TargetActorId = InObjectInfo.Id;
+					}
+				});
 		}
 	}
 
 	return TargetActorId;
 }
 
-
 void FRewindDebugger::GetTargetObjectIds(TArray<uint64>& OutTargetObjectIds) const
 {
 	OutTargetObjectIds.Empty(2);
 
-	if (DebugTargetActor.Get() == "")
+	if (SelectedObjectName.Get() == "")
 	{
 		return;
 	}
@@ -307,7 +300,7 @@ void FRewindDebugger::GetTargetObjectIds(TArray<uint64>& OutTargetObjectIds) con
 		{
 			GameplayProvider->EnumerateObjects(CurrentTraceRange.GetLowerBoundValue(), CurrentTraceRange.GetUpperBoundValue(), [this, &OutTargetObjectIds](const FObjectInfo& InObjectInfo)
 				{
-					if (DebugTargetActor.Get() == InObjectInfo.Name)
+					if (SelectedObjectName.Get() == InObjectInfo.Name)
 					{
 						OutTargetObjectIds.Add(InObjectInfo.Id);
 					}
@@ -319,9 +312,9 @@ void FRewindDebugger::GetTargetObjectIds(TArray<uint64>& OutTargetObjectIds) con
 #if OBJECT_TRACE_ENABLED
 	if (IsRecording())
 	{
-		for (uint64 OutTargetObjectId : TargetObjectIds)
+		for (uint64 CandidateId : CandidateIds)
 		{
-			if (UObject* TargetObject = FObjectTrace::GetObjectFromId(OutTargetObjectId))
+			if (const UObject* TargetObject = FObjectTrace::GetObjectFromId(CandidateId))
 			{
 				TraceSubobjects(TargetObject);
 			}
@@ -330,22 +323,21 @@ void FRewindDebugger::GetTargetObjectIds(TArray<uint64>& OutTargetObjectIds) con
 #endif
 }
 
-
 void FRewindDebugger::RefreshDebugTracks()
 {
 	static const FName DebugMessageTrackName = "DebugMessageDummyTrack";
 	TRACE_CPUPROFILER_EVENT_SCOPE(FRewindDebugger::RefreshDebugTracks);
 
-	if (TargetObjectIds.Num() == 0)
+	if (CandidateIds.Num() == 0)
 	{
-		GetTargetObjectIds(TargetObjectIds);
+		GetTargetObjectIds(CandidateIds);
 	}
-	
-	FString DebugTargetActorName = DebugTargetActor.Get();
-	
-	if (TargetObjectIds.Num() == 0 && !DebugTargetActorName.IsEmpty())
+
+	const FString SelectionName = SelectedObjectName.Get();
+
+	if (CandidateIds.Num() == 0 && !SelectionName.IsEmpty())
 	{
-		// fallback codepath for when the target object is not found
+		// fallback code path for when the target object is not found
 
 		if (DebugTracks.Num() != 2)
 		{
@@ -354,79 +346,58 @@ void FRewindDebugger::RefreshDebugTracks()
 			DebugTracks.SetNum(2);
 		}
 
-		if (DebugTracks[1] == nullptr || DebugTracks[0] == nullptr || DebugTracks[0]->GetName().ToString() != DebugTargetActor.Get() )
+		if (DebugTracks[1] == nullptr || DebugTracks[0] == nullptr || DebugTracks[0]->GetName().ToString() != SelectionName)
 		{
-			DebugTracks[0] = MakeShared<FRewindDebuggerPlaceholderTrack>(FName(DebugTargetActorName), FText::FromString(DebugTargetActorName)); 
+			DebugTracks[0] = MakeShared<FRewindDebuggerPlaceholderTrack>(FName(SelectionName), FText::FromString(SelectionName));
 			DebugTracks[1] = MakeShared<FRewindDebuggerPlaceholderTrack>(DebugMessageTrackName, NSLOCTEXT("RewindDebugger", "No Debug Data", " - Start a recording to debug"));
-			ComponentListChangedDelegate.ExecuteIfBound();
+			TrackListChangedDelegate.ExecuteIfBound();
 		}
 	}
 	else
 	{
-		if (const TraceServices::IAnalysisSession* Session = GetAnalysisSession())
-    	{
-    		TraceServices::FAnalysisSessionReadScope SessionReadScope(*Session);
-    
-    		UWorld* World = GetWorldToVisualize();
-    
-    		if (const IGameplayProvider* GameplayProvider = Session->ReadProvider<IGameplayProvider>("GameplayProvider"))
-    		{
-    
-    			bool bChanged = false;
-    
-    			// remove any existing tracks that don't match the current list of object ids
-    			for (int TrackIndex = DebugTracks.Num()-1; TrackIndex>=0; TrackIndex--)
-    			{
-    				uint64 TrackId = DebugTracks[TrackIndex]->GetObjectId();
-    				if (!TargetObjectIds.Contains(TrackId))
-    				{
-    					DebugTracks.RemoveAt(TrackIndex);
-    				}
-    			}
-    
-    			// add new tracks for current list of object ids if they don't already exist
-    			for (uint64 TargetObjectId : TargetObjectIds)
-    			{
-    				TSharedPtr<RewindDebugger::FRewindDebuggerTrack>* FoundTrack = DebugTracks.FindByPredicate([TargetObjectId](const TSharedPtr<RewindDebugger::FRewindDebuggerTrack>& Track) { return Track->GetObjectId() == TargetObjectId; });
-    
-    				if (!FoundTrack)
-    				{
-    					DebugTracks.Add(MakeShared<RewindDebugger::FRewindDebuggerObjectTrack>(TargetObjectId, DebugTargetActor.Get(), true));
-    					bChanged = true;
-    				}
-    			}
-    
-    			// update all tracks
-    			for (TSharedPtr<RewindDebugger::FRewindDebuggerTrack>& DebugTrack : DebugTracks )
-    			{
-    				if (DebugTrack->Update())
-    				{
-    					bChanged = true;
-    				}
-    			}
-    
-    			if (bChanged)
-    			{
-    				ComponentListChangedDelegate.ExecuteIfBound();
-    			}
-    		}
-    	}	
-	}
-}
+		if (GetAnalysisSession())
+		{
+			bool bChanged = false;
 
-namespace
-{
-	static void DisableAllTraceChannels()
-	{
-		UE::Trace::EnumerateChannels([](const ANSICHAR* ChannelName, bool bEnabled, void*)
-						{
-							if (bEnabled)
-							{
-								FString ChannelNameFString(ChannelName);
-								UE::Trace::ToggleChannel(ChannelNameFString.GetCharArray().GetData(), false);
-							}
-						}
-						, nullptr);
+			// remove any existing tracks that don't match the current list of object ids
+			for (int TrackIndex = DebugTracks.Num() - 1; TrackIndex >= 0; TrackIndex--)
+			{
+				if (!CandidateIds.Contains(DebugTracks[TrackIndex]->GetObjectId()))
+				{
+					DebugTracks.RemoveAt(TrackIndex);
+				}
+			}
+
+			// add new tracks for current list of object identifiers if they don't already exist
+			for (uint64 CandidateIdentifier : CandidateIds)
+			{
+				const TSharedPtr<RewindDebugger::FRewindDebuggerTrack>* FoundTrack = DebugTracks.FindByPredicate(
+					[CandidateIdentifier](const TSharedPtr<RewindDebugger::FRewindDebuggerTrack>& Track)
+					{
+						return Track->GetObjectId() == CandidateIdentifier;
+					});
+
+				if (!FoundTrack)
+				{
+					DebugTracks.Add(MakeShared<RewindDebugger::FRewindDebuggerObjectTrack>(CandidateIdentifier, SelectedObjectName.Get(), true));
+					bChanged = true;
+				}
+			}
+
+			// update all tracks
+			for (const TSharedPtr<RewindDebugger::FRewindDebuggerTrack>& DebugTrack : DebugTracks)
+			{
+				if (DebugTrack->Update())
+				{
+					bChanged = true;
+				}
+			}
+
+			if (bChanged)
+			{
+				TrackListChangedDelegate.ExecuteIfBound();
+			}
+		}
 	}
 }
 
@@ -437,7 +408,7 @@ void FRewindDebugger::OnConnection()
 	FTraceAuxiliary::OnConnection.RemoveAll(this);
 }
 
-void FRewindDebugger::StartRecording()
+void FRewindDebugger::StartRecording() const
 {
 	if (!CanStartRecording())
 	{
@@ -454,33 +425,33 @@ void FRewindDebugger::OnClearRecording()
 {
 	ClearTrace();
 	RecordingDuration.Set(0);
-	TargetObjectIds.Empty(2);
+	CandidateIds.Empty(2);
 	bTargetActorPositionValid = false;
 
 	IterateExtensions([this](IRewindDebuggerExtension* Extension)
-	{
-		Extension->Clear(this);
-	}
+		{
+			Extension->Clear(this);
+		}
 	);
 }
 
 void FRewindDebugger::OnRecordingStarted()
 {
 	IterateExtensions([this](IRewindDebuggerExtension* Extension)
-	{
-		Extension->RecordingStarted(this);
-	}
+		{
+			Extension->RecordingStarted(this);
+		}
 	);
-	
+
 	UnrealInsightsModule->StartAnalysisForLastLiveSession(5.0);
 }
 
 void FRewindDebugger::OnRecordingStopped()
 {
 	IterateExtensions([this](IRewindDebuggerExtension* Extension)
-	{
-		Extension->RecordingStopped(this);
-	}
+		{
+			Extension->RecordingStopped(this);
+		}
 	);
 }
 
@@ -491,11 +462,11 @@ bool FRewindDebugger::CanOpenTrace() const
 
 
 void FRewindDebugger::OpenTrace(const FString& FilePath)
-{ 
+{
 	ClearTrace();
 
 	bDisplayWorldIdValid = false;
-	
+
 	IUnrealInsightsModule& TraceInsightsModule = FModuleManager::LoadModuleChecked<IUnrealInsightsModule>("TraceInsights");
 	TraceInsightsModule.StartAnalysisForTraceFile(*FilePath);
 
@@ -504,14 +475,14 @@ void FRewindDebugger::OpenTrace(const FString& FilePath)
 
 void FRewindDebugger::OpenTrace()
 {
-	FString FolderPath = "";
-	
+	const FString FolderPath = "";
+
 	TArray<FString> OutOpenFilenames;
 	if (IDesktopPlatform* DesktopPlatform = FDesktopPlatformModule::Get())
 	{
 		FString ExtensionStr;
 		ExtensionStr += TEXT("Unreal Trace|*.utrace|");
-	
+
 		DesktopPlatform->OpenFileDialog(
 			FSlateApplication::Get().FindBestParentWindowHandleForDialogs(nullptr),
 			LOCTEXT("OpenDialogTitle", "Open Rewind Debugger Recording").ToString(),
@@ -552,7 +523,7 @@ void FRewindDebugger::AttachToSession()
 
 		if (!bSuccess)
 		{
-			FMessageDialog::Open(EAppMsgType::Ok, LOCTEXT("FailedToConnectToSessionMessage", "Failed to connect to session"));	
+			FMessageDialog::Open(EAppMsgType::Ok, LOCTEXT("FailedToConnectToSessionMessage", "Failed to connect to session"));
 		}
 	}
 }
@@ -566,22 +537,22 @@ void FRewindDebugger::ClearTrace()
 {
 	StopRecording();
 	RecordingDuration.Set(0);
-	
-	TargetObjectIds.Empty();
+
+	CandidateIds.Empty();
 	CurrentTraceRange.SetLowerBoundValue(0);
 	CurrentTraceRange.SetUpperBoundValue(0);
 	RecordingDuration.Set(0.0);
 	SetCurrentScrubTime(0.0);
 
-	ComponentSelectionChanged(nullptr);
-	
+	TrackSelectionChanged(nullptr);
+
 	// update extensions
 	IterateExtensions([this](IRewindDebuggerExtension* Extension)
 		{
 			Extension->Clear(this);
 		}
 	);
-	
+
 	IUnrealInsightsModule& TraceInsightsModule = FModuleManager::LoadModuleChecked<IUnrealInsightsModule>("TraceInsights");
 	// only way I can find to clear the session is trying to load a name that doesn't exist.
 	TraceInsightsModule.StartAnalysisForTraceFile(TEXT("0"));
@@ -597,32 +568,31 @@ bool FRewindDebugger::CanSaveTrace() const
 
 void FRewindDebugger::SaveTrace(FString FileName)
 {
-
 	if (const TraceServices::IAnalysisSession* Session = GetAnalysisSession())
 	{
 		if (Session->IsAnalysisComplete())
 		{
-			FString SourceFileName = Session->GetName();
+			const FString SourceFileName = Session->GetName();
 
 			FPlatformFileManager& FileManager = FPlatformFileManager::Get();
 			IPlatformFile& PlatformFile = FileManager.GetPlatformFile();
 
 			PlatformFile.CopyFile(*FileName, *SourceFileName);
 		}
-		
+
 	}
 }
 
 void FRewindDebugger::SaveTrace()
 {
-	FString FolderPath = "";
-	
+	const FString FolderPath = "";
+
 	TArray<FString> OutFilenames;
 	if (IDesktopPlatform* DesktopPlatform = FDesktopPlatformModule::Get())
 	{
 		FString ExtensionStr;
 		ExtensionStr += TEXT("Rewind Debugger Recording |*.utrace|");
-	
+
 		DesktopPlatform->SaveFileDialog(
 			FSlateApplication::Get().FindBestParentWindowHandleForDialogs(nullptr),
 			LOCTEXT("SaveDialogTitle", "Save Rewind Debugger Recording").ToString(),
@@ -697,12 +667,12 @@ void FRewindDebugger::Pause()
 
 bool FRewindDebugger::IsPlaying() const
 {
-	return ControlState==EControlState::Play && !bPIESimulating;
+	return ControlState == EControlState::Play && !bPIESimulating;
 }
 
 bool FRewindDebugger::CanPlay() const
 {
-	return ControlState!=EControlState::Play && !bPIESimulating && RecordingDuration.Get() > 0;
+	return ControlState != EControlState::Play && !bPIESimulating && RecordingDuration.Get() > 0;
 }
 
 void FRewindDebugger::Play()
@@ -720,7 +690,7 @@ void FRewindDebugger::Play()
 
 bool FRewindDebugger::CanPlayReverse() const
 {
-	return ControlState!=EControlState::PlayReverse && !bPIESimulating && RecordingDuration.Get() > 0;
+	return ControlState != EControlState::PlayReverse && !bPIESimulating && RecordingDuration.Get() > 0;
 }
 
 void FRewindDebugger::PlayReverse()
@@ -761,7 +731,7 @@ void FRewindDebugger::ScrubToEnd()
 	}
 }
 
-void FRewindDebugger::Step(int frames)
+void FRewindDebugger::Step(const int32 InNumberOfFrames)
 {
 	if (CanScrub())
 	{
@@ -770,21 +740,19 @@ void FRewindDebugger::Step(int frames)
 		if (const TraceServices::IAnalysisSession* Session = GetAnalysisSession())
 		{
 			TraceServices::FAnalysisSessionReadScope SessionReadScope(*Session);
-			UWorld* World = GetWorldToVisualize();
-
 			if (const IGameplayProvider* GameplayProvider = Session->ReadProvider<IGameplayProvider>("GameplayProvider"))
 			{
 				if (const IGameplayProvider::RecordingInfoTimeline* Recording = GameplayProvider->GetRecordingInfo(RecordingIndex))
 				{
-					uint64 EventCount = Recording->GetEventCount();
+					const uint64 EventCount = Recording->GetEventCount();
 
 					if (EventCount > 0)
 					{
-						ScrubTimeInformation.FrameIndex = FMath::Clamp<int64>(ScrubTimeInformation.FrameIndex + frames, 0, (int64)EventCount - 1);
+						ScrubTimeInformation.FrameIndex = FMath::Clamp<int64>(ScrubTimeInformation.FrameIndex + InNumberOfFrames, 0, (int64)EventCount - 1);
 						const FRecordingInfoMessage& Event = Recording->GetEvent(ScrubTimeInformation.FrameIndex);
 
 						SetCurrentScrubTime(Event.ElapsedTime);
-						
+
 						TrackCursorDelegate.ExecuteIfBound(false);
 					}
 				}
@@ -822,7 +790,7 @@ UWorld* FRewindDebugger::GetWorldToVisualize() const
 	UEditorEngine* EditorEngine = Cast<UEditorEngine>(GEngine);
 	if (GIsEditor && EditorEngine != nullptr && World == nullptr)
 	{
-		// lets use PlayWorld during PIE/Simulate and regular world from editor otherwise, to draw debug information
+		// let's use PlayWorld during PIE/Simulate and regular world from editor otherwise, to draw debug information
 		World = EditorEngine->PlayWorld != nullptr ? ToRawPtr(EditorEngine->PlayWorld) : EditorEngine->GetEditorWorldContext().World();
 	}
 
@@ -831,7 +799,7 @@ UWorld* FRewindDebugger::GetWorldToVisualize() const
 
 bool FRewindDebugger::IsRecording() const
 {
-	if (RewindDebugger::FRewindDebuggerRuntime* Runtime = RewindDebugger::FRewindDebuggerRuntime::Instance())
+	if (const RewindDebugger::FRewindDebuggerRuntime* Runtime = RewindDebugger::FRewindDebuggerRuntime::Instance())
 	{
 		return Runtime->IsRecording();
 	}
@@ -840,7 +808,7 @@ bool FRewindDebugger::IsRecording() const
 
 bool FRewindDebugger::IsTraceFileLoaded() const
 {
-	return GetAnalysisSession()!=nullptr && !bPIEStarted;
+	return GetAnalysisSession() != nullptr && !bPIEStarted;
 }
 
 void FRewindDebugger::SetCurrentViewRange(const TRange<double>& Range)
@@ -850,7 +818,7 @@ void FRewindDebugger::SetCurrentViewRange(const TRange<double>& Range)
 	{
 		GetScrubTimeInformation(CurrentViewRange.GetLowerBoundValue(), LowerBoundViewTimeInformation, RecordingIndex, Session);
 		GetScrubTimeInformation(CurrentViewRange.GetUpperBoundValue(), UpperBoundViewTimeInformation, RecordingIndex, Session);
-		
+
 		CurrentTraceRange.SetLowerBoundValue(LowerBoundViewTimeInformation.ProfileTime);
 		CurrentTraceRange.SetUpperBoundValue(UpperBoundViewTimeInformation.ProfileTime);
 	}
@@ -863,20 +831,20 @@ void FRewindDebugger::SetCurrentScrubTime(double Time)
 	if (const TraceServices::IAnalysisSession* Session = GetAnalysisSession())
 	{
 		GetScrubTimeInformation(CurrentScrubTime, ScrubTimeInformation, RecordingIndex, Session);
-		
+
 		TraceTime.Set(ScrubTimeInformation.ProfileTime);
 	}
 }
 
-void FRewindDebugger::GetScrubTimeInformation(double InDebugTime, FScrubTimeInformation & InOutTimeInformation, uint16 InRecordingIndex, const TraceServices::IAnalysisSession* AnalysisSession)
+void FRewindDebugger::GetScrubTimeInformation(double InDebugTime, FScrubTimeInformation& InOutTimeInformation, uint16 InRecordingIndex, const TraceServices::IAnalysisSession* AnalysisSession)
 {
 	const IGameplayProvider* GameplayProvider = AnalysisSession->ReadProvider<IGameplayProvider>("GameplayProvider");
 	const IAnimationProvider* AnimationProvider = AnalysisSession->ReadProvider<IAnimationProvider>("AnimationProvider");
-	
+
 	if (GameplayProvider && AnimationProvider)
 	{
 		TraceServices::FAnalysisSessionReadScope SessionReadScope(*AnalysisSession);
-		
+
 		if (const IGameplayProvider::RecordingInfoTimeline* Recording = GameplayProvider->GetRecordingInfo(InRecordingIndex))
 		{
 			const uint64 EventCount = Recording->GetEventCount();
@@ -887,7 +855,7 @@ void FRewindDebugger::GetScrubTimeInformation(double InDebugTime, FScrubTimeInfo
 				const FRecordingInfoMessage& FirstEvent = Recording->GetEvent(0);
 				const FRecordingInfoMessage& LastEvent = Recording->GetEvent(EventCount - 1);
 
-				// Check if we are outside of the recorded range, and apply the first or last frame
+				// Check if we are outside the recorded range, and apply the first or last frame
 				if (InDebugTime <= FirstEvent.ElapsedTime)
 				{
 					ScrubFrameIndex = FMath::Min<uint64>(1, EventCount - 1);
@@ -901,7 +869,7 @@ void FRewindDebugger::GetScrubTimeInformation(double InDebugTime, FScrubTimeInfo
 				{
 					const FRecordingInfoMessage& ScrubEvent = Recording->GetEvent(ScrubFrameIndex);
 					constexpr float MaxTimeDifferenceInSeconds = 15.0f / 60.0f;
-					
+
 					// Use linear search on smaller time differences
 					if (FMath::Abs(InDebugTime - ScrubEvent.ElapsedTime) <= MaxTimeDifferenceInSeconds)
 					{
@@ -950,8 +918,8 @@ void FRewindDebugger::GetScrubTimeInformation(double InDebugTime, FScrubTimeInfo
 					else
 					{
 						uint64 StartEventIndex = 0;
-						uint64 EndEventIndex = EventCount -1;
-						
+						uint64 EndEventIndex = EventCount - 1;
+
 						while (EndEventIndex - StartEventIndex > 1)
 						{
 							const uint64 MiddleEventIndex = ((StartEventIndex + EndEventIndex) / 2);
@@ -969,24 +937,24 @@ void FRewindDebugger::GetScrubTimeInformation(double InDebugTime, FScrubTimeInfo
 						// Ensure there is not frames between start and end index
 						check(EndEventIndex == StartEventIndex + 1)
 
-						const FRecordingInfoMessage& Event = Recording->GetEvent(StartEventIndex);
+							const FRecordingInfoMessage& Event = Recording->GetEvent(StartEventIndex);
 						const FRecordingInfoMessage& NextEvent = Recording->GetEvent(EndEventIndex);
 
 						// Ensure debug time is between both frames time range
-						check (Event.ElapsedTime <= InDebugTime && NextEvent.ElapsedTime >= InDebugTime)
+						check(Event.ElapsedTime <= InDebugTime && NextEvent.ElapsedTime >= InDebugTime)
 
-						// Choose frame that is nearest to the debug time
-						if (InDebugTime - Event.ElapsedTime < NextEvent.ElapsedTime - InDebugTime)
-						{
-							ScrubFrameIndex = StartEventIndex;
-						}
-						else
-						{
-							ScrubFrameIndex = EndEventIndex;
-						}
+							// Choose frame that is nearest to the debug time
+							if (InDebugTime - Event.ElapsedTime < NextEvent.ElapsedTime - InDebugTime)
+							{
+								ScrubFrameIndex = StartEventIndex;
+							}
+							else
+							{
+								ScrubFrameIndex = EndEventIndex;
+							}
 					}
 				}
-				
+
 				const FRecordingInfoMessage& Event = Recording->GetEvent(ScrubFrameIndex);
 				InOutTimeInformation.FrameIndex = ScrubFrameIndex;
 				InOutTimeInformation.ProfileTime = Event.ProfileTime;
@@ -1005,28 +973,24 @@ const TraceServices::IAnalysisSession* FRewindDebugger::GetAnalysisSession() con
 	return UnrealInsightsModule ? UnrealInsightsModule->GetAnalysisSession().Get() : nullptr;
 }
 
-const FObjectInfo* FRewindDebugger::FindOwningActorInfo(const IGameplayProvider* GameplayProvider, uint64 ObjectId) const
+const FObjectInfo* FRewindDebugger::FindOwningActorInfo(const IGameplayProvider* GameplayProvider, uint64 InObjectId) const
 {
 	const FClassInfo* ActorClassInfo = GameplayProvider->FindClassInfo(*AActor::StaticClass()->GetPathName());
-	
-	while(true)
+
+	while (true)
 	{
-		const FObjectInfo& ObjectInfo = GameplayProvider->GetObjectInfo(ObjectId);
+		const FObjectInfo& ObjectInfo = GameplayProvider->GetObjectInfo(InObjectId);
 		if (GameplayProvider->IsSubClassOf(ObjectInfo.ClassId, ActorClassInfo->Id))
 		{
 			return &ObjectInfo;
 		}
-		else
+
+		if (ObjectInfo.OuterId == 0)
 		{
-			if (ObjectInfo.OuterId != 0)
-			{
-				ObjectId = ObjectInfo.OuterId;
-			}
-			else
-			{
-				return nullptr;
-			}
+			return nullptr;
 		}
+
+		InObjectId = ObjectInfo.OuterId;
 	}
 }
 
@@ -1044,7 +1008,7 @@ void FRewindDebugger::Tick(float DeltaTime)
 	{
 		const IAnimationProvider* AnimationProvider = Session->ReadProvider<IAnimationProvider>("AnimationProvider");
 		const IGameplayProvider* GameplayProvider = Session->ReadProvider<IGameplayProvider>("GameplayProvider");
-		
+
 		if (AnimationProvider && GameplayProvider)
 		{
 			TraceServices::FAnalysisSessionReadScope SessionReadScope(*Session);
@@ -1053,39 +1017,37 @@ void FRewindDebugger::Tick(float DeltaTime)
 			if (IsTraceFileLoaded() && !bDisplayWorldIdValid)
 			{
 				GameplayProvider->EnumerateWorlds([this](const FWorldInfo& WorldInfo)
-           		{
-					if (WorldInfo.Type == FWorldInfo::EType::PIE)
 					{
-						if(WorldInfo.NetMode == FWorldInfo::ENetMode::Client && WorldInfo.PIEInstanceId == 1)
+						if (WorldInfo.Type == FWorldInfo::EType::PIE)
 						{
-							DisplayWorldId = WorldInfo.Id;
-							bDisplayWorldIdValid = true;
+							if (WorldInfo.NetMode == FWorldInfo::ENetMode::Client && WorldInfo.PIEInstanceId == 1)
+							{
+								DisplayWorldId = WorldInfo.Id;
+								bDisplayWorldIdValid = true;
+							}
+							if (WorldInfo.NetMode == FWorldInfo::ENetMode::Standalone && WorldInfo.PIEInstanceId == 0)
+							{
+								DisplayWorldId = WorldInfo.Id;
+								bDisplayWorldIdValid = true;
+							}
 						}
-						if(WorldInfo.NetMode == FWorldInfo::ENetMode::Standalone && WorldInfo.PIEInstanceId == 0)
+						else if (WorldInfo.Type == FWorldInfo::EType::Game)
 						{
 							DisplayWorldId = WorldInfo.Id;
 							bDisplayWorldIdValid = true;
 						}
-					}
-					else if (WorldInfo.Type == FWorldInfo::EType::Game)
-					{
-						DisplayWorldId = WorldInfo.Id;
-						bDisplayWorldIdValid = true;
-					}
-           		});
+					});
 			}
 
-			double RecordingDurationValue = GameplayProvider->GetRecordingDuration();
+			const double RecordingDurationValue = GameplayProvider->GetRecordingDuration();
 			if (IsTraceFileLoaded() && RecordingDurationValue > RecordingDuration.Get())
 			{
 				// while trace file is loading up, force the trace range to update.
 				SetCurrentViewRange(GetCurrentViewRange());
 			}
 			RecordingDuration.Set(RecordingDurationValue);
-			
+
 			RefreshDebugTracks();
-			
-			UWorld* World = GetWorldToVisualize();
 
 			if (bPIESimulating)
 			{
@@ -1099,15 +1061,15 @@ void FRewindDebugger::Tick(float DeltaTime)
 			}
 			else
 			{
-				if (RecordingDuration.Get() > 0 && CurrentScrubTime <= RecordingDuration.Get()) 
+				if (RecordingDuration.Get() > 0 && CurrentScrubTime <= RecordingDuration.Get())
 				{
 					if (ControlState == EControlState::Play || ControlState == EControlState::PlayReverse)
 					{
-						float PlaybackRate = URewindDebuggerSettings::Get().PlaybackRate;
+						const float PlaybackRate = URewindDebuggerSettings::Get().PlaybackRate;
 						TRACE_CPUPROFILER_EVENT_SCOPE(FRewindDebugger::Tick_UpdatePlayback);
-						float Rate = PlaybackRate * (ControlState == EControlState::Play ? 1 : -1);
+						const float Rate = PlaybackRate * (ControlState == EControlState::Play ? 1 : -1);
 						SetCurrentScrubTime(FMath::Clamp(CurrentScrubTime + Rate * DeltaTime, 0.0f, RecordingDuration.Get()));
-						TrackCursorDelegate.ExecuteIfBound(Rate<0);
+						TrackCursorDelegate.ExecuteIfBound(Rate < 0);
 
 						if (CurrentScrubTime == 0 || CurrentScrubTime == RecordingDuration.Get())
 						{
@@ -1117,7 +1079,7 @@ void FRewindDebugger::Tick(float DeltaTime)
 					}
 
 					SetCurrentScrubTime(CurrentScrubTime);// update trace time
-					
+
 					const double CurrentTraceTime = TraceTime.Get();
 					if (CurrentTraceTime != PreviousTraceTime)
 					{
@@ -1129,51 +1091,51 @@ void FRewindDebugger::Tick(float DeltaTime)
 						if (FrameProvider.GetFrameFromTime(ETraceFrameType::TraceFrameType_Game, CurrentTraceTime, Frame))
 						{
 							bool bNewActor = false;
-							if (!TargetObjectIds.Contains(TargetActorIdForMesh))
+							if (!CandidateIds.Contains(TargetActorIdForMesh))
 							{
 								AnimationProvider->EnumerateSkeletalMeshPoseTimelines([this, &bNewActor, GameplayProvider](uint64 ObjectId, const IAnimationProvider::SkeletalMeshPoseTimeline& TimelineData)
-								{
-									// until we have actor transforms traced out, the first (from a non-server) skeletal mesh component transform on the target actor be used as as the actor position
-
-									if (const FWorldInfo* WorldInfo = GameplayProvider->FindWorldInfoFromObject(ObjectId))
 									{
-										if (WorldInfo->NetMode != FWorldInfo::ENetMode::DedicatedServer)
+										// until we have actor transforms traced out, the first (from a non-server) skeletal mesh component transform on the target actor be used as the actor position
+
+										if (const FWorldInfo* WorldInfo = GameplayProvider->FindWorldInfoFromObject(ObjectId))
 										{
-											if (const FObjectInfo* ActorInfo = FindOwningActorInfo(GameplayProvider, ObjectId))
+											if (WorldInfo->NetMode != FWorldInfo::ENetMode::DedicatedServer)
 											{
-												if (TargetObjectIds.Contains(ActorInfo->Id))
+												if (const FObjectInfo* ActorInfo = FindOwningActorInfo(GameplayProvider, ObjectId))
 												{
-													bNewActor = true;
-													TargetActorIdForMesh = ActorInfo->Id;
-													TargetActorMeshId = ObjectId;
+													const uint64 ActorId = ActorInfo->Id;
+													if (CandidateIds.Contains(ActorId))
+													{
+														bNewActor = true;
+														TargetActorIdForMesh = ActorId;
+														TargetActorMeshId = ObjectId;
+													}
 												}
 											}
 										}
-									}
-								});
+									});
 							}
-						
-						
+
 							AnimationProvider->ReadSkeletalMeshPoseTimeline(TargetActorMeshId, [this, &Frame, bNewActor](const IAnimationProvider::SkeletalMeshPoseTimeline& TimelineData, bool bHasCurves)
-							{
-								const FSkeletalMeshPoseMessage * PoseMessage = nullptr;
+								{
+									const FSkeletalMeshPoseMessage* PoseMessage = nullptr;
 
-								// Get last pose in frame
-								TimelineData.EnumerateEvents(Frame.StartTime, Frame.EndTime,
-									[&PoseMessage](double InStartTime, double InEndTime, uint32 InDepth, const FSkeletalMeshPoseMessage& InPoseMessage)
-									{
-										PoseMessage = &InPoseMessage;
-										return TraceServices::EEventEnumerate::Continue;
-									});
+									// Get last pose in frame
+									TimelineData.EnumerateEvents(Frame.StartTime, Frame.EndTime,
+										[&PoseMessage](double InStartTime, double InEndTime, uint32 InDepth, const FSkeletalMeshPoseMessage& InPoseMessage)
+										{
+											PoseMessage = &InPoseMessage;
+											return TraceServices::EEventEnumerate::Continue;
+										});
 
-								// Update position based on pose
-								if (PoseMessage)
-								{
-									// mark the target position as invalid for a frame when the actor changes, so it will be treated as a teleport by the camera system
-									bTargetActorPositionValid = !bNewActor;
-									TargetActorPosition = PoseMessage->ComponentToWorld.GetTranslation();
-								}
-							});
+									// Update position based on pose
+									if (PoseMessage)
+									{
+										// mark the target position as invalid for a frame when the actor changes, so it will be treated as a teleport by the camera system
+										bTargetActorPositionValid = !bNewActor;
+										TargetActorPosition = PoseMessage->ComponentToWorld.GetTranslation();
+									}
+								});
 						}
 					}
 				}
@@ -1193,21 +1155,21 @@ void FRewindDebugger::Tick(float DeltaTime)
 void FRewindDebugger::OpenDetailsPanel()
 {
 	bIsDetailsPanelOpen = true;
-	ComponentSelectionChanged(SelectedTrack);
+	TrackSelectionChanged(SelectedTrack);
 }
 
-void FRewindDebugger::ComponentSelectionChanged(TSharedPtr<RewindDebugger::FRewindDebuggerTrack> SelectedObject)
+void FRewindDebugger::TrackSelectionChanged(const TSharedPtr<RewindDebugger::FRewindDebuggerTrack> InSelectedTrack)
 {
-	SelectedTrack = SelectedObject;
+	SelectedTrack = InSelectedTrack;
 
 	if (bIsDetailsPanelOpen)
 	{
-		FLevelEditorModule& LevelEditorModule = FModuleManager::GetModuleChecked<FLevelEditorModule>("LevelEditor");
-		TSharedPtr<FTabManager> LevelEditorTabManager = LevelEditorModule.GetLevelEditorTabManager();
-		
-		 // if we now have no selection, don't force the tab into focus - this happens when tracks disappear and can cause PIE to lose focus while playing
-		bool bInvokeAsInactive = !SelectedTrack.IsValid();
-		TSharedPtr<SDockTab> DetailsTab = LevelEditorTabManager->TryInvokeTab(FRewindDebuggerModule::DetailsTabName, bInvokeAsInactive);
+		const FLevelEditorModule& LevelEditorModule = FModuleManager::GetModuleChecked<FLevelEditorModule>("LevelEditor");
+		const TSharedPtr<FTabManager> LevelEditorTabManager = LevelEditorModule.GetLevelEditorTabManager();
+
+		// if we now have no selection, don't force the tab into focus - this happens when tracks disappear and can cause PIE to lose focus while playing
+		const bool bInvokeAsInactive = !SelectedTrack.IsValid();
+		const TSharedPtr<SDockTab> DetailsTab = LevelEditorTabManager->TryInvokeTab(FRewindDebuggerModule::DetailsTabName, bInvokeAsInactive);
 
 		if (DetailsTab.IsValid())
 		{
@@ -1243,27 +1205,26 @@ void FRewindDebugger::UpdateDetailsPanel(TSharedRef<SDockTab> DetailsTab)
 	}
 }
 
-void FRewindDebugger::RegisterComponentContextMenu()
+void FRewindDebugger::RegisterTrackContextMenu()
 {
-	UToolMenu* Menu = UToolMenus::Get()->FindMenu("RewindDebugger.ComponentContextMenu");
-	
+	UToolMenu* Menu = UToolMenus::Get()->FindMenu(FRewindDebuggerModule::TrackContextMenuName);
+
 	FToolMenuSection& Section = Menu->FindOrAddSection("SelectedTrack");
-	
-	FToolMenuEntry& Entry = Section.AddDynamicEntry(NAME_None, FNewToolMenuSectionDelegate::CreateLambda([](FToolMenuSection& InSection)
-	{
-		const UComponentContextMenuContext* Context = InSection.FindContext<UComponentContextMenuContext>();
 
-		if (Context && Context->SelectedTrack.IsValid())
+	Section.AddDynamicEntry(NAME_None, FNewToolMenuSectionDelegate::CreateLambda([](FToolMenuSection& InSection)
 		{
-			Context->SelectedTrack->BuildContextMenu(InSection);
-		}
-	}));
+			const URewindDebuggerTrackContextMenuContext* Context = InSection.FindContext<URewindDebuggerTrackContextMenuContext>();
+			if (Context && Context->SelectedTrack.IsValid())
+			{
+				Context->SelectedTrack->BuildContextMenu(InSection);
+			}
+		}));
 }
 
 void FRewindDebugger::MakeOtherWorldsMenu(UToolMenu* Menu)
 {
-	FRewindDebugger* RewindDebugger = FRewindDebugger::Instance();
-	
+	const FRewindDebugger* RewindDebugger = FRewindDebugger::Instance();
+
 	FToolMenuSection& Section = Menu->AddSection("Other Worlds", LOCTEXT("Other Worlds", "Other Worlds"));
 
 	if (const TraceServices::IAnalysisSession* Session = RewindDebugger->GetAnalysisSession())
@@ -1272,76 +1233,76 @@ void FRewindDebugger::MakeOtherWorldsMenu(UToolMenu* Menu)
 		const IGameplayProvider* GameplayProvider = Session->ReadProvider<IGameplayProvider>("GameplayProvider");
 
 		GameplayProvider->EnumerateWorlds([GameplayProvider, &Section](const FWorldInfo& WorldInfo)
-		{
-			const FObjectInfo* ObjectInfo = GameplayProvider->FindObjectInfo(WorldInfo.Id);
-			FString Name = ObjectInfo->Name;
-
-			if(WorldInfo.NetMode == FWorldInfo::ENetMode::DedicatedServer)
-			{
-				return;
-			}
-			else if (WorldInfo.Type == FWorldInfo::EType::Game || WorldInfo.Type == FWorldInfo::EType::PIE)
-			{
-				return;
-			}
-			else
 			{
-				if (WorldInfo.Type == FWorldInfo::EType::Editor)
-				{
-					Name = Name + " (Editor)";
-				}
-				else if (WorldInfo.Type == FWorldInfo::EType::Inactive)
-				{
-					Name = Name + " (Editor)";
-				}
-				else if (WorldInfo.Type == FWorldInfo::EType::EditorPreview)
+				const FObjectInfo* ObjectInfo = GameplayProvider->FindObjectInfo(WorldInfo.Id);
+				FString Name = ObjectInfo->Name;
+
+				if (WorldInfo.NetMode == FWorldInfo::ENetMode::DedicatedServer)
 				{
-					Name = Name + " (Editor Preview)";
+					return;
 				}
-				else if (WorldInfo.Type == FWorldInfo::EType::GamePreview)
+				else if (WorldInfo.Type == FWorldInfo::EType::Game || WorldInfo.Type == FWorldInfo::EType::PIE)
 				{
-					Name = Name + " (Game Preview)";
+					return;
 				}
-				else if (WorldInfo.Type == FWorldInfo::EType::GameRPC)
+				else
 				{
-					Name = Name + " (Game RPC)";
+					if (WorldInfo.Type == FWorldInfo::EType::Editor)
+					{
+						Name = Name + " (Editor)";
+					}
+					else if (WorldInfo.Type == FWorldInfo::EType::Inactive)
+					{
+						Name = Name + " (Editor)";
+					}
+					else if (WorldInfo.Type == FWorldInfo::EType::EditorPreview)
+					{
+						Name = Name + " (Editor Preview)";
+					}
+					else if (WorldInfo.Type == FWorldInfo::EType::GamePreview)
+					{
+						Name = Name + " (Game Preview)";
+					}
+					else if (WorldInfo.Type == FWorldInfo::EType::GameRPC)
+					{
+						Name = Name + " (Game RPC)";
+					}
 				}
-			}
-		
-			Section.AddMenuEntry(FName(ObjectInfo->Name,WorldInfo.Id),
-								FText::FromString(Name),
-								FText(),
-								FSlateIcon(),
-								FUIAction( FExecuteAction::CreateLambda([World = WorldInfo.Id]()
-								{
-									FRewindDebugger::Instance()->SetDisplayWorld(World);
-								}),
-								FCanExecuteAction(),
-								FIsActionChecked::CreateLambda([World = WorldInfo.Id]()
-								{
-									return FRewindDebugger::Instance()->DisplayWorldId == World;
-								})),
-								EUserInterfaceActionType::Check
-							);
-		
-		});
+
+				Section.AddMenuEntry(FName(ObjectInfo->Name, WorldInfo.Id),
+					FText::FromString(Name),
+					FText(),
+					FSlateIcon(),
+					FUIAction(FExecuteAction::CreateLambda([World = WorldInfo.Id]()
+						{
+							FRewindDebugger::Instance()->SetDisplayWorld(World);
+						}),
+						FCanExecuteAction(),
+						FIsActionChecked::CreateLambda([World = WorldInfo.Id]()
+							{
+								return FRewindDebugger::Instance()->DisplayWorldId == World;
+							})),
+					EUserInterfaceActionType::Check
+				);
+
+			});
 	}
 }
 
 void FRewindDebugger::SetDisplayWorld(uint64 WorldId)
 {
 	DisplayWorldId = WorldId;
-	
+
 	IterateExtensions([this](IRewindDebuggerExtension* Extension)
-	{
-		Extension->Clear(this);
-		Extension->Update(0.0,this);
-	});
+		{
+			Extension->Clear(this);
+			Extension->Update(0.0, this);
+		});
 }
 void FRewindDebugger::MakeWorldsMenu(UToolMenu* Menu)
 {
-	FRewindDebugger* RewindDebugger = FRewindDebugger::Instance();
-	
+	const FRewindDebugger* RewindDebugger = FRewindDebugger::Instance();
+
 	FToolMenuSection& ServerWorldsSection = Menu->AddSection("Server Worlds", LOCTEXT("Server", "Server"));
 	FToolMenuSection& GameWorldsSection = Menu->AddSection("Game Worlds", LOCTEXT("Game Worlds", "Game Worlds"));
 	FToolMenuSection& OtherWorldsSection = Menu->AddSection("Other Worlds", LOCTEXT("Other Worlds", "Other Worlds"));
@@ -1351,203 +1312,203 @@ void FRewindDebugger::MakeWorldsMenu(UToolMenu* Menu)
 		LOCTEXT("Other Worlds Tooltip", "Additional worlds such as  Editor Preview worlds"),
 		FNewToolMenuChoice(
 			FNewToolMenuDelegate::CreateStatic(FRewindDebugger::MakeOtherWorldsMenu)
-			));
-	
+		));
+
 	if (const TraceServices::IAnalysisSession* Session = RewindDebugger->GetAnalysisSession())
 	{
 		TraceServices::FAnalysisSessionReadScope SessionReadScope(*Session);
 		const IGameplayProvider* GameplayProvider = Session->ReadProvider<IGameplayProvider>("GameplayProvider");
 
-		GameplayProvider->EnumerateWorlds([GameplayProvider,&GameWorldsSection, &OtherWorldsSection, &ServerWorldsSection](const FWorldInfo& WorldInfo)
-		{
-			const FObjectInfo* ObjectInfo = GameplayProvider->FindObjectInfo(WorldInfo.Id);
-			FString Name = ObjectInfo->Name;
+		GameplayProvider->EnumerateWorlds([GameplayProvider, &GameWorldsSection, &OtherWorldsSection, &ServerWorldsSection](const FWorldInfo& WorldInfo)
+			{
+				const FObjectInfo* ObjectInfo = GameplayProvider->FindObjectInfo(WorldInfo.Id);
+				FString Name = ObjectInfo->Name;
 
-			FToolMenuSection* Section = &OtherWorldsSection;
+				FToolMenuSection* Section = &OtherWorldsSection;
 
-			if(WorldInfo.NetMode == FWorldInfo::ENetMode::DedicatedServer)
-			{
-				Section = &ServerWorldsSection;
-				Name = Name + " (Server)";
-			}
-			else if (WorldInfo.Type == FWorldInfo::EType::Game || WorldInfo.Type == FWorldInfo::EType::PIE)
-			{
-				Section = &GameWorldsSection;
-				if(WorldInfo.NetMode == FWorldInfo::ENetMode::Client && WorldInfo.PIEInstanceId >= 0)
+				if (WorldInfo.NetMode == FWorldInfo::ENetMode::DedicatedServer)
 				{
-					Name = Name + " (Client " + FString::FromInt(WorldInfo.PIEInstanceId) + ")";
+					Section = &ServerWorldsSection;
+					Name = Name + " (Server)";
 				}
-				if(WorldInfo.NetMode == FWorldInfo::ENetMode::Standalone && WorldInfo.PIEInstanceId >= 0)
+				else if (WorldInfo.Type == FWorldInfo::EType::Game || WorldInfo.Type == FWorldInfo::EType::PIE)
 				{
-					Name = Name + " (Standalone " + FString::FromInt(WorldInfo.PIEInstanceId) + ")";
+					Section = &GameWorldsSection;
+					if (WorldInfo.NetMode == FWorldInfo::ENetMode::Client && WorldInfo.PIEInstanceId >= 0)
+					{
+						Name = Name + " (Client " + FString::FromInt(WorldInfo.PIEInstanceId) + ")";
+					}
+					if (WorldInfo.NetMode == FWorldInfo::ENetMode::Standalone && WorldInfo.PIEInstanceId >= 0)
+					{
+						Name = Name + " (Standalone " + FString::FromInt(WorldInfo.PIEInstanceId) + ")";
+					}
 				}
-			}
-			else
-			{
-				return;
-			}
-			
-			Section->AddMenuEntry(FName(ObjectInfo->Name,WorldInfo.Id),
-								FText::FromString(Name),
-								FText(),
-								FSlateIcon(),
-								FUIAction( FExecuteAction::CreateLambda([World = WorldInfo.Id]()
-								{
-									FRewindDebugger::Instance()->SetDisplayWorld(World);
-								}),
-								FCanExecuteAction(),
-								FIsActionChecked::CreateLambda([World = WorldInfo.Id]()
-								{
-									return FRewindDebugger::Instance()->DisplayWorldId == World;
-								})),
-								EUserInterfaceActionType::Check
-							);
-		});
+				else
+				{
+					return;
+				}
+
+				Section->AddMenuEntry(FName(ObjectInfo->Name, WorldInfo.Id),
+					FText::FromString(Name),
+					FText(),
+					FSlateIcon(),
+					FUIAction(FExecuteAction::CreateLambda([World = WorldInfo.Id]()
+						{
+							FRewindDebugger::Instance()->SetDisplayWorld(World);
+						}),
+						FCanExecuteAction(),
+						FIsActionChecked::CreateLambda([World = WorldInfo.Id]()
+							{
+								return FRewindDebugger::Instance()->DisplayWorldId == World;
+							})),
+					EUserInterfaceActionType::Check
+				);
+			});
 	}
 }
 
 void FRewindDebugger::RegisterToolBar()
 {
 	UToolMenu* Menu = UToolMenus::Get()->RegisterMenu("RewindDebugger.ToolBar", NAME_None, EMultiBoxType::ToolBar);
-	
+
 	FToolMenuSection& Section = Menu->FindOrAddSection("VCRControls");
 	const FRewindDebuggerCommands& Commands = FRewindDebuggerCommands::Get();
 	Section.AddEntry(FToolMenuEntry::InitToolBarButton(
-			Commands.FirstFrame,
-			FText(),
-			TAttribute<FText>(),
-			FSlateIcon("RewindDebuggerStyle", "RewindDebugger.FirstFrame.small")));
-	
+		Commands.FirstFrame,
+		FText(),
+		TAttribute<FText>(),
+		FSlateIcon("RewindDebuggerStyle", "RewindDebugger.FirstFrame.small")));
+
 	Section.AddEntry(FToolMenuEntry::InitToolBarButton(
-			Commands.PreviousFrame,
-			FText(),
-			TAttribute<FText>(),
-			FSlateIcon("RewindDebuggerStyle", "RewindDebugger.PreviousFrame.small")));
-			
+		Commands.PreviousFrame,
+		FText(),
+		TAttribute<FText>(),
+		FSlateIcon("RewindDebuggerStyle", "RewindDebugger.PreviousFrame.small")));
+
 	Section.AddEntry(FToolMenuEntry::InitToolBarButton(
-				Commands.ReversePlay,
-				FText(),
-				TAttribute<FText>(),
-				FSlateIcon("RewindDebuggerStyle", "RewindDebugger.ReversePlay.small")));
-	
+		Commands.ReversePlay,
+		FText(),
+		TAttribute<FText>(),
+		FSlateIcon("RewindDebuggerStyle", "RewindDebugger.ReversePlay.small")));
+
 	Section.AddEntry(FToolMenuEntry::InitToolBarButton(
-				Commands.Pause,
-				FText(),
-				FText::Format(LOCTEXT("PauseButtonTooltip", "{0} ({1})"), Commands.Pause->GetDescription(), Commands.PauseOrPlay->GetInputText()),
-				FSlateIcon("RewindDebuggerStyle", "RewindDebugger.Pause.small")));
-	
+		Commands.Pause,
+		FText(),
+		FText::Format(LOCTEXT("PauseButtonTooltip", "{0} ({1})"), Commands.Pause->GetDescription(), Commands.PauseOrPlay->GetInputText()),
+		FSlateIcon("RewindDebuggerStyle", "RewindDebugger.Pause.small")));
+
 	Section.AddEntry(FToolMenuEntry::InitToolBarButton(
-				Commands.Play,
-				FText(),
-				FText::Format(LOCTEXT("PlayButtonTooltip", "{0} ({1})"), Commands.Play->GetDescription(), Commands.PauseOrPlay->GetInputText()),
-				FSlateIcon("RewindDebuggerStyle", "RewindDebugger.Play.small")));
+		Commands.Play,
+		FText(),
+		FText::Format(LOCTEXT("PlayButtonTooltip", "{0} ({1})"), Commands.Play->GetDescription(), Commands.PauseOrPlay->GetInputText()),
+		FSlateIcon("RewindDebuggerStyle", "RewindDebugger.Play.small")));
 
 	Section.AddEntry(
-    		FToolMenuEntry::InitComboButton(
-    			"PlaybackRate",
-    			FToolUIActionChoice(),
-    			FNewToolMenuChoice(
-    				FNewToolMenuDelegate::CreateLambda([](UToolMenu* InNewToolMenu)
-    				{
-    					FToolMenuSection& Section = InNewToolMenu->AddSection("PlaybackSpeed", LOCTEXT("Playback Speed", "Playback Speed"));
-    					
+		FToolMenuEntry::InitComboButton(
+			"PlaybackRate",
+			FToolUIActionChoice(),
+			FNewToolMenuChoice(
+				FNewToolMenuDelegate::CreateLambda([](UToolMenu* InNewToolMenu)
+					{
+						FToolMenuSection& Section = InNewToolMenu->AddSection("PlaybackSpeed", LOCTEXT("Playback Speed", "Playback Speed"));
+
 						Section.AddEntry(
 							FToolMenuEntry::InitMenuEntry(
-								"001",LOCTEXT("0.1","0.1"), LOCTEXT("Set playback speed to 0.1", "Set playback speed to 0.1"), FSlateIcon(),
+								"001", LOCTEXT("0.1", "0.1"), LOCTEXT("Set playback speed to 0.1", "Set playback speed to 0.1"), FSlateIcon(),
 								FUIAction(
-								FExecuteAction::CreateLambda([]()
-									{ 
-										URewindDebuggerSettings::Get().PlaybackRate = 0.1;
-									}),
+									FExecuteAction::CreateLambda([]()
+										{
+											URewindDebuggerSettings::Get().PlaybackRate = 0.1;
+										}),
 									FCanExecuteAction(),
 									FIsActionChecked::CreateLambda([]
-									{
-										return FMath::IsNearlyEqual(URewindDebuggerSettings::Get().PlaybackRate, 0.1);
-									})
-									)
-									, EUserInterfaceActionType::RadioButton
+										{
+											return FMath::IsNearlyEqual(URewindDebuggerSettings::Get().PlaybackRate, 0.1);
+										})
 								)
-    					);
+								, EUserInterfaceActionType::RadioButton
+							)
+						);
 						Section.AddEntry(
 							FToolMenuEntry::InitMenuEntry(
-								"025",LOCTEXT("0.25","0.25"), LOCTEXT("Set playback speed to 0.25", "Set playback speed to 0.25"), FSlateIcon(),
+								"025", LOCTEXT("0.25", "0.25"), LOCTEXT("Set playback speed to 0.25", "Set playback speed to 0.25"), FSlateIcon(),
 								FUIAction(
-								FExecuteAction::CreateLambda([]()
-									{ 
-										URewindDebuggerSettings::Get().PlaybackRate = 0.25;
-									}),
+									FExecuteAction::CreateLambda([]()
+										{
+											URewindDebuggerSettings::Get().PlaybackRate = 0.25;
+										}),
 									FCanExecuteAction(),
 									FIsActionChecked::CreateLambda([]
-									{
-										return FMath::IsNearlyEqual(URewindDebuggerSettings::Get().PlaybackRate, 0.25);
-									})
-									)
-									, EUserInterfaceActionType::RadioButton
+										{
+											return FMath::IsNearlyEqual(URewindDebuggerSettings::Get().PlaybackRate, 0.25);
+										})
+								)
+								, EUserInterfaceActionType::RadioButton
 							)
-    					);
-    					Section.AddEntry(
+						);
+						Section.AddEntry(
 							FToolMenuEntry::InitMenuEntry(
-								"05",LOCTEXT("0.5","0.5"), LOCTEXT("Set playback speed to 0.5", "Set playback speed to 0.5"), FSlateIcon(),
+								"05", LOCTEXT("0.5", "0.5"), LOCTEXT("Set playback speed to 0.5", "Set playback speed to 0.5"), FSlateIcon(),
 								FUIAction(
-								FExecuteAction::CreateLambda([]()
-									{ 
-										URewindDebuggerSettings::Get().PlaybackRate = 0.5;
-									}),
+									FExecuteAction::CreateLambda([]()
+										{
+											URewindDebuggerSettings::Get().PlaybackRate = 0.5;
+										}),
 									FCanExecuteAction(),
 									FIsActionChecked::CreateLambda([]
-									{
-										return FMath::IsNearlyEqual(URewindDebuggerSettings::Get().PlaybackRate, 0.5);
-									})
-									)
-									, EUserInterfaceActionType::RadioButton
+										{
+											return FMath::IsNearlyEqual(URewindDebuggerSettings::Get().PlaybackRate, 0.5);
+										})
+								)
+								, EUserInterfaceActionType::RadioButton
 							)
 						);
 
 						Section.AddEntry(
 							FToolMenuEntry::InitMenuEntry(
-								"1",LOCTEXT("1","1"), LOCTEXT("Set playback speed to 1", "Set playback speed to 1"), FSlateIcon(),
+								"1", LOCTEXT("1", "1"), LOCTEXT("Set playback speed to 1", "Set playback speed to 1"), FSlateIcon(),
 								FUIAction(
-								FExecuteAction::CreateLambda([]()
-									{ 
-										URewindDebuggerSettings::Get().PlaybackRate = 1;
-									}),
+									FExecuteAction::CreateLambda([]()
+										{
+											URewindDebuggerSettings::Get().PlaybackRate = 1;
+										}),
 									FCanExecuteAction(),
 									FIsActionChecked::CreateLambda([]
-									{
-										return FMath::IsNearlyEqual(URewindDebuggerSettings::Get().PlaybackRate, 1);
-									})
-									)
-									, EUserInterfaceActionType::RadioButton
+										{
+											return FMath::IsNearlyEqual(URewindDebuggerSettings::Get().PlaybackRate, 1);
+										})
+								)
+								, EUserInterfaceActionType::RadioButton
 							)
 						);
-    					
-    					Section.AddEntry(
+
+						Section.AddEntry(
 							FToolMenuEntry::InitMenuEntry(
-								"2",LOCTEXT("2","2"), LOCTEXT("Set playback speed to 2", "Set playback speed to 2"), FSlateIcon(),
+								"2", LOCTEXT("2", "2"), LOCTEXT("Set playback speed to 2", "Set playback speed to 2"), FSlateIcon(),
 								FUIAction(
-								FExecuteAction::CreateLambda([]()
-									{ 
-										URewindDebuggerSettings::Get().PlaybackRate = 2;
-									}),
+									FExecuteAction::CreateLambda([]()
+										{
+											URewindDebuggerSettings::Get().PlaybackRate = 2;
+										}),
 									FCanExecuteAction(),
 									FIsActionChecked::CreateLambda([]
-									{
-										return FMath::IsNearlyEqual(URewindDebuggerSettings::Get().PlaybackRate, 2);
-									})
-									)
-									, EUserInterfaceActionType::RadioButton
+										{
+											return FMath::IsNearlyEqual(URewindDebuggerSettings::Get().PlaybackRate, 2);
+										})
+								)
+								, EUserInterfaceActionType::RadioButton
 							)
 						);
 
 						Section.AddEntry(
 							FToolMenuEntry::InitWidget(
-								"EditInSequencerMenu", 
+								"EditInSequencerMenu",
 								SNew(SNumericEntryBox<float>)
-									.Value_Lambda([]()
+								.Value_Lambda([]()
 									{
 										return URewindDebuggerSettings::Get().PlaybackRate;
 									})
-									.OnValueChanged_Lambda([](float Value)
+								.OnValueChanged_Lambda([](float Value)
 									{
 										URewindDebuggerSettings::Get().PlaybackRate = Value;
 									}),
@@ -1555,122 +1516,121 @@ void FRewindDebugger::RegisterToolBar()
 								true, false, true
 							)
 						);
-    				})
-    			),
-				FText(),
-    			LOCTEXT("PlaybackRate_Tooltip", "Playback Options"),
-				FSlateIcon(FAppStyle::GetAppStyleSetName(), "Sequencer.PlaybackOptions")
-    		)
-    	);
+					})
+			),
+			FText(),
+			LOCTEXT("PlaybackRate_Tooltip", "Playback Options"),
+			FSlateIcon(FAppStyle::GetAppStyleSetName(), "Sequencer.PlaybackOptions")
+		)
+	);
 
 	Section.AddEntry(FToolMenuEntry::InitToolBarButton(
-				Commands.NextFrame,
-				FText(),
-				TAttribute<FText>(),
-				FSlateIcon("RewindDebuggerStyle", "RewindDebugger.NextFrame.small")));
-	
+		Commands.NextFrame,
+		FText(),
+		TAttribute<FText>(),
+		FSlateIcon("RewindDebuggerStyle", "RewindDebugger.NextFrame.small")));
+
 	Section.AddEntry(FToolMenuEntry::InitToolBarButton(
-				Commands.LastFrame,
-				FText(),
-				TAttribute<FText>(),
-				FSlateIcon("RewindDebuggerStyle", "RewindDebugger.LastFrame.small")));
+		Commands.LastFrame,
+		FText(),
+		TAttribute<FText>(),
+		FSlateIcon("RewindDebuggerStyle", "RewindDebugger.LastFrame.small")));
 
 	Section.AddEntry(FToolMenuEntry::InitToolBarButton(
-				Commands.StartRecording,
-				FText(),
-				TAttribute<FText>(),
-				FSlateIcon("RewindDebuggerStyle", "RewindDebugger.StartRecording.small")));
-				
+		Commands.StartRecording,
+		FText(),
+		TAttribute<FText>(),
+		FSlateIcon("RewindDebuggerStyle", "RewindDebugger.StartRecording.small")));
+
 
 	Section.AddEntry(FToolMenuEntry::InitToolBarButton(
-				Commands.StopRecording,
-				FText(),
-				TAttribute<FText>(),
-				FSlateIcon("RewindDebuggerStyle", "RewindDebugger.StopRecording.small")));
+		Commands.StopRecording,
+		FText(),
+		TAttribute<FText>(),
+		FSlateIcon("RewindDebuggerStyle", "RewindDebugger.StopRecording.small")));
 
 	Section.AddSeparator(NAME_None);
 
-
 	Section.AddEntry(FToolMenuEntry::InitToolBarButton(
-				Commands.AttachToSession,
-				FText(),
-				TAttribute<FText>(),
-				FSlateIcon("RewindDebuggerStyle", "RewindDebugger.ConnectToSession")));
-				
+		Commands.AttachToSession,
+		FText(),
+		TAttribute<FText>(),
+		FSlateIcon("RewindDebuggerStyle", "RewindDebugger.ConnectToSession")));
+
 	Section.AddEntry(FToolMenuEntry::InitToolBarButton(
-				Commands.OpenTrace,
-				FText(),
-				TAttribute<FText>(),
-				 FSlateIcon(FAppStyle::GetAppStyleSetName(), "Icons.FolderOpen")));
-	
+		Commands.OpenTrace,
+		FText(),
+		TAttribute<FText>(),
+		FSlateIcon(FAppStyle::GetAppStyleSetName(), "Icons.FolderOpen")));
+
 	Section.AddEntry(FToolMenuEntry::InitToolBarButton(
-    			Commands.SaveTrace,
-				FText(),
-    			TAttribute<FText>(),
-    			 FSlateIcon(FAppStyle::GetAppStyleSetName(), "Icons.Save")));
-				 
+		Commands.SaveTrace,
+		FText(),
+		TAttribute<FText>(),
+		FSlateIcon(FAppStyle::GetAppStyleSetName(), "Icons.Save")));
+
 	Section.AddEntry(FToolMenuEntry::InitToolBarButton(
-				Commands.ClearTrace,
-				FText(),
-				TAttribute<FText>(),
-				 FSlateIcon(FAppStyle::GetAppStyleSetName(), "Icons.Delete")));
+		Commands.ClearTrace,
+		FText(),
+		TAttribute<FText>(),
+		FSlateIcon(FAppStyle::GetAppStyleSetName(), "Icons.Delete")));
 
-	
 	Section.AddEntry(FToolMenuEntry::InitToolBarButton(
-				Commands.AutoEject,
-				FText(),
-				TAttribute<FText>(),
-				FSlateIcon("RewindDebuggerStyle", "RewindDebugger.AutoEject")));
-					Section.AddSeparator(NAME_None);
-                				
+		Commands.AutoEject,
+		FText(),
+		TAttribute<FText>(),
+		FSlateIcon("RewindDebuggerStyle", "RewindDebugger.AutoEject")));
+
+	Section.AddSeparator(NAME_None);
+
 	Section.AddEntry(FToolMenuEntry::InitToolBarButton(
-				Commands.AutoRecord,
-				FText(),
-				TAttribute<FText>(),
-				FSlateIcon("RewindDebuggerStyle", "RewindDebugger.AutoRecord")));
-	
+		Commands.AutoRecord,
+		FText(),
+		TAttribute<FText>(),
+		FSlateIcon("RewindDebuggerStyle", "RewindDebugger.AutoRecord")));
+
 	Section.AddSeparator("NAME_None");
 
 	Section.AddEntry(FToolMenuEntry::InitComboButton(
 		"Display World",
 		FUIAction(
 			FExecuteAction(),
-			FCanExecuteAction::CreateLambda([](){ return FRewindDebugger::Instance()->IsTraceFileLoaded(); })
-			),
+			FCanExecuteAction::CreateLambda([]() { return FRewindDebugger::Instance()->IsTraceFileLoaded(); })
+		),
 		FNewToolMenuDelegate::CreateStatic(&FRewindDebugger::MakeWorldsMenu),
 		LOCTEXT("Display World", "Display World"),
 		LOCTEXT("Display World Tooltip", "When loading trace files, only the objects (Such as Skeletal Meshes) from the world selected here will be spawned for preview")
-		));
-	
+	));
+
 	Menu->SetStyleSet(&FAppStyle::Get());
 	Menu->StyleName = "PaletteToolBar";
 }
 
 
-void FRewindDebugger::ComponentDoubleClicked(TSharedPtr<RewindDebugger::FRewindDebuggerTrack> SelectedObject)
+void FRewindDebugger::TrackDoubleClicked(TSharedPtr<RewindDebugger::FRewindDebuggerTrack> InSelectedTrack)
 {
-	if (!SelectedObject.IsValid())
+	if (!InSelectedTrack.IsValid())
 	{
 		return;
 	}
-	
-	SelectedTrack = SelectedObject;
+
+	SelectedTrack = InSelectedTrack;
 	SelectedTrack->HandleDoubleClick();
 }
 
-TSharedPtr<SWidget> FRewindDebugger::BuildComponentContextMenu() const
+TSharedPtr<SWidget> FRewindDebugger::BuildTrackContextMenu() const
 {
-	UComponentContextMenuContext* MenuContext = NewObject<UComponentContextMenuContext>();
-	MenuContext->SelectedObject = GetSelectedComponent();
+	URewindDebuggerTrackContextMenuContext* MenuContext = NewObject<URewindDebuggerTrackContextMenuContext>();
+	MenuContext->SelectedObject = GetSelectedObject();
 	MenuContext->SelectedTrack = SelectedTrack;
-	
+
 	if (SelectedTrack.IsValid())
 	{
 		// build a list of class hierarchy names to make it easier for extensions to enable menu entries by type
 		if (const TraceServices::IAnalysisSession* Session = GetAnalysisSession())
 		{
 			TraceServices::FAnalysisSessionReadScope SessionReadScope(*Session);
-	
+
 			const IGameplayProvider* GameplayProvider = Session->ReadProvider<IGameplayProvider>("GameplayProvider");
 			const FObjectInfo& ObjectInfo = GameplayProvider->GetObjectInfo(SelectedTrack->GetObjectId());
 			uint64 ClassId = ObjectInfo.ClassId;
@@ -1683,27 +1643,25 @@ TSharedPtr<SWidget> FRewindDebugger::BuildComponentContextMenu() const
 		}
 	}
 
-	return UToolMenus::Get()->GenerateWidget("RewindDebugger.ComponentContextMenu", FToolMenuContext(MenuContext));
- }
+	return UToolMenus::Get()->GenerateWidget(FRewindDebuggerModule::TrackContextMenuName, FToolMenuContext(MenuContext));
+}
 
 
-TSharedPtr<FDebugObjectInfo> FRewindDebugger::GetSelectedComponent() const
+TSharedPtr<FDebugObjectInfo> FRewindDebugger::GetSelectedObject() const
 {
-	if (!SelectedComponent.IsValid())
-	{
-		SelectedComponent = MakeShared<FDebugObjectInfo>(0, "");
-	}
-	
 	if (SelectedTrack.IsValid())
 	{
-		SelectedComponent->ObjectId = SelectedTrack->GetObjectId();
-		SelectedComponent->ObjectName = SelectedTrack->GetDisplayName().ToString();
-		return SelectedComponent;
-	}
-	else
-	{
-		return TSharedPtr<FDebugObjectInfo>();
+		if (!SelectedObject.IsValid())
+		{
+			SelectedObject = MakeShared<FDebugObjectInfo>();
+		}
+
+		SelectedObject->ObjectId = SelectedTrack->GetObjectId();
+		SelectedObject->ObjectName = SelectedTrack->GetDisplayName().ToString();
+		return SelectedObject;
 	}
+
+	return TSharedPtr<FDebugObjectInfo>();
 }
 
 TSharedPtr<RewindDebugger::FRewindDebuggerTrack> FRewindDebugger::GetSelectedTrack() const
@@ -1711,29 +1669,28 @@ TSharedPtr<RewindDebugger::FRewindDebuggerTrack> FRewindDebugger::GetSelectedTra
 	return SelectedTrack;
 }
 
-// build a component tree that's compatible with the public api from 5.0 for GetDebugComponents.
-void FRewindDebugger::RefreshDebugComponents(TArray<TSharedPtr<RewindDebugger::FRewindDebuggerTrack>>& InTracks, TArray<TSharedPtr<FDebugObjectInfo>>& OutComponents)
+// build a tree that's compatible with the public api from 5.0 for GetDebuggedObjects.
+void FRewindDebugger::RefreshDebuggedObjects(TArray<TSharedPtr<RewindDebugger::FRewindDebuggerTrack>>& InTracks, TArray<TSharedPtr<FDebugObjectInfo>>& OutObjects)
 {
-	OutComponents.SetNum(0, EAllowShrinking::No);
-	for(TSharedPtr<RewindDebugger::FRewindDebuggerTrack>& Track : InTracks)
+	OutObjects.SetNum(0, EAllowShrinking::No);
+	for (const TSharedPtr<RewindDebugger::FRewindDebuggerTrack>& Track : InTracks)
 	{
-		int Index = OutComponents.Num();
-		OutComponents.Add(MakeShared<FDebugObjectInfo>(Track->GetObjectId(), Track->GetDisplayName().ToString()));
+		const int Index = OutObjects.Add(MakeShared<FDebugObjectInfo>(Track->GetObjectId(), Track->GetDisplayName().ToString()));
 		TArray<TSharedPtr<RewindDebugger::FRewindDebuggerTrack>> TrackChildren;
-		Track->IterateSubTracks([&TrackChildren](TSharedPtr<RewindDebugger::FRewindDebuggerTrack> Child) { TrackChildren.Add(Child); });
-		RefreshDebugComponents(TrackChildren, OutComponents[Index]->Children);
+		Track->IterateSubTracks([&TrackChildren](const TSharedPtr<RewindDebugger::FRewindDebuggerTrack>& Child) { TrackChildren.Add(Child); });
+		RefreshDebuggedObjects(TrackChildren, OutObjects[Index]->Children);
 	}
 }
 
-TArray<TSharedPtr<FDebugObjectInfo>>& FRewindDebugger::GetDebugComponents()
+TArray<TSharedPtr<FDebugObjectInfo>>& FRewindDebugger::GetDebuggedObjects()
 {
-	RefreshDebugComponents(DebugTracks, DebugComponents);
-	return DebugComponents;
+	RefreshDebuggedObjects(DebugTracks, DebuggedObjects);
+	return DebuggedObjects;
 }
 
-bool FRewindDebugger::IsContainedByDebugComponent(uint64 ObjectId) const
+bool FRewindDebugger::IsObjectCurrentlyDebugged(uint64 ObjectId) const
 {
-	for(auto Track : DebugTracks)
+	for (const TSharedPtr<RewindDebugger::FRewindDebuggerTrack>& Track : DebugTracks)
 	{
 		if (Track->GetObjectId() == ObjectId)
 		{
@@ -1741,14 +1698,14 @@ bool FRewindDebugger::IsContainedByDebugComponent(uint64 ObjectId) const
 		}
 
 		bool Found = false;
-		Track->IterateSubTracks( [ObjectId, &Found] (TSharedPtr<RewindDebugger::FRewindDebuggerTrack> Child)
-		{
-			if (Child->GetObjectId() == ObjectId)
+		Track->IterateSubTracks([ObjectId, &Found](const TSharedPtr<RewindDebugger::FRewindDebuggerTrack>& Child)
 			{
-				Found = true;
-				// Todo: want to stop iteration here
-			}
-		});
+				if (Child->GetObjectId() == ObjectId)
+				{
+					Found = true;
+					// @Todo STDBG: want to stop iteration here
+				}
+			});
 
 		if (Found)
 		{
diff --git a/Engine/Plugins/Animation/GameplayInsights/Source/RewindDebugger/Private/RewindDebugger.h b/Engine/Plugins/Animation/GameplayInsights/Source/RewindDebugger/Private/RewindDebugger.h
--- a/Engine/Plugins/Animation/GameplayInsights/Source/RewindDebugger/Private/RewindDebugger.h
+++ b/Engine/Plugins/Animation/GameplayInsights/Source/RewindDebugger/Private/RewindDebugger.h
@@ -1,105 +1,142 @@
 // Copyright Epic Games, Inc. All Rights Reserved.
 #pragma once
 
+#if UE_ENABLE_INCLUDE_ORDER_DEPRECATED_IN_5_7
 #include "CoreMinimal.h"
-#include "IRewindDebugger.h"
-#include "RewindDebuggerTrack.h"
+#include "UObject/WeakObjectPtr.h"
+#endif // UE_ENABLE_INCLUDE_ORDER_DEPRECATED_IN_5_7
+
 #include "BindableProperty.h"
-#include "IRewindDebuggerTrackCreator.h"
 #include "Containers/Ticker.h"
-#include "UObject/WeakObjectPtr.h"
+#include "IRewindDebugger.h"
+#include "IRewindDebuggerTrackCreator.h"
+#include "RewindDebuggerTrack.h"
 
+namespace TraceServices { class IAnalysisSession; }
 class FMenuBuilder;
-class USkeletalMeshComponent;
-
-namespace TraceServices
-{
-	class IAnalysisSession;
-}
-
-class SWidget;
+class IGameplayProvider;
 class SDockTab;
-
+class SWidget;
+class USkeletalMeshComponent;
+struct FGameplayObjectIdentifier;
 struct FObjectInfo;
-class IGameplayProvider;
-
-// Singleton class that handles the logic for the Rewind Debugger
-// handles:
-//  Playback/Scrubbing state
-//  Start/Stop recording
-//  Keeping track of the current Debug Target actor, and outputing a list of it's Components for the UI
 
+/**
+ * Singleton class that handles the logic for the Rewind Debugger:
+ * - Playback/Scrubbing state
+ * - Start/Stop recording
+ * - Keeping track of the current Debug Target object, and outputting a list of its sub objects/elements for the UI
+ */
 class FRewindDebugger : public IRewindDebugger
 {
 public:
 	FRewindDebugger();
 	virtual ~FRewindDebugger();
 
-	// IRewindDebugger interface
-	virtual double CurrentTraceTime() const override { return TraceTime.Get(); }
-	virtual double GetScrubTime() const override { return CurrentScrubTime; }
-	virtual const TRange<double>& GetCurrentTraceRange() const override { return CurrentTraceRange; }
-	virtual const TRange<double>& GetCurrentViewRange() const override { return CurrentViewRange; }
+	//~ IRewindDebugger interface
+	virtual double CurrentTraceTime() const override
+	{
+		return TraceTime.Get();
+	}
+
+	virtual double GetScrubTime() const override
+	{
+		return CurrentScrubTime;
+	}
+
+	virtual const TRange<double>& GetCurrentTraceRange() const override
+	{
+		return CurrentTraceRange;
+	}
+
+	virtual const TRange<double>& GetCurrentViewRange() const override
+	{
+		return CurrentViewRange;
+	}
+
 	virtual const TraceServices::IAnalysisSession* GetAnalysisSession() const override;
 	virtual uint64 GetTargetActorId() const override;
 	virtual bool GetTargetActorPosition(FVector& OutPosition) const override;
 	virtual UWorld* GetWorldToVisualize() const override;
 	virtual bool IsRecording() const override;
-	virtual bool IsPIESimulating() const override { return bPIESimulating; }
+	virtual bool IsPIESimulating() const override
+	{
+		return bPIESimulating;
+	}
+
 	virtual bool IsTraceFileLoaded() const override;
-	virtual double GetRecordingDuration() const override { return RecordingDuration.Get(); }
-	virtual TSharedPtr<FDebugObjectInfo> GetSelectedComponent() const override;
+	virtual double GetRecordingDuration() const override
+	{
+		return RecordingDuration.Get();
+	}
+
+	virtual TSharedPtr<FDebugObjectInfo> GetSelectedObject() const override;
 	virtual TSharedPtr<RewindDebugger::FRewindDebuggerTrack> GetSelectedTrack() const override;
-	virtual TArray<TSharedPtr<FDebugObjectInfo>>& GetDebugComponents() override;
-	virtual bool IsContainedByDebugComponent(uint64 ObjectId) const override;
-	virtual bool ShouldDisplayWorld(uint64 WorldId) override { return DisplayWorldId == WorldId; } 
+	virtual TArray<TSharedPtr<FDebugObjectInfo>>& GetDebuggedObjects() override;
+	virtual bool IsObjectCurrentlyDebugged(uint64 ObjectId) const override;
+	virtual bool ShouldDisplayWorld(uint64 WorldId) override
+	{
+		return DisplayWorldId == WorldId;
+	}
 
 	void OnConnection();
 
-	void GetTargetObjectIds(TArray<uint64>& OutActorIds) const;
+	void GetTargetObjectIds(TArray<uint64>& OutTargetObjectIds) const;
 
-	TArray<TSharedPtr<RewindDebugger::FRewindDebuggerTrack>>& GetDebugTracks() { return DebugTracks; }
+	TArray<TSharedPtr<RewindDebugger::FRewindDebuggerTrack>>& GetTracks()
+	{
+		return DebugTracks;
+	}
 
-	// create singleton instance
+	/** create singleton instance */
 	static void Initialize();
 
-	// destroy singleton instance
+	/** destroy singleton instance */
 	static void Shutdown();
 
-	// get singleton instance
-	static FRewindDebugger* Instance() { return static_cast<FRewindDebugger*>(InternalInstance); }
+	/** get singleton instance */
+	static FRewindDebugger* Instance()
+	{
+		return static_cast<FRewindDebugger*>(InternalInstance);
+	}
 
-	// Start a new Recording:  Start tracing Object + Animation data, increment the current recording index, and reset the recording elapsed time to 0
-	void StartRecording();
+	/** Start a new Recording:  Start tracing Object + Animation data, increment the current recording index, and reset the recording elapsed time to 0 */
+	void StartRecording() const;
 	void OnClearRecording();
 	void OnRecordingStarted();
 	void OnRecordingStopped();
-	bool CanStartRecording() const { return !IsRecording() && bPIESimulating; }
-	
+	bool CanStartRecording() const
+	{
+		return !IsRecording() && bPIESimulating;
+	}
+
 	bool CanOpenTrace() const;
-   	void OpenTrace(const FString& FilePath);
-   	void OpenTrace();
-	
-   	void AttachToSession();
+	void OpenTrace(const FString& FilePath);
+	void OpenTrace();
+
+	void AttachToSession();
 
 	bool CanClearTrace() const;
 	void ClearTrace();
-	
+
 	bool CanSaveTrace() const;
 	void SaveTrace(FString FileName);
 	void SaveTrace();
-        	
+
 	bool ShouldAutoRecordOnPIE() const;
-	void SetShouldAutoRecordOnPIE(bool value);
-	
+	void SetShouldAutoRecordOnPIE(bool InValue);
+
 	bool ShouldAutoEject() const;
-	void SetShouldAutoEject(bool value);
+	void SetShouldAutoEject(bool InValue);
 
-	// Stop recording: Stop tracing Object + Animation Data.
+	/** Stop recording: Stop tracing Object + Animation Data. */
 	void StopRecording();
-	bool CanStopRecording() const { return IsRecording(); }
+	bool CanStopRecording() const
+	{
+		return IsRecording();
+	}
 
-	// VCR controls
+	//~ VCR controls
 
 	bool CanPause() const;
 	void Pause();
@@ -113,53 +150,69 @@ public:
 
 	void ScrubToStart();
 	void ScrubToEnd();
-	void Step(int32 frames);
+	void Step(int32 InNumberOfFrames);
 	void StepForward();
 	void StepBackward();
 
 	bool CanScrub() const;
 	void ScrubToTime(double ScrubTime, bool bIsScrubbing);
 
-	// Tick function: While recording, update recording duration.  While paused, and we have recorded data, update skinned mesh poses for the current frame, and handle playback.
+	/** Tick function: While recording, update recording duration.  While paused, and we have recorded data, update skinned mesh poses for the current frame, and handle playback. */
 	void Tick(float DeltaTime);
 
-	// update the list of tracks for the currently selected debug target
+	/** update the list of tracks for the currently selected debug target */
 	void RefreshDebugTracks();
-	
+
 	void SetCurrentViewRange(const TRange<double>& Range);
 
-	DECLARE_DELEGATE(FOnComponentListChanged)
-	void OnComponentListChanged(const FOnComponentListChanged& ComponentListChangedCallback);
-	
-	void ComponentDoubleClicked(TSharedPtr<RewindDebugger::FRewindDebuggerTrack> SelectedObject);
-	void ComponentSelectionChanged(TSharedPtr<RewindDebugger::FRewindDebuggerTrack> SelectedObject);
-	TSharedPtr<SWidget> BuildComponentContextMenu() const;
-	
+	DECLARE_DELEGATE(FOnTrackListChanged)
+	void SetTrackListChangedDelegate(const FOnTrackListChanged& InTrackListChangedDelegate);
+
+	void TrackDoubleClicked(TSharedPtr<RewindDebugger::FRewindDebuggerTrack> InSelectedTrack);
+	void TrackSelectionChanged(TSharedPtr<RewindDebugger::FRewindDebuggerTrack> InSelectedTrack);
+	TSharedPtr<SWidget> BuildTrackContextMenu() const;
+
 	void UpdateDetailsPanel(TSharedRef<SDockTab> DetailsTab);
-	static void RegisterComponentContextMenu();
+	static void RegisterTrackContextMenu();
 	static void MakeOtherWorldsMenu(class UToolMenu* Menu);
 	void SetDisplayWorld(uint64 WorldId);
 	static void MakeWorldsMenu(class UToolMenu* Menu);
 	static void RegisterToolBar();
-	
-	DECLARE_DELEGATE_OneParam( FOnTrackCursor, bool)
-	void OnTrackCursor(const FOnTrackCursor& TrackCursorCallback);
 
-	TBindableProperty<double>* GetTraceTimeProperty() { return &TraceTime; }
-	TBindableProperty<double>* GetRecordingDurationProperty() { return &RecordingDuration; }
-	TBindableProperty<FString, BindingType_Out>* GetDebugTargetActorProperty() { return &DebugTargetActor; }
+	DECLARE_DELEGATE_OneParam(FOnTrackCursor, bool)
+	void SetTrackCursorDelegate(const FOnTrackCursor& InTrackCursorDelegate);
+
+	TBindableProperty<double>* GetTraceTimeProperty()
+	{
+		return &TraceTime;
+	}
+
+	TBindableProperty<double>* GetRecordingDurationProperty()
+	{
+		return &RecordingDuration;
+	}
+
+	TBindableProperty<FString, BindingType_Out>* GetDebugTargetActorProperty()
+	{
+		return &SelectedObjectName;
+	}
 
 	virtual void OpenDetailsPanel() override;
-	void SetIsDetailsPanelOpen(bool bIsOpen) { bIsDetailsPanelOpen = bIsOpen; }
-	bool IsDetailsPanelOpen(bool bIsOpen) { return bIsDetailsPanelOpen; }
-	
+	void SetIsDetailsPanelOpen(bool bIsOpen)
+	{
+		bIsDetailsPanelOpen = bIsOpen;
+	}
+
 	virtual const FObjectInfo* FindOwningActorInfo(const IGameplayProvider* GameplayProvider, uint64 ObjectId) const override;
 
-	TArrayView<RewindDebugger::FRewindDebuggerTrackType> GetTrackTypes() { return TrackTypes; };
+	TArrayView<RewindDebugger::FRewindDebuggerTrackType> GetTrackTypes()
+	{
+		return TrackTypes;
+	};
 
 private:
-	void RefreshDebugComponents(TArray<TSharedPtr<RewindDebugger::FRewindDebuggerTrack>>& InTracks, TArray<TSharedPtr<FDebugObjectInfo>>& OutComponents);
-	
+	static void RefreshDebuggedObjects(TArray<TSharedPtr<RewindDebugger::FRewindDebuggerTrack>>& InTracks, TArray<TSharedPtr<FDebugObjectInfo>>& OutObjects);
+
 	void OnPIEStarted(bool bSimulating);
 	void OnPIEPaused(bool bSimulating);
 	void OnPIEResumed(bool bSimulating);
@@ -170,7 +223,7 @@ private:
 
 	TBindableProperty<double> TraceTime;
 	TBindableProperty<double> RecordingDuration;
-	TBindableProperty<FString, BindingType_Out> DebugTargetActor;
+	TBindableProperty<FString, BindingType_Out> SelectedObjectName;
 
 	enum class EControlState : int8
 	{
@@ -181,43 +234,45 @@ private:
 
 	EControlState ControlState = EControlState::Pause;
 
-	FOnComponentListChanged ComponentListChangedDelegate;
+	FOnTrackListChanged TrackListChangedDelegate;
 	FOnTrackCursor TrackCursorDelegate;
 
 	bool bQueueStartRecording = false;
 	bool bTraceJustConnected = false;
 	bool bPIEStarted = false;
 	bool bPIESimulating = false;
-	
+
 	bool bRecording = false;
 
 	double PreviousTraceTime = -1;
 	double CurrentScrubTime = 0;
-	TRange<double> CurrentViewRange {0, 10};
-	TRange<double> CurrentTraceRange {0, 0};
+	TRange<double> CurrentViewRange{ 0, 10 };
+	TRange<double> CurrentTraceRange{ 0, 0 };
 	uint16 RecordingIndex = 0;
 
 	struct FScrubTimeInformation
 	{
-		double ProfileTime = 0;  // Profiling/Tracing time
-		int64 FrameIndex = 0;    // Scrub Frame Index
+		/** Profiling/Tracing time */
+		double ProfileTime = 0;
+		/** Scrub Frame Index */
+		int64 FrameIndex = 0;
 	};
-	
+
 	FScrubTimeInformation ScrubTimeInformation;
 	FScrubTimeInformation LowerBoundViewTimeInformation;
 	FScrubTimeInformation UpperBoundViewTimeInformation;
 
-	static void GetScrubTimeInformation(double InDebugTime, FScrubTimeInformation & InOutTimeInformation, uint16 InRecordingIndex, const TraceServices::IAnalysisSession* AnalysisSession);
-	
-	TArray<TSharedPtr<FDebugObjectInfo>> DebugComponents;
-	mutable TSharedPtr<FDebugObjectInfo> SelectedComponent;
-	
+	static void GetScrubTimeInformation(double InDebugTime, FScrubTimeInformation& InOutTimeInformation, uint16 InRecordingIndex, const TraceServices::IAnalysisSession* AnalysisSession);
+
+	TArray<TSharedPtr<FDebugObjectInfo>> DebuggedObjects;
+	mutable TSharedPtr<FDebugObjectInfo> SelectedObject;
+
 	TArray<TSharedPtr<RewindDebugger::FRewindDebuggerTrack>> DebugTracks;
 	TSharedPtr<RewindDebugger::FRewindDebuggerTrack> SelectedTrack;
 
-	TArray<uint64> TargetObjectIds;
+	TArray<uint64> CandidateIds;
 
-	mutable class IUnrealInsightsModule *UnrealInsightsModule = 0;
+	mutable class IUnrealInsightsModule* UnrealInsightsModule = nullptr;
 	FTSTicker::FDelegateHandle TickerHandle;
 
 	bool bTargetActorPositionValid = false;
diff --git a/Engine/Plugins/Animation/GameplayInsights/Source/RewindDebugger/Private/RewindDebuggerCamera.cpp b/Engine/Plugins/Animation/GameplayInsights/Source/RewindDebugger/Private/RewindDebuggerCamera.cpp
--- a/Engine/Plugins/Animation/GameplayInsights/Source/RewindDebugger/Private/RewindDebuggerCamera.cpp
+++ b/Engine/Plugins/Animation/GameplayInsights/Source/RewindDebugger/Private/RewindDebuggerCamera.cpp
@@ -8,6 +8,7 @@
 #include "Insights/IUnrealInsightsModule.h"
 #include "IRewindDebugger.h"
 #include "LevelEditor.h"
+#include "RewindDebuggerModule.h"
 #include "Modules/ModuleManager.h"
 #include "SLevelViewport.h"
 #include "ToolMenu.h"
@@ -18,12 +19,13 @@
 
 FRewindDebuggerCamera::FRewindDebuggerCamera()
 	: bLastPositionValid(false)
+	, LastPosition()
 {
 }
 
 void FRewindDebuggerCamera::Initialize()
 {
-	UToolMenu* Menu = UToolMenus::Get()->FindMenu("RewindDebugger.MainMenu");
+	UToolMenu* Menu = UToolMenus::Get()->FindMenu(FRewindDebuggerModule::MainMenuName);
 
 	Menu->AddSection("Camera Mode", LOCTEXT("Camera Mode","Camera Mode"));
 	Menu->AddMenuEntry("Camera Mode",
diff --git a/Engine/Plugins/Animation/GameplayInsights/Source/RewindDebugger/Private/RewindDebuggerModule.cpp b/Engine/Plugins/Animation/GameplayInsights/Source/RewindDebugger/Private/RewindDebuggerModule.cpp
--- a/Engine/Plugins/Animation/GameplayInsights/Source/RewindDebugger/Private/RewindDebuggerModule.cpp
+++ b/Engine/Plugins/Animation/GameplayInsights/Source/RewindDebugger/Private/RewindDebuggerModule.cpp
@@ -1,32 +1,31 @@
 // Copyright Epic Games, Inc. All Rights Reserved.
 
 #include "RewindDebuggerModule.h"
-
-
+#include "AnimInstanceHelpers.h"
+#include "BlueprintDoubleClickHandler.h"
 #include "Engine/Selection.h"
 #include "Features/IModularFeatures.h"
 #include "Framework/Docking/LayoutExtender.h"
 #include "IAnimationBlueprintEditorModule.h"
+#include "Kismet2/DebuggerCommands.h"
+#include "LevelEditor.h"
 #include "Modules/ModuleManager.h"
+#include "PropertyTraceMenu.h"
 #include "RewindDebugger.h"
-#include "RewindDebuggerStyle.h"
 #include "RewindDebuggerCommands.h"
+#include "RewindDebuggerStyle.h"
 #include "SRewindDebugger.h"
-#include "SRewindDebuggerDetails.h"
+#include "ToolMenus.h"
 #include "Widgets/Docking/SDockTab.h"
 #include "WorkspaceMenuStructure.h"
 #include "WorkspaceMenuStructureModule.h"
-#include "AnimInstanceHelpers.h"
-#include "BlueprintDoubleClickHandler.h"
-#include "PropertyTraceMenu.h"
-#include "ToolMenus.h"
-#include "LevelEditor.h"
-#include "Kismet2/DebuggerCommands.h"
 
 #define LOCTEXT_NAMESPACE "RewindDebuggerModule"
 
 const FName FRewindDebuggerModule::MainTabName("RewindDebugger2");
 const FName FRewindDebuggerModule::DetailsTabName("RewindDebuggerDetails2");
+const FName FRewindDebuggerModule::MainMenuName("RewindDebugger.MainMenu");
+const FName FRewindDebuggerModule::TrackContextMenuName("RewindDebugger.TrackContextMenu");
 
 TSharedRef<SDockTab> FRewindDebuggerModule::SpawnRewindDebuggerDetailsTab(const FSpawnTabArgs& SpawnTabArgs)
 {
@@ -199,20 +198,20 @@ TSharedRef<SDockTab> FRewindDebuggerModule::SpawnRewindDebuggerTab(const FSpawnT
 	}
 	
 	RewindDebuggerWidget = SNew(SRewindDebugger, CommandList.ToSharedRef(), MajorTab, SpawnTabArgs.GetOwnerWindow())
-								.DebugTargetActor({ DebuggerInstance->GetDebugTargetActorProperty(), URewindDebuggerSettings::Get().DebugTargetActor})
+								.DebuggedObjectName({ DebuggerInstance->GetDebugTargetActorProperty(), URewindDebuggerSettings::Get().DebugTargetActor})
 								.RecordingDuration(DebuggerInstance->GetRecordingDurationProperty())
-								.DebugComponents(&DebuggerInstance->GetDebugTracks())
+								.Tracks(&DebuggerInstance->GetTracks())
 								.TraceTime(DebuggerInstance->GetTraceTimeProperty())
 								.OnScrubPositionChanged_Raw(DebuggerInstance,&FRewindDebugger::ScrubToTime)
 								.OnViewRangeChanged_Raw(DebuggerInstance,&FRewindDebugger::SetCurrentViewRange)
-								.OnComponentDoubleClicked_Raw(DebuggerInstance, &FRewindDebugger::ComponentDoubleClicked)
-								.OnComponentSelectionChanged_Raw(DebuggerInstance, &FRewindDebugger::ComponentSelectionChanged)
-								.BuildComponentContextMenu_Raw(DebuggerInstance, &FRewindDebugger::BuildComponentContextMenu)
+								.OnTrackDoubleClicked_Raw(DebuggerInstance, &FRewindDebugger::TrackDoubleClicked)
+								.OnTrackSelectionChanged_Raw(DebuggerInstance, &FRewindDebugger::TrackSelectionChanged)
+								.BuildTrackContextMenu_Raw(DebuggerInstance, &FRewindDebugger::BuildTrackContextMenu)
 								.TrackTypes_Lambda([]() { return FRewindDebugger::Instance()->GetTrackTypes(); })
 								.ScrubTime_Lambda([]() { return FRewindDebugger::Instance()->GetScrubTime(); });
 
-	DebuggerInstance->OnTrackCursor(FRewindDebugger::FOnTrackCursor::CreateSP(RewindDebuggerWidget.Get(), &SRewindDebugger::TrackCursor));
-	DebuggerInstance->OnComponentListChanged(FRewindDebugger::FOnComponentListChanged::CreateSP(RewindDebuggerWidget.Get(), &SRewindDebugger::RefreshDebugComponents));
+	DebuggerInstance->SetTrackCursorDelegate(FRewindDebugger::FOnTrackCursor::CreateSP(RewindDebuggerWidget.Get(), &SRewindDebugger::TrackCursor));
+	DebuggerInstance->SetTrackListChangedDelegate(FRewindDebugger::FOnTrackListChanged::CreateSP(RewindDebuggerWidget.Get(), &SRewindDebugger::RefreshTracks));
 
 	MajorTab->SetContent(RewindDebuggerWidget.ToSharedRef());
 
@@ -224,8 +223,8 @@ static FBlueprintDoubleClickHandler BlueprintDoubleClickHandler;
 
 void FRewindDebuggerModule::StartupModule()
 {
-	UToolMenus::Get()->RegisterMenu("RewindDebugger.MainMenu");
-	UToolMenus::Get()->RegisterMenu("RewindDebugger.ComponentContextMenu");
+	UToolMenus::Get()->RegisterMenu(FRewindDebuggerModule::MainMenuName);
+	UToolMenus::Get()->RegisterMenu(FRewindDebuggerModule::TrackContextMenuName);
 
 	FRewindDebuggerStyle::Initialize();
 	FRewindDebuggerCommands::Register();
@@ -245,7 +244,7 @@ void FRewindDebuggerModule::StartupModule()
 				.SetGroup(WorkspaceMenu::GetMenuStructure().GetDeveloperToolsDebugCategory());
 
 			LevelEditorTabManager->RegisterTabSpawner(
-				DetailsTabName, FOnSpawnTab::CreateRaw(this, &FRewindDebuggerModule::SpawnRewindDebuggerDetailsTab))
+				DetailsTabName, FOnSpawnTab::CreateStatic(&FRewindDebuggerModule::SpawnRewindDebuggerDetailsTab))
 				.SetDisplayName(LOCTEXT("DetailsTabTitle", "Rewind Debugger Details"))
 				.SetIcon(FSlateIcon("RewindDebuggerStyle", "RewindDebugger.RewindDetailsIcon"))
 				.SetTooltipText(LOCTEXT("DetailsWindowTooltipText", "Opens Rewind Debugger Details Window."))
@@ -270,7 +269,7 @@ void FRewindDebuggerModule::StartupModule()
 
 	FPropertyTraceMenu::Register();
 	FAnimInstanceMenu::Register();
-	FRewindDebugger::RegisterComponentContextMenu();
+	FRewindDebugger::RegisterTrackContextMenu();
 	FRewindDebugger::RegisterToolBar();
 }
 
diff --git a/Engine/Plugins/Animation/GameplayInsights/Source/RewindDebugger/Private/RewindDebuggerModule.h b/Engine/Plugins/Animation/GameplayInsights/Source/RewindDebugger/Private/RewindDebuggerModule.h
--- a/Engine/Plugins/Animation/GameplayInsights/Source/RewindDebugger/Private/RewindDebuggerModule.h
+++ b/Engine/Plugins/Animation/GameplayInsights/Source/RewindDebugger/Private/RewindDebuggerModule.h
@@ -2,7 +2,6 @@
 
 #pragma once
 
-#include "CoreMinimal.h"
 #include "Modules/ModuleInterface.h"
 #include "Framework/Docking/TabManager.h"
 #include "RewindDebuggerCamera.h"
@@ -19,10 +18,12 @@ public:
 	virtual void ShutdownModule() override;
 
 	TSharedRef<SDockTab> SpawnRewindDebuggerTab(const FSpawnTabArgs& SpawnTabArgs);
-	TSharedRef<SDockTab> SpawnRewindDebuggerDetailsTab(const FSpawnTabArgs& SpawnTabArgs);
+	static TSharedRef<SDockTab> SpawnRewindDebuggerDetailsTab(const FSpawnTabArgs& SpawnTabArgs);
 
 	static const FName MainTabName;
 	static const FName DetailsTabName;
+	static const FName MainMenuName;
+	static const FName TrackContextMenuName;
 
 private:
 	TSharedPtr<SRewindDebugger> RewindDebuggerWidget;
diff --git a/Engine/Plugins/Animation/GameplayInsights/Source/RewindDebugger/Private/RewindDebuggerObjectTrack.cpp b/Engine/Plugins/Animation/GameplayInsights/Source/RewindDebugger/Private/RewindDebuggerObjectTrack.cpp
--- a/Engine/Plugins/Animation/GameplayInsights/Source/RewindDebugger/Private/RewindDebuggerObjectTrack.cpp
+++ b/Engine/Plugins/Animation/GameplayInsights/Source/RewindDebugger/Private/RewindDebuggerObjectTrack.cpp
@@ -1,23 +1,22 @@
 // Copyright Epic Games, Inc. All Rights Reserved.
 
 #include "RewindDebuggerObjectTrack.h"
-#include "RewindDebuggerFallbackTrack.h"
-#include "IRewindDebugger.h"
-#include "IGameplayProvider.h"
 #include "IAnimationProvider.h"
-#include "RewindDebuggerViewCreators.h"
+#include "IGameplayProvider.h"
+#include "IRewindDebugger.h"
+#include "IRewindDebuggerDoubleClickHandler.h"
+#include "ObjectTrace.h"
+#include "RewindDebuggerFallbackTrack.h"
 #include "RewindDebuggerTrackCreators.h"
+#include "RewindDebuggerViewCreators.h"
 #include "SEventTimelineView.h"
 #include "Styling/SlateIconFinder.h"
-#include "ObjectTrace.h"
-#include "IRewindDebuggerDoubleClickHandler.h"
-
 
 #define LOCTEXT_NAMESPACE "RewindDebuggerObjectTrack"
 
 namespace RewindDebugger
 {
-	
+
 // check if an object is or is a subclass of a type by name, based on Insights traced type info
 static bool IsTargetType(uint64 ObjectId, FName TargetTypeName, const TraceServices::IAnalysisSession& Session)
 {
@@ -39,7 +38,7 @@ static bool IsTargetType(uint64 ObjectId, FName TargetTypeName, const TraceServi
 
 	return false;
 }
-	
+
 FRewindDebuggerObjectTrack::FRewindDebuggerObjectTrack(uint64 InObjectId, const FString& InObjectName, bool bInAddController)
 	: ObjectName(InObjectName)
 	, ObjectId(InObjectId)
@@ -48,33 +47,33 @@ FRewindDebuggerObjectTrack::FRewindDebuggerObjectTrack(uint64 InObjectId, const
 	, bIconSearched(false)
 {
 	ExistenceRange = MakeShared<SEventTimelineView::FTimelineEventData>();
-	
+
 	IRewindDebugger* RewindDebugger = IRewindDebugger::Instance();
 	const TraceServices::IAnalysisSession* Session = RewindDebugger->GetAnalysisSession();
 	TraceServices::FAnalysisSessionReadScope SessionReadScope(*Session);
-	FRewindDebuggerTrackCreators::EnumerateCreators([this, &Session, RewindDebugger](const IRewindDebuggerTrackCreator* Creator)
-	{
-		if (IsTargetType(ObjectId, Creator->GetTargetTypeName(), *Session))
+	FRewindDebuggerTrackCreators::EnumerateCreators([this, &Session](const IRewindDebuggerTrackCreator* Creator)
 		{
-			TrackChildren.Push({Creator, nullptr});
-		}
-	});
+			if (IsTargetType(ObjectId, Creator->GetTargetTypeName(), *Session))
+			{
+				TrackChildren.Push({ Creator, nullptr });
+			}
+		});
 
 	// sort by creators by priority + name
 	TrackChildren.Sort([](const FTrackCreatorAndTrack& A, const FTrackCreatorAndTrack& B)
 		{
 			const int32 SortOrderPriorityA = A.Creator->GetSortOrderPriority();
 			const int32 SortOrderPriorityB = B.Creator->GetSortOrderPriority();
-			
+
 			if (SortOrderPriorityA != SortOrderPriorityB)
 			{
 				return SortOrderPriorityA > SortOrderPriorityB;
 			}
-			
+
 			return A.Creator->GetName().ToString() < B.Creator->GetName().ToString();
 		});
 }
-	
+
 TSharedPtr<SWidget> FRewindDebuggerObjectTrack::GetTimelineViewInternal()
 {
 	return SNew(SEventTimelineView)
@@ -88,23 +87,23 @@ bool FRewindDebuggerObjectTrack::HandleDoubleClickInternal()
 	static const FName HandlerFeatureName = IRewindDebuggerDoubleClickHandler::ModularFeatureName;
 
 	IRewindDebugger* RewindDebugger = IRewindDebugger::Instance();
-		
+
 	if (const TraceServices::IAnalysisSession* Session = RewindDebugger->GetAnalysisSession())
 	{
 		TraceServices::FAnalysisSessionReadScope SessionReadScope(*Session);
-	
+
 		const IGameplayProvider* GameplayProvider = Session->ReadProvider<IGameplayProvider>("GameplayProvider");
 		const FObjectInfo& ObjectInfo = GameplayProvider->GetObjectInfo(ObjectId);
 		uint64 ClassId = ObjectInfo.ClassId;
 		bool bHandled = false;
-		
+
 		const int32 NumExtensions = ModularFeatures.GetModularFeatureImplementationCount(HandlerFeatureName);
 
-		// iterate up the class hierarchy, looking for a registered double click handler, until we find the one that succeeeds that is most specific to the type of this object
+		// iterate up the class hierarchy, looking for a registered double click handler, until we find the one that succeeds that is most specific to the type of this object
 		while (ClassId != 0 && !bHandled)
 		{
 			const FClassInfo& ClassInfo = GameplayProvider->GetClassInfo(ClassId);
-		
+
 			for (int32 ExtensionIndex = 0; ExtensionIndex < NumExtensions; ++ExtensionIndex)
 			{
 				IRewindDebuggerDoubleClickHandler* Handler = static_cast<IRewindDebuggerDoubleClickHandler*>(ModularFeatures.GetModularFeatureImplementation(HandlerFeatureName, ExtensionIndex));
@@ -117,11 +116,11 @@ bool FRewindDebuggerObjectTrack::HandleDoubleClickInternal()
 					}
 				}
 			}
-		
+
 			ClassId = ClassInfo.SuperId;
 		}
 	}
-	
+
 	return true;
 }
 
@@ -135,13 +134,13 @@ void FRewindDebuggerObjectTrack::IterateSubTracksInternal(TFunction<void(TShared
 			IteratorFunction(TrackChild.Track);
 		}
 	}
-	
+
 	for (TSharedPtr<FRewindDebuggerTrack>& Track : Children)
 	{
 		IteratorFunction(Track);
 	}
 };
-	
+
 FText FRewindDebuggerObjectTrack::GetDisplayNameInternal() const
 {
 	if (!bDisplayNameValid)
@@ -164,9 +163,7 @@ FText FRewindDebuggerObjectTrack::GetDisplayNameInternal() const
 
 			if (const FWorldInfo* WorldInfo = GameplayProvider->FindWorldInfoFromObject(ObjectId))
 			{
-				bool bIsServer = WorldInfo->NetMode == FWorldInfo::ENetMode::DedicatedServer;
-
-				if (bIsServer)
+				if (WorldInfo->NetMode == FWorldInfo::ENetMode::DedicatedServer)
 				{
 					DisplayName = FText::Format(NSLOCTEXT("RewindDebuggerTrack", " (Server)", "{0} (Server)"), FText::FromString(ObjectInfo.Name));
 				}
@@ -192,10 +189,10 @@ bool FRewindDebuggerObjectTrack::UpdateInternal()
 
 	TRange<double> Existence = GameplayProvider->GetObjectRecordingLifetime(ObjectId);
 
-	ExistenceRange->Windows.SetNum(0,EAllowShrinking::No);
+	ExistenceRange->Windows.SetNum(0, EAllowShrinking::No);
 	if (Existence.HasLowerBound() && Existence.HasUpperBound())
 	{
-		ExistenceRange->Windows.Add({Existence.GetLowerBoundValue(), Existence.GetUpperBoundValue(), LOCTEXT("Object Existence","Object Existence"), LOCTEXT("Object Existence","Object Existence"), FLinearColor(0.1f,0.11f,0.1f)});
+		ExistenceRange->Windows.Add({ Existence.GetLowerBoundValue(), Existence.GetUpperBoundValue(), LOCTEXT("Object Existence","Object Existence"), LOCTEXT("Object Existence","Object Existence"), FLinearColor(0.1f,0.11f,0.1f) });
 	}
 
 	if (!bIconSearched)
@@ -205,7 +202,7 @@ bool FRewindDebuggerObjectTrack::UpdateInternal()
 		{
 			Icon = GameplayProvider->FindIconForClass(ObjectInfo->ClassId);
 		}
-						
+
 		if (!Icon.IsSet())
 		{
 			Icon = FSlateIconFinder::FindIconForClass(UObject::StaticClass());
@@ -220,7 +217,7 @@ bool FRewindDebuggerObjectTrack::UpdateInternal()
 
 	// add debug views as children
 
-	for(FTrackCreatorAndTrack& TrackChild : TrackChildren)
+	for (FTrackCreatorAndTrack& TrackChild : TrackChildren)
 	{
 		const bool bHasDebugInfo = TrackChild.Creator->HasDebugInfo(ObjectId);
 
@@ -242,7 +239,7 @@ bool FRewindDebuggerObjectTrack::UpdateInternal()
 		}
 	}
 
-	// Fallback codepath to add views with no track implementation
+	// Fallback code path to add views with no track implementation
 	{
 		TRACE_CPUPROFILER_EVENT_SCOPE(FRewindDebuggerObjectTrack::UpdateInternal_AddViews);
 		FRewindDebuggerViewCreators::EnumerateCreators([this, &Session, &bChanged, RewindDebugger](const IRewindDebuggerViewCreator* Creator)
@@ -273,32 +270,32 @@ bool FRewindDebuggerObjectTrack::UpdateInternal()
 		);
 	}
 
-	// add child objects/components
+	// add child objects
 	{
 		TRACE_CPUPROFILER_EVENT_SCOPE(FRewindDebuggerObjectTrack::UpdateInternal_AddChildComponents);
 
-		TRange<double> TraceRange = RewindDebugger->GetCurrentTraceRange();
 		TRange<double> ViewRange = RewindDebugger->GetCurrentViewRange();
 
-		GameplayProvider->EnumerateSubobjects(ObjectId, [this, &FoundObjects, &bChanged,&TraceRange, &ViewRange,GameplayProvider](uint64 SubobjectId)
+		GameplayProvider->EnumerateSubobjects(ObjectId, [this, &FoundObjects, &bChanged, &ViewRange, GameplayProvider](uint64 SubObjectId)
 			{
-				TRange<double> Lifetime = GameplayProvider->GetObjectRecordingLifetime(SubobjectId);
-				TRange<double> Overlap = TRange<double>::Intersection(Lifetime, ViewRange);
+				const TRange<double> Lifetime = GameplayProvider->GetObjectRecordingLifetime(SubObjectId);
+				const TRange<double> Overlap = TRange<double>::Intersection(Lifetime, ViewRange);
 				// only display the track if the lifetime of the object and the view range overlap
 				if (!Overlap.IsEmpty())
 				{
-					int32 FoundIndex = Children.FindLastByPredicate([SubobjectId](const TSharedPtr<FRewindDebuggerTrack>& Info) { return Info->GetObjectId() == SubobjectId; });
-					bool bExisted = FoundIndex >= 0;
+					const int32 FoundIndex = Children.FindLastByPredicate([SubObjectId](const TSharedPtr<FRewindDebuggerTrack>& Track)
+						{
+							return Track->GetObjectId() == SubObjectId;
+						});
 
-					if (!bExisted)
+					if (FoundIndex == INDEX_NONE)
 					{
-						FoundIndex = Children.Num();
-
-						Children.Add(MakeShared<FRewindDebuggerObjectTrack>(SubobjectId, ""));
+						const FObjectInfo& ObjectInfo = GameplayProvider->GetObjectInfo(SubObjectId);
+						Children.Add(MakeShared<FRewindDebuggerObjectTrack>(SubObjectId, ObjectInfo.Name));
+						bChanged = true;
 					}
 
-					bChanged = bChanged || !bExisted;
-					FoundObjects.Add(SubobjectId);
+					FoundObjects.Add(SubObjectId);
 				}
 			});
 	}
@@ -307,18 +304,21 @@ bool FRewindDebuggerObjectTrack::UpdateInternal()
 	if (bAddController)
 	{
 		TRACE_CPUPROFILER_EVENT_SCOPE(FRewindDebuggerObjectTrack::FindController);
-		// Should probably update this to use a time range and return all posessing controllers from the visible time range.  For now just returns the one at the current time.
+		// Should probably update this to use a time range and return all possessing controllers from the visible time range.  For now just returns the one at the current time.
 		if (uint64 ControllerId = GameplayProvider->FindPossessingController(ObjectId, RewindDebugger->CurrentTraceTime()))
 		{
 			const FObjectInfo& ObjectInfo = GameplayProvider->GetObjectInfo(ControllerId);
-			const int32 FoundIndex = Children.FindLastByPredicate([ObjectInfo](const TSharedPtr<FRewindDebuggerTrack>& Info) { return Info->GetObjectId() == ObjectInfo.Id; });
+			const int32 FoundIndex = Children.FindLastByPredicate([ObjectInfo](const TSharedPtr<FRewindDebuggerTrack>& Track)
+				{
+					return Track->GetObjectId() == ObjectInfo.Id;
+				});
 
 			if (FoundIndex < 0)
 			{
 				bChanged = true;
-				Children.Add(MakeShared<FRewindDebuggerObjectTrack>(ObjectInfo.Id,ObjectInfo.Name));
+				Children.Add(MakeShared<FRewindDebuggerObjectTrack>(ObjectInfo.Id, ObjectInfo.Name));
 			}
-			
+
 			FoundObjects.Add(ControllerId);
 		}
 	}
@@ -337,9 +337,9 @@ bool FRewindDebuggerObjectTrack::UpdateInternal()
 	{
 		// sort child object tracks by name
 		Children.Sort([](const TSharedPtr<FRewindDebuggerTrack>& A, const TSharedPtr<FRewindDebuggerTrack>& B)
-		{
-			return A->GetDisplayName().ToString() < B->GetDisplayName().ToString();
-		});
+			{
+				return A->GetDisplayName().ToString() < B->GetDisplayName().ToString();
+			});
 	}
 
 	{
diff --git a/Engine/Plugins/Animation/GameplayInsights/Source/RewindDebugger/Private/RewindDebuggerSettings.h b/Engine/Plugins/Animation/GameplayInsights/Source/RewindDebugger/Private/RewindDebuggerSettings.h
--- a/Engine/Plugins/Animation/GameplayInsights/Source/RewindDebugger/Private/RewindDebuggerSettings.h
+++ b/Engine/Plugins/Animation/GameplayInsights/Source/RewindDebugger/Private/RewindDebuggerSettings.h
@@ -59,7 +59,7 @@ public:
 	TArray<FName> HiddenTrackTypes;
 
 	
-	//** Currently selected target actor's name */
+	/** Currently selected target actor's name */
 	UPROPERTY(Config)
 	FString DebugTargetActor;
 	
diff --git a/Engine/Plugins/Animation/GameplayInsights/Source/RewindDebugger/Private/SRewindDebugger.cpp b/Engine/Plugins/Animation/GameplayInsights/Source/RewindDebugger/Private/SRewindDebugger.cpp
--- a/Engine/Plugins/Animation/GameplayInsights/Source/RewindDebugger/Private/SRewindDebugger.cpp
+++ b/Engine/Plugins/Animation/GameplayInsights/Source/RewindDebugger/Private/SRewindDebugger.cpp
@@ -1,24 +1,25 @@
 // Copyright Epic Games, Inc. All Rights Reserved.
 
 #include "SRewindDebugger.h"
-
 #include "ActorPickerMode.h"
 #include "Editor.h"
-#include "IGameplayProvider.h"
-#include "IRewindDebuggerTrackCreator.h"
-#include "LocalizationDescriptor.h"
 #include "Editor/EditorEngine.h"
 #include "Framework/Application/SlateApplication.h"
 #include "Framework/MultiBox/MultiBoxBuilder.h"
+#include "IGameplayProvider.h"
 #include "Insights/IUnrealInsightsModule.h"
+#include "IRewindDebuggerTrackCreator.h"
+#include "Kismet2/DebuggerCommands.h"
+#include "LocalizationDescriptor.h"
 #include "Modules/ModuleManager.h"
 #include "ObjectTrace.h"
-#include "RewindDebuggerStyle.h"
 #include "RewindDebuggerCommands.h"
-#include "SSimpleTimeSlider.h"
-#include "SceneOutlinerPublicTypes.h"
+#include "RewindDebuggerModule.h"
+#include "RewindDebuggerSettings.h"
+#include "RewindDebuggerStyle.h"
 #include "SceneOutlinerModule.h"
 #include "Selection.h"
+#include "SSimpleTimeSlider.h"
 #include "Styling/SlateIconFinder.h"
 #include "ToolMenu.h"
 #include "ToolMenus.h"
@@ -26,37 +27,28 @@
 #include "Widgets/Input/SButton.h"
 #include "Widgets/Input/SComboButton.h"
 #include "Widgets/Input/SSearchBox.h"
-#include "Widgets/Layout/SExpandableArea.h"
-#include "Widgets/Layout/SScrollBox.h"
-#include "Kismet2/DebuggerCommands.h"
-#include "RewindDebuggerSettings.h"
 #include "Widgets/Input/STextComboBox.h"
+#include "Widgets/Layout/SScrollBox.h"
 
 #define LOCTEXT_NAMESPACE "SRewindDebugger"
 
-SRewindDebugger::SRewindDebugger() 
-	: SCompoundWidget()
-	, ViewRange(0,10)
+SRewindDebugger::SRewindDebugger()
+	: ViewRange(0, 10)
 	, Commands(FRewindDebuggerCommands::Get())
-	, DebugComponents(nullptr)
 	, Settings(URewindDebuggerSettings::Get())
-{ 
-}
-
-SRewindDebugger::~SRewindDebugger() 
 {
 }
 
-void SRewindDebugger::TrackCursor(bool bReverse)
+void SRewindDebugger::TrackCursor(const bool bReverse)
 {
-	float ScrubTime = ScrubTimeAttribute.Get();
+	const float ScrubTime = ScrubTimeAttribute.Get();
 	TRange<double> CurrentViewRange = ViewRange;
-	float ViewSize = CurrentViewRange.GetUpperBoundValue() - CurrentViewRange.GetLowerBoundValue();
+	const float ViewSize = CurrentViewRange.GetUpperBoundValue() - CurrentViewRange.GetLowerBoundValue();
 
-	static const double LeadingEdgeSize = 0.05;
-	static const double TrailingEdgeThreshold = 0.01;
+	static constexpr double LeadingEdgeSize = 0.05;
+	static constexpr double TrailingEdgeThreshold = 0.01;
 
-	if(bReverse)
+	if (bReverse)
 	{
 		// playing in reverse (cursor moving left)
 		if (ScrubTime < CurrentViewRange.GetLowerBoundValue() + ViewSize * LeadingEdgeSize)
@@ -96,9 +88,9 @@ void SRewindDebugger::SetViewRange(TRange<double> NewRange)
 
 void SRewindDebugger::ToggleHideTrackType(const FName& TrackType)
 {
-	int32 Index = Settings.HiddenTrackTypes.Find(TrackType);
+	const int32 Index = Settings.HiddenTrackTypes.Find(TrackType);
 
-	if (Index >=0)
+	if (Index >= 0)
 	{
 		Settings.HiddenTrackTypes.RemoveAtSwap(Index);
 	}
@@ -106,7 +98,7 @@ void SRewindDebugger::ToggleHideTrackType(const FName& TrackType)
 	{
 		Settings.HiddenTrackTypes.Add(TrackType);
 	}
-	RefreshDebugComponents();
+	RefreshTracks();
 }
 
 bool SRewindDebugger::ShouldHideTrackType(const FName& TrackType) const
@@ -118,7 +110,7 @@ bool SRewindDebugger::ShouldHideTrackType(const FName& TrackType) const
 void SRewindDebugger::ToggleDisplayEmptyTracks()
 {
 	Settings.bShowEmptyObjectTracks = !Settings.bShowEmptyObjectTracks;
-	RefreshDebugComponents();
+	RefreshTracks();
 }
 
 bool SRewindDebugger::ShouldDisplayEmptyTracks() const
@@ -129,75 +121,76 @@ bool SRewindDebugger::ShouldDisplayEmptyTracks() const
 FReply SRewindDebugger::OnPreviewKeyDown(const FGeometry& MyGeometry, const FKeyEvent& InKeyEvent)
 {
 	const FInputChord KeyEventAsInputChord = FInputChord(InKeyEvent.GetKey(), EModifierKey::FromBools(InKeyEvent.IsControlDown(), InKeyEvent.IsAltDown(), InKeyEvent.IsShiftDown(), InKeyEvent.IsCommandDown()));
-	FReply bReply = FReply::Unhandled();
+	FReply Reply = FReply::Unhandled();
 
 	// Handle Rewind Debugger VCR Commands
 	if (CommandList->ProcessCommandBindings(InKeyEvent))
 	{
-		bReply = FReply::Handled();
+		Reply = FReply::Handled();
 	}
 
 	// Prevent bubbling up shortcut for "ReversePlay"
 	if (Commands.ReversePlay->HasDefaultChord(KeyEventAsInputChord) && !IsPIESimulating.Get())
 	{
-		bReply = FReply::Handled();
+		Reply = FReply::Handled();
 	}
 
-	return bReply;
+	return Reply;
 }
 
-void SRewindDebugger::SetDebugTargetActor(AActor* Actor)
+void SRewindDebugger::SetDebuggedObject(const UObject* Object)
 {
-	FString ActorLabel = Actor->GetActorLabel();
 	// Spawned actors have a "RewindDebugger: " prefix on their label
-	if (ActorLabel.StartsWith("RewindDebugger: "))
+	const AActor* Actor = Cast<AActor>(Object);
+	if (Actor != nullptr && Actor->GetActorLabel().StartsWith("RewindDebugger: "))
 	{
+		FString ActorLabel = Actor->GetActorLabel();
 		ActorLabel.RemoveFromStart("RewindDebugger: ");
-		DebugTargetActor.Set(ActorLabel);
+		DebuggedObjectName.Set(ActorLabel);
 	}
 	else
 	{
-		DebugTargetActor.Set(Actor->GetName());
+		DebuggedObjectName.Set(Object->GetName());
 	}
 }
 
-TSharedRef<SWidget> SRewindDebugger::MakeSelectActorMenu()
+TSharedRef<SWidget> SRewindDebugger::MakeObjectSelectionMenu()
 {
 	FMenuBuilder MenuBuilder(true, nullptr);
 
 	MenuBuilder.AddSearchWidget();
 
 	// Menu entry for each actor selected in the scene
+	MenuBuilder.BeginSection("From Selection Section", LOCTEXT("FromSelection", "From Scene Selection"));
+
 	TArray<AActor*> SelectedActors;
 	GEditor->GetSelectedActors()->GetSelectedObjects(SelectedActors);
-
-	MenuBuilder.BeginSection("From Selection Section", LOCTEXT("FromSelection", "From Scene Selection"));
 	if (SelectedActors.Num() > 0)
 	{
-		for(AActor* SelectedActor : SelectedActors)
+		for (AActor* SelectedActor : SelectedActors)
 		{
 			FText SelectedLabel = FText::FromString(SelectedActor->GetActorLabel());
 			FSlateIcon ActorIcon = FSlateIconFinder::FindIconForClass(SelectedActors[0]->GetClass());
 
-			MenuBuilder.AddMenuEntry(SelectedLabel, FText(), ActorIcon, FExecuteAction::CreateLambda([this, SelectedActor]{
+			MenuBuilder.AddMenuEntry(SelectedLabel, FText(), ActorIcon, FExecuteAction::CreateLambda([this, SelectedActor] {
 				FSlateApplication::Get().DismissAllMenus();
-				SetDebugTargetActor(SelectedActor);
-			}));
+				SetDebuggedObject(SelectedActor);
+				}));
 		}
 	}
 	else
 	{
 		MenuBuilder.AddMenuEntry(LOCTEXT("No scene selection", "No scene selection"),
-		 LOCTEXT("SceneSelectionToolTip", "If you select an object in the scene, then it will be listed here"),
-		 FSlateIcon(),
-		  FUIAction(FExecuteAction(),
-		  	FCanExecuteAction::CreateLambda([](){return false;})));
+			LOCTEXT("SceneSelectionToolTip", "If you select an object in the scene, then it will be listed here"),
+			FSlateIcon(),
+			FUIAction(FExecuteAction(),
+				FCanExecuteAction::CreateLambda([]() {return false; })));
 	}
 	MenuBuilder.EndSection();
 
 	MenuBuilder.BeginSection("From Recording Section", LOCTEXT("FromRecording", "From Recording:"));
 
-	if (IRewindDebugger* RewindDebugger = IRewindDebugger::Instance())
+	if (const IRewindDebugger* RewindDebugger = IRewindDebugger::Instance())
 	{
 		if (const TraceServices::IAnalysisSession* Session = RewindDebugger->GetAnalysisSession())
 		{
@@ -205,33 +198,32 @@ TSharedRef<SWidget> SRewindDebugger::MakeSelectActorMenu()
 			const IGameplayProvider* GameplayProvider = Session->ReadProvider<IGameplayProvider>("GameplayProvider");
 
 			const FClassInfo* ActorClassInfo = GameplayProvider->FindClassInfo(*AActor::StaticClass()->GetPathName());
-			
+
 			GameplayProvider->EnumerateObjects([this, &MenuBuilder, GameplayProvider, ActorClassInfo](const FObjectInfo& ObjectInfo)
-			{
-				if (GameplayProvider->IsSubClassOf(ObjectInfo.ClassId, ActorClassInfo->Id))
 				{
-					FString ActorName = ObjectInfo.Name;
-					FText SelectedLabel = FText::FromString(ActorName);
-					FSlateIcon ActorIcon = GameplayProvider->FindIconForClass(ObjectInfo.ClassId);
-					const FClassInfo& ClassInfo = GameplayProvider->GetClassInfo(ObjectInfo.ClassId);
-
-					MenuBuilder.AddMenuEntry(SelectedLabel, FText(), ActorIcon, FExecuteAction::CreateLambda([this, ActorName]()
+					if (GameplayProvider->IsSubClassOf(ObjectInfo.ClassId, ActorClassInfo->Id))
 					{
-						FSlateApplication::Get().DismissAllMenus();
-						DebugTargetActor.Set(ActorName);
-					}));
-				}
-			});
+						const FString Name = ObjectInfo.Name;
+						const FText SelectedLabel = FText::FromString(Name);
+						const FSlateIcon Icon = GameplayProvider->FindIconForClass(ObjectInfo.ClassId);
+
+						MenuBuilder.AddMenuEntry(SelectedLabel, FText(), Icon, FExecuteAction::CreateLambda([this, Name]()
+							{
+								FSlateApplication::Get().DismissAllMenus();
+								DebuggedObjectName.Set(Name);
+							}));
+					}
+				});
 		}
 		else
 		{
 			MenuBuilder.AddMenuEntry(LOCTEXT("No recording loaded", "No recording loaded"),
-			 LOCTEXT("NoRecordingToolTip", "Start or load a recording, and the recorded actors will be listed here"),
-			 FSlateIcon(),
-			  FUIAction(FExecuteAction(),
-				FCanExecuteAction::CreateLambda([](){return false;})));
-        }
-		
+				LOCTEXT("NoRecordingToolTip", "Start or load a recording, and the recorded actors will be listed here"),
+				FSlateIcon(),
+				FUIAction(FExecuteAction(),
+					FCanExecuteAction::CreateLambda([]() {return false; })));
+		}
+
 	}
 
 	return MenuBuilder.MakeWidget();
@@ -241,39 +233,39 @@ void SRewindDebugger::Construct(const FArguments& InArgs, TSharedRef<FUICommandL
 {
 	OnScrubPositionChanged = InArgs._OnScrubPositionChanged;
 	OnViewRangeChanged = InArgs._OnViewRangeChanged;
-	OnComponentSelectionChanged = InArgs._OnComponentSelectionChanged;
-	BuildComponentContextMenu = InArgs._BuildComponentContextMenu;
+	OnTrackSelectionChanged = InArgs._OnTrackSelectionChanged;
+	BuildTrackContextMenu = InArgs._BuildTrackContextMenu;
 	TrackTypesAttribute = InArgs._TrackTypes;
 	ScrubTimeAttribute = InArgs._ScrubTime;
-	DebugComponents = InArgs._DebugComponents;
+	Tracks = InArgs._Tracks;
 	TraceTime.Initialize(InArgs._TraceTime);
 	RecordingDuration.Initialize(InArgs._RecordingDuration);
-	DebugTargetActor.Initialize(InArgs._DebugTargetActor);
+	DebuggedObjectName.Initialize(InArgs._DebuggedObjectName);
 	IsPIESimulating = InArgs._IsPIESimulating;
 	CommandList = InCommandList;
 
-	TrackFilterBox = SNew(SSearchBox).HintText(LOCTEXT("Filter Tracks","Filter Tracks")).OnTextChanged_Lambda([this](const FText&)
-	{
-		RefreshDebugComponents();
-	});
-	
-	TSharedPtr<SScrollBar> ScrollBar = SNew(SScrollBar);
-	FToolMenuContext ToolMenuContext(CommandList);
+	TrackFilterBox = SNew(SSearchBox).HintText(LOCTEXT("Filter Tracks", "Filter Tracks")).OnTextChanged_Lambda([this](const FText&)
+		{
+			RefreshTracks();
+		});
 
-	ComponentTreeView =	SNew(SRewindDebuggerComponentTree)
+	const TSharedPtr<SScrollBar> ScrollBar = SNew(SScrollBar);
+	const FToolMenuContext ToolMenuContext(CommandList);
+
+	TrackTreeView = SNew(SRewindDebuggerTrackTree)
 		.ExternalScrollBar(ScrollBar)
 		.OnExpansionChanged_Lambda([this]()
-		{
-			if (!bInExpansionChanged)
 			{
-				bInExpansionChanged = true;
-				TimelinesView->RestoreExpansion();
-				bInExpansionChanged = false;
-			}
-		})
-		.OnScrolled_Lambda([this](double ScrollOffset){ TimelinesView->ScrollTo(ScrollOffset); })
-		.DebugComponents(InArgs._DebugComponents)
-		.OnMouseButtonDoubleClick(InArgs._OnComponentDoubleClicked)
+				if (!bInExpansionChanged)
+				{
+					bInExpansionChanged = true;
+					TimelinesView->RestoreExpansion();
+					bInExpansionChanged = false;
+				}
+			})
+		.OnScrolled_Lambda([this](double ScrollOffset) { TimelinesView->ScrollTo(ScrollOffset); })
+		.Tracks(InArgs._Tracks)
+		.OnMouseButtonDoubleClick(InArgs._OnTrackDoubleClicked)
 		.OnContextMenuOpening(this, &SRewindDebugger::OnContextMenuOpening)
 		.OnSelectionChanged_Lambda(
 			[this](TSharedPtr<RewindDebugger::FRewindDebuggerTrack> SelectedItem, ESelectInfo::Type SelectInfo)
@@ -282,31 +274,31 @@ void SRewindDebugger::Construct(const FArguments& InArgs, TSharedRef<FUICommandL
 				{
 					bInSelectionChanged = true;
 					TimelinesView->SetSelection(SelectedItem);
-					ComponentSelectionChanged(SelectedItem, SelectInfo);
+					OnSelectedTrackChanged(SelectedItem, SelectInfo);
 					bInSelectionChanged = false;
 				}
 			});
 
-	 TimelinesView = SNew(SRewindDebuggerTimelines)
+	TimelinesView = SNew(SRewindDebuggerTimelines)
 		.ExternalScrollbar(ScrollBar)
 		.OnExpansionChanged_Lambda(
-				[this]()
+			[this]()
+			{
+				if (!bInExpansionChanged)
 				{
-					if (!bInExpansionChanged)
-					{
-						bInExpansionChanged = true;
-						ComponentTreeView->RestoreExpansion();
-						bInExpansionChanged = false;
-					}
-				})
-		.OnScrolled_Lambda([this](double ScrollOffset){ ComponentTreeView->ScrollTo(ScrollOffset); })
-		.DebugComponents(InArgs._DebugComponents)
-		.ViewRange_Lambda([this](){return ViewRange;})
+					bInExpansionChanged = true;
+					TrackTreeView->RestoreExpansion();
+					bInExpansionChanged = false;
+				}
+			})
+		.OnScrolled_Lambda([this](double ScrollOffset) { TrackTreeView->ScrollTo(ScrollOffset); })
+		.Tracks(InArgs._Tracks)
+		.ViewRange_Lambda([this]() {return ViewRange; })
 		.ClampRange_Lambda(
-			 [this]()
-			 {
-				 return TRange<double>(0.0f, RecordingDuration.Get());
-			 })
+			[this]()
+			{
+				return TRange<double>(0.0f, RecordingDuration.Get());
+			})
 		.OnViewRangeChanged(this, &SRewindDebugger::SetViewRange)
 		.ScrubPosition(ScrubTimeAttribute)
 		.OnScrubPositionChanged_Lambda(
@@ -314,7 +306,7 @@ void SRewindDebugger::Construct(const FArguments& InArgs, TSharedRef<FUICommandL
 			{
 				if (bIsScrubbing)
 				{
-					OnScrubPositionChanged.ExecuteIfBound( NewScrubTime, bIsScrubbing );
+					OnScrubPositionChanged.ExecuteIfBound(NewScrubTime, bIsScrubbing);
 				}
 			})
 		.OnSelectionChanged_Lambda(
@@ -323,167 +315,167 @@ void SRewindDebugger::Construct(const FArguments& InArgs, TSharedRef<FUICommandL
 				if (!bInSelectionChanged)
 				{
 					bInSelectionChanged = true;
-					ComponentTreeView->SetSelection(SelectedItem);
-					ComponentSelectionChanged(SelectedItem, SelectInfo);
+					TrackTreeView->SetSelection(SelectedItem);
+					OnSelectedTrackChanged(SelectedItem, SelectInfo);
 					bInSelectionChanged = false;
 				}
 			});
 
-	
-	TSharedRef<SWidget> ToolBar = UToolMenus::Get()->GenerateWidget("RewindDebugger.ToolBar", ToolMenuContext);
+
+	const TSharedRef<SWidget> ToolBar = UToolMenus::Get()->GenerateWidget("RewindDebugger.ToolBar", ToolMenuContext);
 
 	ChildSlot
-	[
-		SNew(SVerticalBox)
-		 + SVerticalBox::Slot().AutoHeight()
 		[
 			SNew(SVerticalBox)
-			+SVerticalBox::Slot().MaxHeight(48)
-			[
-				SNew(SHorizontalBox)
-				+SHorizontalBox::Slot().AutoWidth()
-				[
-					SNew(SComboButton)
-						.ComboButtonStyle(FAppStyle::Get(), "SimpleComboButton")
-						.OnGetMenuContent(this, &SRewindDebugger::MakeMainMenu)
-						.ButtonContent()
-					[
-						SNew(SImage)
-						.Image(FAppStyle::Get().GetBrush("ClassIcon.CameraComponent"))
-					]
-				]
-				+SHorizontalBox::Slot().FillWidth(1.0)
+				+ SVerticalBox::Slot().AutoHeight()
 				[
-					ToolBar
+					SNew(SVerticalBox)
+						+ SVerticalBox::Slot().MaxHeight(48)
+						[
+							SNew(SHorizontalBox)
+								+ SHorizontalBox::Slot().AutoWidth()
+								[
+									SNew(SComboButton)
+										.ComboButtonStyle(FAppStyle::Get(), "SimpleComboButton")
+										.OnGetMenuContent(this, &SRewindDebugger::MakeMainMenu)
+										.ButtonContent()
+										[
+											SNew(SImage)
+												.Image(FAppStyle::Get().GetBrush("ClassIcon.CameraComponent"))
+										]
+								]
+								+ SHorizontalBox::Slot().FillWidth(1.0)
+								[
+									ToolBar
+								]
+						]
 				]
-			]
-		]
-		 + SVerticalBox::Slot().FillHeight(1.0)
-		 [
-			SNew(SSplitter)
-			+SSplitter::Slot().MinSize(280).Value(0)
-			[
-				SNew(SVerticalBox)
-				+SVerticalBox::Slot().AutoHeight()
+				+ SVerticalBox::Slot().FillHeight(1.0)
 				[
-					SNew(SHorizontalBox)
-					+SHorizontalBox::Slot() .FillWidth(1.0f)
-					[
-						// SNew(SRewindDebuggerTargetSelector)
-						SNew(SComboButton)
-						.OnGetMenuContent(this, &SRewindDebugger::MakeSelectActorMenu)
-						.ButtonContent()
+					SNew(SSplitter)
+						+ SSplitter::Slot().MinSize(280).Value(0)
 						[
-							SNew(SHorizontalBox)
-							+SHorizontalBox::Slot().AutoWidth().Padding(3)
-							[
-								SNew(SImage)
-								.Image_Lambda([this]
-									{
-										if (DebugComponents != nullptr && DebugComponents->Num()>0)
-										{
-											return (*DebugComponents)[0]->GetIcon().GetIcon();
-										}
-
-										return FSlateIconFinder::FindIconForClass(AActor::StaticClass()).GetIcon();
-									}
-								)
-							]
-							+SHorizontalBox::Slot().Padding(3)
-							[
-								SNew(STextBlock)
-								.Text_Lambda([this](){
-									if (DebugComponents == nullptr || DebugComponents->Num()==0)
-									{
-										return LOCTEXT("Select Actor", "Debug Target Actor");
-									}
-
-									FText ReadableName = (*DebugComponents)[0]->GetDisplayName();
+							SNew(SVerticalBox)
+								+ SVerticalBox::Slot().AutoHeight()
+								[
+									SNew(SHorizontalBox)
+										+ SHorizontalBox::Slot().FillWidth(1.0f)
+										[
+											// SNew(SRewindDebuggerTargetSelector)
+											SNew(SComboButton)
+												.OnGetMenuContent(this, &SRewindDebugger::MakeObjectSelectionMenu)
+												.ButtonContent()
+												[
+													SNew(SHorizontalBox)
+														+ SHorizontalBox::Slot().AutoWidth().Padding(3)
+														[
+															SNew(SImage)
+																.Image_Lambda([this]
+																	{
+																		if (Tracks != nullptr && Tracks->Num() > 0)
+																		{
+																			return (*Tracks)[0]->GetIcon().GetIcon();
+																		}
+
+																		return FSlateIconFinder::FindIconForClass(AActor::StaticClass()).GetIcon();
+																	}
+																)
+														]
+														+ SHorizontalBox::Slot().Padding(3)
+														[
+															SNew(STextBlock)
+																.Text_Lambda([this]() {
+																if (Tracks == nullptr || Tracks->Num() == 0)
+																{
+																	return LOCTEXT("Select Actor", "Debug Target Actor");
+																}
+
+																FText ReadableName = (*Tracks)[0]->GetDisplayName();
 #if OBJECT_TRACE_ENABLED
-									if (UObject* Object = FObjectTrace::GetObjectFromId((*DebugComponents)[0]->GetObjectId()))
-									{
-										if (AActor* Actor = Cast<AActor>(Object))
-										{
-											ReadableName = FText::FromString(Actor->GetActorLabel());
-										}
-									}
+																if (UObject* Object = FObjectTrace::GetObjectFromId((*Tracks)[0]->GetObjectId()))
+																{
+																	if (const AActor* Actor = Cast<AActor>(Object))
+																	{
+																		ReadableName = FText::FromString(Actor->GetActorLabel());
+																	}
+																}
 #endif // OBJECT_TRACE_ENABLED
 
-									return ReadableName;
-								} )
-							]
+																return ReadableName;
+																	})
+														]
+												]
+										]
+										+ SHorizontalBox::Slot().AutoWidth().HAlign(HAlign_Right)
+										[
+											SNew(SButton)
+												.ButtonStyle(FAppStyle::Get(), "SimpleButton")
+												.OnClicked(this, &SRewindDebugger::OnSelectActorClicked)
+												.ToolTipText(LOCTEXT("SelectActorTooltip", "Select Target Actor in Scene"))
+												[
+													SNew(SImage)
+														.Image(FRewindDebuggerStyle::Get().GetBrush("RewindDebugger.SelectActor"))
+												]
+										]
+										+ SHorizontalBox::Slot().AutoWidth()
+										[
+											SNew(SComboButton)
+												.ComboButtonStyle(FAppStyle::Get(), "SimpleComboButton")
+												.OnGetMenuContent(this, &SRewindDebugger::MakeFilterMenu)
+												.ButtonContent()
+												[
+													SNew(SImage)
+														.Image(FAppStyle::Get().GetBrush("Icons.Filter"))
+												]
+										]
+								]
+								+ SVerticalBox::Slot().FillHeight(1.0f)
+								[
+									TrackTreeView.ToSharedRef()
+								]
 						]
-					]
-					+SHorizontalBox::Slot().AutoWidth().HAlign(HAlign_Right)
-					[
-						SNew(SButton)
-							.ButtonStyle(FAppStyle::Get(), "SimpleButton")
-							.OnClicked(this, &SRewindDebugger::OnSelectActorClicked)
-							.ToolTipText(LOCTEXT("SelectActorTooltip", "Select Target Actor in Scene"))
+						+ SSplitter::Slot()
 						[
-							SNew(SImage)
-							.Image(FRewindDebuggerStyle::Get().GetBrush("RewindDebugger.SelectActor"))
+							SNew(SVerticalBox)
+								+ SVerticalBox::Slot().AutoHeight()
+								[
+									SNew(SSimpleTimeSlider)
+										.DesiredSize({ 100,24 })
+										.ClampRangeHighlightSize(0.15f)
+										.ClampRangeHighlightColor(FLinearColor::Red.CopyWithNewOpacity(0.5f))
+										.ScrubPosition(ScrubTimeAttribute)
+										.ViewRange_Lambda([this]() { return ViewRange; })
+										.OnViewRangeChanged(this, &SRewindDebugger::SetViewRange)
+										.ClampRange_Lambda(
+											[this]()
+											{
+												return TRange<double>(0.0f, RecordingDuration.Get());
+											})
+										.OnScrubPositionChanged_Lambda(
+											[this](double NewScrubTime, bool bIsScrubbing)
+											{
+												if (bIsScrubbing)
+												{
+													OnScrubPositionChanged.ExecuteIfBound(NewScrubTime, bIsScrubbing);
+												}
+											})
+								]
+								+ SVerticalBox::Slot().FillHeight(1.0f)
+								[
+									SNew(SOverlay)
+										+ SOverlay::Slot()
+										[
+											TimelinesView.ToSharedRef()
+										]
+										+ SOverlay::Slot().HAlign(EHorizontalAlignment::HAlign_Right)
+										[
+											ScrollBar.ToSharedRef()
+										]
+
+								]
 						]
-					]
-					+SHorizontalBox::Slot().AutoWidth()
-					[
-						SNew(SComboButton)
-							.ComboButtonStyle(FAppStyle::Get(), "SimpleComboButton")
-							.OnGetMenuContent(this, &SRewindDebugger::MakeFilterMenu)
-							.ButtonContent()
-						[
-							SNew(SImage)
-							.Image(FAppStyle::Get().GetBrush("Icons.Filter"))
-						]
-					]
 				]
-				+SVerticalBox::Slot().FillHeight(1.0f)
-				[
-					ComponentTreeView.ToSharedRef()
-				]
-			]
-			+SSplitter::Slot() 
-			[
-				SNew(SVerticalBox)
-				+SVerticalBox::Slot().AutoHeight()
-				[
-					SNew(SSimpleTimeSlider)
-						.DesiredSize({100,24})
-						.ClampRangeHighlightSize(0.15f)
-						.ClampRangeHighlightColor(FLinearColor::Red.CopyWithNewOpacity(0.5f))
-						.ScrubPosition(ScrubTimeAttribute)
-						.ViewRange_Lambda([this](){ return ViewRange; })
-						.OnViewRangeChanged(this, &SRewindDebugger::SetViewRange)
-						.ClampRange_Lambda(
-								[this]()
-								{ 
-									return TRange<double>(0.0f,RecordingDuration.Get());
-								})	
-						.OnScrubPositionChanged_Lambda(
-							[this](double NewScrubTime, bool bIsScrubbing)
-									{
-										if (bIsScrubbing)
-										{
-											OnScrubPositionChanged.ExecuteIfBound( NewScrubTime, bIsScrubbing );
-										}
-									})
-				]
-				+SVerticalBox::Slot().FillHeight(1.0f)
-				[
-					SNew(SOverlay)
-					+ SOverlay::Slot()
-					[				
-						TimelinesView.ToSharedRef()
-					]
-					+SOverlay::Slot().HAlign(EHorizontalAlignment::HAlign_Right)
-					[
-						ScrollBar.ToSharedRef()
-					]
-					
-				]
-			]
-		]
-	];
+		];
 }
 
 bool FilterTrack(TSharedPtr<RewindDebugger::FRewindDebuggerTrack>& Track, const FString& FilterString, bool bRemoveNoData, const TArray<FName>& FilteredTrackTypes, bool bParentFilterPassed = false)
@@ -493,66 +485,65 @@ bool FilterTrack(TSharedPtr<RewindDebugger::FRewindDebuggerTrack>& Track, const
 		Track->SetIsVisible(false);
 		return false;
 	}
-	
-	const bool bStringFilterEmpty =  FilterString.IsEmpty();
+
+	const bool bStringFilterEmpty = FilterString.IsEmpty();
 	const bool bStringFilterPassed = bParentFilterPassed || bStringFilterEmpty || Track->GetDisplayName().ToString().Contains(FilterString);
 
 	const bool bThisFilterPassed = (!bStringFilterEmpty && bStringFilterPassed);
 
-
 	bool bAnyChildVisible = false;
 	Track->IterateSubTracks([&bAnyChildVisible, bThisFilterPassed, FilterString, bRemoveNoData, FilteredTrackTypes](TSharedPtr<RewindDebugger::FRewindDebuggerTrack> ChildTrack)
-	{
-		const bool bChildIsVisible = FilterTrack(ChildTrack, FilterString, bRemoveNoData, FilteredTrackTypes, bThisFilterPassed);
-		bAnyChildVisible |= bChildIsVisible;
-	});
+		{
+			const bool bChildIsVisible = FilterTrack(ChildTrack, FilterString, bRemoveNoData, FilteredTrackTypes, bThisFilterPassed);
+			bAnyChildVisible |= bChildIsVisible;
+		});
 
-	bool bVisible = bAnyChildVisible || ((!bRemoveNoData || Track->HasDebugData()) && bStringFilterPassed);
+	const bool bVisible = bAnyChildVisible || ((!bRemoveNoData || Track->HasDebugData()) && bStringFilterPassed);
 
 	Track->SetIsVisible(bVisible);
 	return bVisible;
 }
 
-void SRewindDebugger::RefreshDebugComponents()
+void SRewindDebugger::RefreshTracks()
 {
-	if (DebugComponents)
+	if (Tracks)
 	{
-		for(TSharedPtr<RewindDebugger::FRewindDebuggerTrack>& DebugComponent : *DebugComponents)
+		for (TSharedPtr<RewindDebugger::FRewindDebuggerTrack>& DebugComponent : *Tracks)
 		{
 			FilterTrack(DebugComponent, TrackFilterBox->GetText().ToString(), !ShouldDisplayEmptyTracks(), Settings.HiddenTrackTypes);
 		}
 	}
-	
-	ComponentTreeView->Refresh();
+
+	TrackTreeView->Refresh();
 	TimelinesView->Refresh();
 }
 
 TSharedRef<SWidget> SRewindDebugger::MakeMainMenu()
 {
-	return UToolMenus::Get()->GenerateWidget("RewindDebugger.MainMenu", FToolMenuContext());
+	return UToolMenus::Get()->GenerateWidget(FRewindDebuggerModule::MainMenuName, FToolMenuContext());
 }
 
 TSharedRef<SWidget> SRewindDebugger::MakeFilterMenu()
 {
 	FMenuBuilder Builder(true, nullptr);
 	Builder.AddWidget(TrackFilterBox.ToSharedRef(), FText(), true, false);
-	
+
 	Builder.BeginSection("TrackTypes", LOCTEXT("Track Types", "Track Types"));
 
 	const TArrayView<RewindDebugger::FRewindDebuggerTrackType> TrackTypes = TrackTypesAttribute.Get();
 
-	for(const RewindDebugger::FRewindDebuggerTrackType& TrackType : TrackTypes)
+	for (const RewindDebugger::FRewindDebuggerTrackType& TrackType : TrackTypes)
 	{
 		Builder.AddMenuEntry(
 			TrackType.DisplayName,
-			FText::Format(LOCTEXT("FilterTrackToolTip", "Show tracks of type {0}"), {TrackType.DisplayName}),
+			FText::Format(LOCTEXT("FilterTrackToolTip", "Show tracks of type {0}"), { TrackType.DisplayName }),
 			FSlateIcon(),
-			FUIAction(FExecuteAction::CreateLambda([this,TrackType](){ ToggleHideTrackType(TrackType.Name); }),
+			FUIAction(FExecuteAction::CreateLambda([this, TrackType]() { ToggleHideTrackType(TrackType.Name); }),
 				FCanExecuteAction(),
-				FIsActionChecked::CreateLambda([this,TrackType]() { return !ShouldHideTrackType(TrackType.Name); })),
+				FIsActionChecked::CreateLambda([this, TrackType]() { return !ShouldHideTrackType(TrackType.Name); })),
 			NAME_None,
 			EUserInterfaceActionType::ToggleButton
-		);	
+		);
 	}
 
 	Builder.AddSeparator();
@@ -568,37 +559,36 @@ TSharedRef<SWidget> SRewindDebugger::MakeFilterMenu()
 		NAME_None,
 		EUserInterfaceActionType::ToggleButton
 	);
-	
+
 	return Builder.MakeWidget();
 }
 
-void SRewindDebugger::ComponentSelectionChanged(TSharedPtr<RewindDebugger::FRewindDebuggerTrack> SelectedItem, ESelectInfo::Type SelectInfo)
+void SRewindDebugger::OnSelectedTrackChanged(TSharedPtr<RewindDebugger::FRewindDebuggerTrack> SelectedItem, ESelectInfo::Type SelectInfo)
 {
 	if (SelectedComponent)
 	{
 		SelectedComponent->SetIsSelected(false);
 	}
-	
+
 	SelectedComponent = SelectedItem;
-	
+
 	if (SelectedComponent)
 	{
 		SelectedComponent->SetIsSelected(true);
 	}
-	
 
-	OnComponentSelectionChanged.ExecuteIfBound(SelectedItem);
+	OnTrackSelectionChanged.ExecuteIfBound(SelectedItem);
 }
 
 TSharedPtr<SWidget> SRewindDebugger::OnContextMenuOpening()
 {
-	return BuildComponentContextMenu.Execute();
+	return BuildTrackContextMenu.Execute();
 }
 
 FReply SRewindDebugger::OnSelectActorClicked()
 {
-	FActorPickerModeModule& ActorPickerMode = FModuleManager::Get().GetModuleChecked<FActorPickerModeModule>("ActorPickerMode");
-	
+	const FActorPickerModeModule& ActorPickerMode = FModuleManager::Get().GetModuleChecked<FActorPickerModeModule>("ActorPickerMode");
+
 	const bool bShouldForceEject = GEditor->PlayWorld && !GEditor->bIsSimulatingInEditor;
 	if (bShouldForceEject)
 	{
@@ -607,19 +597,17 @@ FReply SRewindDebugger::OnSelectActorClicked()
 	}
 
 	ActorPickerMode.BeginActorPickingMode(
-		FOnGetAllowedClasses(), 
+		FOnGetAllowedClasses(),
 		FOnShouldFilterActor(),
-		FOnActorSelected::CreateLambda([this, bShouldForceEject](AActor* InActor)
-		{
-			SetDebugTargetActor(InActor);
-			if (bShouldForceEject && GEditor->bIsSimulatingInEditor)
+		FOnActorSelected::CreateLambda([this, bShouldForceEject](const AActor* InActor)
 			{
-				// If we force ejected PIE, revert this after actor selection.
-				GEditor->RequestToggleBetweenPIEandSIE();
-			}
-		}));
-
-
+				SetDebuggedObject(InActor);
+				if (bShouldForceEject && GEditor->bIsSimulatingInEditor)
+				{
+					// If we force ejected PIE, revert this after actor selection.
+					GEditor->RequestToggleBetweenPIEandSIE();
+				}
+			}));
 
 	return FReply::Handled();
 }
diff --git a/Engine/Plugins/Animation/GameplayInsights/Source/RewindDebugger/Private/SRewindDebugger.h b/Engine/Plugins/Animation/GameplayInsights/Source/RewindDebugger/Private/SRewindDebugger.h
--- a/Engine/Plugins/Animation/GameplayInsights/Source/RewindDebugger/Private/SRewindDebugger.h
+++ b/Engine/Plugins/Animation/GameplayInsights/Source/RewindDebugger/Private/SRewindDebugger.h
@@ -2,85 +2,74 @@
 #pragma once
 
 #include "BindableProperty.h"
-#include "CoreMinimal.h"
-#include "IRewindDebuggerTrackCreator.h"
 #include "Framework/Commands/UICommandList.h"
-#include "IRewindDebuggerView.h"
+#include "IRewindDebuggerTrackCreator.h"
 #include "RewindDebuggerCommands.h"
-#include "SRewindDebuggerComponentTree.h"
+#include "RewindDebuggerSettings.h"
 #include "SRewindDebuggerTimelines.h"
+#include "SRewindDebuggerTrackTree.h"
 #include "Widgets/DeclarativeSyntaxSupport.h"
 #include "Widgets/Input/SSearchBox.h"
 #include "Widgets/SCompoundWidget.h"
-#include "RewindDebuggerSettings.h"
 
 class FRewindDebuggerModule;
+class SCheckBox;
 class SDockTab;
-
 class SSearchBox;
-class SCheckBox;
 
 class SRewindDebugger : public SCompoundWidget
 {
 	typedef TBindablePropertyInitializer<FString, BindingType_Out> DebugTargetInitializer;
 
 public:
-	DECLARE_DELEGATE_TwoParams( FOnScrubPositionChanged, double, bool )
-	DECLARE_DELEGATE_OneParam( FOnViewRangeChanged, const TRange<double>& )
-	DECLARE_DELEGATE_OneParam( FOnDebugTargetChanged, TSharedPtr<FString> )
-	DECLARE_DELEGATE_OneParam( FOnComponentDoubleClicked, TSharedPtr<RewindDebugger::FRewindDebuggerTrack> )
-	DECLARE_DELEGATE_OneParam( FOnComponentSelectionChanged, TSharedPtr<RewindDebugger::FRewindDebuggerTrack> )
-	DECLARE_DELEGATE_RetVal( TSharedPtr<SWidget>, FBuildComponentContextMenu )
-
-	SLATE_BEGIN_ARGS(SRewindDebugger) { }
-		SLATE_ARGUMENT( TArray< TSharedPtr< RewindDebugger::FRewindDebuggerTrack > >*, DebugComponents );
-		SLATE_ARGUMENT(DebugTargetInitializer, DebugTargetActor);
+	DECLARE_DELEGATE_TwoParams(FOnScrubPositionChanged, double, bool)
+	DECLARE_DELEGATE_OneParam(FOnViewRangeChanged, const TRange<double>&)
+	DECLARE_DELEGATE_OneParam(FOnDebugTargetChanged, TSharedPtr<FString>)
+	DECLARE_DELEGATE_OneParam(FOnTrackDoubleClicked, TSharedPtr<RewindDebugger::FRewindDebuggerTrack>)
+	DECLARE_DELEGATE_OneParam(FOnTrackSelectionChanged, TSharedPtr<RewindDebugger::FRewindDebuggerTrack>)
+	DECLARE_DELEGATE_RetVal(TSharedPtr<SWidget>, FBuildTrackContextMenu)
+
+	SLATE_BEGIN_ARGS(SRewindDebugger) {}
+		SLATE_ARGUMENT(TArray< TSharedPtr< RewindDebugger::FRewindDebuggerTrack > >*, Tracks);
+		SLATE_ARGUMENT(DebugTargetInitializer, DebuggedObjectName);
 		SLATE_ARGUMENT(TBindablePropertyInitializer<double>, TraceTime);
 		SLATE_ARGUMENT(TBindablePropertyInitializer<double>, RecordingDuration);
 		SLATE_ATTRIBUTE(TArrayView<RewindDebugger::FRewindDebuggerTrackType>, TrackTypes);
 		SLATE_ATTRIBUTE(double, ScrubTime);
 		SLATE_ATTRIBUTE(bool, IsPIESimulating);
-		SLATE_EVENT( FOnScrubPositionChanged, OnScrubPositionChanged);
-		SLATE_EVENT( FOnViewRangeChanged, OnViewRangeChanged);
-		SLATE_EVENT( FBuildComponentContextMenu, BuildComponentContextMenu );
-		SLATE_EVENT( FOnComponentDoubleClicked, OnComponentDoubleClicked );
-		SLATE_EVENT( FOnComponentSelectionChanged, OnComponentSelectionChanged );
+		SLATE_EVENT(FOnScrubPositionChanged, OnScrubPositionChanged);
+		SLATE_EVENT(FOnViewRangeChanged, OnViewRangeChanged);
+		SLATE_EVENT(FBuildTrackContextMenu, BuildTrackContextMenu);
+		SLATE_EVENT(FOnTrackDoubleClicked, OnTrackDoubleClicked);
+		SLATE_EVENT(FOnTrackSelectionChanged, OnTrackSelectionChanged);
 	SLATE_END_ARGS()
-	
-public:
-
-	/**
-	* Default constructor.
-	*/
-	SRewindDebugger();
-	virtual ~SRewindDebugger();
 
 	/**
-	* Constructs the application.
-	*
-	* @param InArgs The Slate argument list.
-	* @param ConstructUnderMajorTab The major tab which will contain the session front-end.
-	* @param ConstructUnderWindow The window in which this widget is being constructed.
-	* @param InStyleSet The style set to use.
-	*/
+	 * Constructs the widget.
+	 *
+	 * @param InArgs The Slate argument list.
+	 * @param CommandList The command list to use.
+	 * @param ConstructUnderMajorTab The major tab which will contain the session front-end.
+	 * @param ConstructUnderWindow The window in which this widget is being constructed.
+	 */
 	void Construct(const FArguments& InArgs, TSharedRef<FUICommandList> CommandList, const TSharedRef<SDockTab>& ConstructUnderMajorTab, const TSharedPtr<SWindow>& ConstructUnderWindow);
 
+	SRewindDebugger();
 	void TrackCursor(bool bReverse);
-	void RefreshDebugComponents();
+	void RefreshTracks();
+
 private:
 	void SetViewRange(TRange<double> NewRange);
 
-	void MainFrameCreationFinished(TSharedPtr<SWindow> InRootWindow, bool bIsNewProjectWindow);
-	
 	void ToggleHideTrackType(const FName& TrackType);
 	bool ShouldHideTrackType(const FName& TrackType) const;
 	void ToggleDisplayEmptyTracks();
 	bool ShouldDisplayEmptyTracks() const;
-	
+
 	virtual FReply OnPreviewKeyDown(const FGeometry& MyGeometry, const FKeyEvent& InKeyEvent) override;
 	virtual bool SupportsKeyboardFocus() const override { return true; }
-	
-	// Time Slider
+
+	/** Time Slider */
 	TAttribute<double> ScrubTimeAttribute;
 	TAttribute<bool> IsPIESimulating;
 	TAttribute<bool> TrackScrubbingAttribute;
@@ -92,37 +81,35 @@ private:
 	TBindableProperty<double> RecordingDuration;
 
 	TSharedPtr<FUICommandList> CommandList;
-	const FRewindDebuggerCommands & Commands;
-	
-	// debug actor selector
-	TSharedRef<SWidget> MakeSelectActorMenu();
-	void SetDebugTargetActor(AActor* Actor);
-	FReply OnSelectActorClicked();
+	const FRewindDebuggerCommands& Commands;
 
-	TBindableProperty<FString, BindingType_Out> DebugTargetActor;
+	/** debug object selector */
+	TSharedRef<SWidget> MakeObjectSelectionMenu();
+	void SetDebuggedObject(const UObject* Object);
 
-	void TraceTimeChanged(double Time);
+	FReply OnSelectActorClicked();
+
+	TBindableProperty<FString, BindingType_Out> DebuggedObjectName;
 
 	TSharedRef<SWidget> MakeMainMenu();
 	TSharedRef<SWidget> MakeFilterMenu();
-	void MakeViewsMenu(FMenuBuilder& MenuBuilder);
-	
-	// component tree view
-	TArray<TSharedPtr<RewindDebugger::FRewindDebuggerTrack>>* DebugComponents;
+
+	/** Track tree view */
+	TArray<TSharedPtr<RewindDebugger::FRewindDebuggerTrack>>* Tracks = nullptr;
 	TSharedPtr<RewindDebugger::FRewindDebuggerTrack> SelectedComponent;
-    void ComponentSelectionChanged(TSharedPtr<RewindDebugger::FRewindDebuggerTrack> SelectedItem, ESelectInfo::Type SelectInfo);
-	FBuildComponentContextMenu BuildComponentContextMenu;
-	FOnComponentSelectionChanged OnComponentSelectionChanged;
-	
-	TSharedPtr<SRewindDebuggerComponentTree> ComponentTreeView;
+	void OnSelectedTrackChanged(TSharedPtr<RewindDebugger::FRewindDebuggerTrack> SelectedItem, ESelectInfo::Type SelectInfo);
+	FBuildTrackContextMenu BuildTrackContextMenu;
+	FOnTrackSelectionChanged OnTrackSelectionChanged;
+
+	TSharedPtr<SRewindDebuggerTrackTree> TrackTreeView;
 	TSharedPtr<SRewindDebuggerTimelines> TimelinesView;
 	TSharedPtr<SWidget> OnContextMenuOpening();
-	
+
 	TSharedPtr<SSearchBox> TrackFilterBox;
 
 	bool bInExpansionChanged = false;
 	bool bInSelectionChanged = false;
 	bool bDisplayEmptyTracks = false;
 
-	URewindDebuggerSettings & Settings;
+	URewindDebuggerSettings& Settings;
 };
diff --git a/Engine/Plugins/Animation/GameplayInsights/Source/RewindDebugger/Private/SRewindDebuggerComponentTree.cpp b/Engine/Plugins/Animation/GameplayInsights/Source/RewindDebugger/Private/SRewindDebuggerComponentTree.cpp
deleted file mode 100644
diff --git a/Engine/Plugins/Animation/GameplayInsights/Source/RewindDebugger/Private/SRewindDebuggerComponentTree.h b/Engine/Plugins/Animation/GameplayInsights/Source/RewindDebugger/Private/SRewindDebuggerComponentTree.h
deleted file mode 100644
diff --git a/Engine/Plugins/Animation/GameplayInsights/Source/RewindDebugger/Private/SRewindDebuggerDetails.cpp b/Engine/Plugins/Animation/GameplayInsights/Source/RewindDebugger/Private/SRewindDebuggerDetails.cpp
--- a/Engine/Plugins/Animation/GameplayInsights/Source/RewindDebugger/Private/SRewindDebuggerDetails.cpp
+++ b/Engine/Plugins/Animation/GameplayInsights/Source/RewindDebugger/Private/SRewindDebuggerDetails.cpp
@@ -1,47 +1,11 @@
 // Copyright Epic Games, Inc. All Rights Reserved.
 
 #include "SRewindDebuggerDetails.h"
-#include "ActorPickerMode.h"
-#include "Editor.h"
-#include "Editor/EditorEngine.h"
-#include "Framework/Application/SlateApplication.h"
-#include "Framework/Docking/LayoutService.h"
-#include "Framework/MultiBox/MultiBoxBuilder.h"
-#include "Insights/IUnrealInsightsModule.h"
-#include "Modules/ModuleManager.h"
-#include "ObjectTrace.h"
-#include "RewindDebuggerStyle.h"
-#include "RewindDebuggerCommands.h"
-#include "SSimpleTimeSlider.h"
-#include "SceneOutlinerPublicTypes.h"
-#include "SceneOutlinerModule.h"
-#include "Selection.h"
-#include "Styling/SlateIconFinder.h"
-#include "ToolMenu.h"
-#include "ToolMenus.h"
-#include "Widgets/Docking/SDockTab.h"
-#include "Widgets/Images/SImage.h"
-#include "Widgets/Input/SComboButton.h"
-#include "Widgets/Input/SButton.h"
-#include "Widgets/Layout/SExpandableArea.h"
-#include "Widgets/Layout/SScrollBox.h"
-#include "RewindDebuggerViewCreators.h"
-#include "IRewindDebuggerDoubleClickHandler.h"
 
 #define LOCTEXT_NAMESPACE "SRewindDebuggerDetails"
 
-SRewindDebuggerDetails::SRewindDebuggerDetails() 
-	: SCompoundWidget()
-{ 
-}
-
-SRewindDebuggerDetails::~SRewindDebuggerDetails() 
-{
-}
-
 void SRewindDebuggerDetails::Construct(const FArguments& InArgs, const TSharedRef<SDockTab>& ConstructUnderMajorTab, const TSharedPtr<SWindow>& ConstructUnderWindow)
 {
 }
 
-
 #undef LOCTEXT_NAMESPACE
diff --git a/Engine/Plugins/Animation/GameplayInsights/Source/RewindDebugger/Private/SRewindDebuggerDetails.h b/Engine/Plugins/Animation/GameplayInsights/Source/RewindDebugger/Private/SRewindDebuggerDetails.h
--- a/Engine/Plugins/Animation/GameplayInsights/Source/RewindDebugger/Private/SRewindDebuggerDetails.h
+++ b/Engine/Plugins/Animation/GameplayInsights/Source/RewindDebugger/Private/SRewindDebuggerDetails.h
@@ -2,43 +2,26 @@
 #pragma once
 
 #include "BindableProperty.h"
-#include "CoreMinimal.h"
-#include "Framework/Commands/UICommandList.h"
 #include "Framework/Docking/TabManager.h"
-#include "IRewindDebuggerView.h"
-#include "RewindDebuggerModule.h"
-#include "SRewindDebuggerComponentTree.h"
 #include "Widgets/DeclarativeSyntaxSupport.h"
-#include "Widgets/Input/SSearchBox.h"
 #include "Widgets/SCompoundWidget.h"
 
 class SDockTab;
 
 class SRewindDebuggerDetails : public SCompoundWidget
 {
-	typedef TBindablePropertyInitializer<FString, BindingType_Out> DebugTargetInitializer;
+	typedef TBindablePropertyInitializer<FString, BindingType_Out> FDebugTargetInitializer;
 
 public:
 	SLATE_BEGIN_ARGS(SRewindDebuggerDetails) { }
 	SLATE_END_ARGS()
-	
-public:
-
-	/**
-	* Default constructor.
-	*/
-	SRewindDebuggerDetails();
-	virtual ~SRewindDebuggerDetails();
 
 	/**
-	* Constructs the application.
-	*
-	* @param InArgs The Slate argument list.
-	* @param ConstructUnderMajorTab The major tab which will contain the session front-end.
-	* @param ConstructUnderWindow The window in which this widget is being constructed.
-	* @param InStyleSet The style set to use.
-	*/
+	 * Constructs the application.
+	 *
+	 * @param InArgs The Slate argument list.
+	 * @param ConstructUnderMajorTab The major tab which will contain the session front-end.
+	 * @param ConstructUnderWindow The window in which this widget is being constructed.
+	 */
 	void Construct(const FArguments& InArgs, const TSharedRef<SDockTab>& ConstructUnderMajorTab, const TSharedPtr<SWindow>& ConstructUnderWindow);
-
-private:
 };
diff --git a/Engine/Plugins/Animation/GameplayInsights/Source/RewindDebugger/Private/SRewindDebuggerTimelines.cpp b/Engine/Plugins/Animation/GameplayInsights/Source/RewindDebugger/Private/SRewindDebuggerTimelines.cpp
--- a/Engine/Plugins/Animation/GameplayInsights/Source/RewindDebugger/Private/SRewindDebuggerTimelines.cpp
+++ b/Engine/Plugins/Animation/GameplayInsights/Source/RewindDebugger/Private/SRewindDebuggerTimelines.cpp
@@ -1,43 +1,41 @@
 // Copyright Epic Games, Inc. All Rights Reserved.
 
 #include "SRewindDebuggerTimelines.h"
-
 #include "ISequencerWidgetsModule.h"
 #include "ObjectTrace.h"
-#include "Widgets/Layout/SSpacer.h"
-
 #include "RewindDebugger.h"
 #include "RewindDebuggerStyle.h"
 #include "SSimpleTimeSlider.h"
+#include "Widgets/Layout/SSpacer.h"
 
 #define LOCTEXT_NAMESPACE "SAnimationInsights"
 
 
 class SRewindDebuggerTimelineTableRow : public STableRow<TSharedPtr<RewindDebugger::FRewindDebuggerTrack>>
 {
-	public:
-	
+public:
+
 	SLATE_BEGIN_ARGS(SRewindDebuggerTimelineTableRow) {}
 	SLATE_END_ARGS()
-	
+
 	void Construct(const FArguments& InArgs, const TSharedRef<STableViewBase>& InOwnerTableView)
 	{
-		STableRow<TSharedPtr<RewindDebugger::FRewindDebuggerTrack>>::Construct( STableRow<TSharedPtr<RewindDebugger::FRewindDebuggerTrack>>::FArguments(), InOwnerTableView);
-		Style =  &FRewindDebuggerStyle::Get().GetWidgetStyle<FTableRowStyle>("RewindDebugger.TableRow");
-	
+		STableRow::Construct(STableRow::FArguments(), InOwnerTableView);
+		Style = &FRewindDebuggerStyle::Get().GetWidgetStyle<FTableRowStyle>("RewindDebugger.TableRow");
+
 		SetExpanderArrowVisibility(EVisibility::Collapsed);
 
 		SetHover(TAttribute<bool>::CreateLambda([this]()
-		{
-			if (auto TableView = OwnerTablePtr.Pin())
 			{
-				if (const TSharedPtr<RewindDebugger::FRewindDebuggerTrack>* Track = GetItemForThis(TableView.ToSharedRef()))
+				if (auto TableView = OwnerTablePtr.Pin())
 				{
-					return (*Track)->GetIsHovered();
+					if (const TSharedPtr<RewindDebugger::FRewindDebuggerTrack>* Track = GetItemForThis(TableView.ToSharedRef()))
+					{
+						return (*Track)->GetIsHovered();
+					}
 				}
-			}
-			return false;
-		}));
+				return false;
+			}));
 	}
 
 	virtual int32 GetIndentLevel() const override
@@ -45,7 +43,7 @@ class SRewindDebuggerTimelineTableRow : public STableRow<TSharedPtr<RewindDebugg
 		// don't indent timeline tracks
 		return 0;
 	}
-	
+
 	virtual void OnMouseEnter(const FGeometry& MyGeometry, const FPointerEvent& MouseEvent) override
 	{
 		if (auto TableView = OwnerTablePtr.Pin())
@@ -55,12 +53,12 @@ class SRewindDebuggerTimelineTableRow : public STableRow<TSharedPtr<RewindDebugg
 				(*Track)->SetIsTrackHovered(true);
 			}
 		}
-		STableRow<TSharedPtr<RewindDebugger::FRewindDebuggerTrack>>::OnMouseEnter(MyGeometry, MouseEvent);
+		STableRow::OnMouseEnter(MyGeometry, MouseEvent);
 	}
-	
+
 	virtual void OnMouseLeave(const FPointerEvent& MouseEvent) override
 	{
-		STableRow<TSharedPtr<RewindDebugger::FRewindDebuggerTrack>>::OnMouseLeave(MouseEvent);
+		STableRow::OnMouseLeave(MouseEvent);
 
 		if (auto TableView = OwnerTablePtr.Pin())
 		{
@@ -72,30 +70,18 @@ class SRewindDebuggerTimelineTableRow : public STableRow<TSharedPtr<RewindDebugg
 	}
 };
 
-
-SRewindDebuggerTimelines::SRewindDebuggerTimelines() 
-	: SSimpleTimeSlider()
-	, DebugComponents(nullptr)
-{ 
-}
-
-SRewindDebuggerTimelines::~SRewindDebuggerTimelines() 
-{
-
-}
-
-void SRewindDebuggerTimelines::SetSelection(TSharedPtr<RewindDebugger::FRewindDebuggerTrack> SelectedItem)
+void SRewindDebuggerTimelines::SetSelection(TSharedPtr<RewindDebugger::FRewindDebuggerTrack> SelectedItem) const
 {
-	ComponentTreeView->SetSelection(SelectedItem);
+	TreeView->SetSelection(SelectedItem);
 }
 
-void SRewindDebuggerTimelines::ScrollTo(double ScrollOffset)
+void SRewindDebuggerTimelines::ScrollTo(double ScrollOffset) const
 {
-	ComponentTreeView->SetScrollOffset(ScrollOffset);
+	TreeView->SetScrollOffset(ScrollOffset);
 }
 
 
-TSharedRef<ITableRow> SRewindDebuggerTimelines::ComponentTreeViewGenerateRow(TSharedPtr<RewindDebugger::FRewindDebuggerTrack> InItem, const TSharedRef<STableViewBase>& OwnerTable)
+TSharedRef<ITableRow> SRewindDebuggerTimelines::TreeViewGenerateRow(TSharedPtr<RewindDebugger::FRewindDebuggerTrack> InItem, const TSharedRef<STableViewBase>& OwnerTable)
 {
 	TSharedPtr<SRewindDebuggerTimelineTableRow> Track = SNew(SRewindDebuggerTimelineTableRow, OwnerTable);
 	TSharedPtr<SWidget> TimelineView = InItem->GetTimelineView();
@@ -113,15 +99,15 @@ TSharedRef<ITableRow> SRewindDebuggerTimelines::ComponentTreeViewGenerateRow(TSh
 void TimelineViewGetChildren(TSharedPtr<RewindDebugger::FRewindDebuggerTrack> InItem, TArray<TSharedPtr<RewindDebugger::FRewindDebuggerTrack>>& OutChildren)
 {
 	InItem->IterateSubTracks([&OutChildren](TSharedPtr<RewindDebugger::FRewindDebuggerTrack> Track)
-	{
-		if (Track->IsVisible())
 		{
-			OutChildren.Add(Track);
-		}
-	});
+			if (Track->IsVisible())
+			{
+				OutChildren.Add(Track);
+			}
+		});
 }
 
-void SRewindDebuggerTimelines::TimelineViewExpansionChanged(TSharedPtr<RewindDebugger::FRewindDebuggerTrack> InItem, bool bShouldBeExpanded)
+void SRewindDebuggerTimelines::TimelineViewExpansionChanged(TSharedPtr<RewindDebugger::FRewindDebuggerTrack> InItem, bool bShouldBeExpanded) const
 {
 	InItem->SetIsExpanded(bShouldBeExpanded);
 	OnExpansionChanged.ExecuteIfBound();
@@ -136,31 +122,31 @@ void SRewindDebuggerTimelines::Construct(const FArguments& InArgs)
 		.ScrubPosition(InArgs._ScrubPosition)
 		.OnScrubPositionChanged(InArgs._OnScrubPositionChanged));
 
-	DebugComponents = InArgs._DebugComponents;
+	Tracks = InArgs._Tracks;
 
 	OnExpansionChanged = InArgs._OnExpansionChanged;
 
-	ComponentTreeView = SNew(STreeView<TSharedPtr<RewindDebugger::FRewindDebuggerTrack>>)
-									.OnSelectionChanged(InArgs._OnSelectionChanged)
-									.TreeItemsSource(DebugComponents)
-									.OnGenerateRow(this, &SRewindDebuggerTimelines::ComponentTreeViewGenerateRow)
-									.OnGetChildren_Static(&TimelineViewGetChildren)
-									.OnExpansionChanged(this, &SRewindDebuggerTimelines::TimelineViewExpansionChanged)
-									.SelectionMode(ESelectionMode::Single)
-									.AllowOverscroll(EAllowOverscroll::No)
-									.ExternalScrollbar(InArgs._ExternalScrollbar)
-									.OnTreeViewScrolled(InArgs._OnScrolled);
+	TreeView = SNew(STreeView<TSharedPtr<RewindDebugger::FRewindDebuggerTrack>>)
+		.OnSelectionChanged(InArgs._OnSelectionChanged)
+		.TreeItemsSource(Tracks)
+		.OnGenerateRow(this, &SRewindDebuggerTimelines::TreeViewGenerateRow)
+		.OnGetChildren_Static(&TimelineViewGetChildren)
+		.OnExpansionChanged(this, &SRewindDebuggerTimelines::TimelineViewExpansionChanged)
+		.SelectionMode(ESelectionMode::Single)
+		.AllowOverscroll(EAllowOverscroll::No)
+		.ExternalScrollbar(InArgs._ExternalScrollbar)
+		.OnTreeViewScrolled(InArgs._OnScrolled);
 
 	ChildSlot
-	[
-		ComponentTreeView.ToSharedRef()
-	];
+		[
+			TreeView.ToSharedRef()
+		];
 }
 
-int32 SRewindDebuggerTimelines::OnPaint( const FPaintArgs& Args, const FGeometry& AllottedGeometry, const FSlateRect& MyCullingRect, FSlateWindowElementList& OutDrawElements, int32 LayerId, const FWidgetStyle& InWidgetStyle, bool bParentEnabled ) const
+int32 SRewindDebuggerTimelines::OnPaint(const FPaintArgs& Args, const FGeometry& AllottedGeometry, const FSlateRect& MyCullingRect, FSlateWindowElementList& OutDrawElements, int32 LayerId, const FWidgetStyle& InWidgetStyle, bool bParentEnabled) const
 {
 	TRange<double> LocalViewRange = ViewRange.Get();
-	FScrubRangeToScreen RangeToScreen( LocalViewRange, AllottedGeometry.GetLocalSize() );
+	FScrubRangeToScreen RangeToScreen(LocalViewRange, AllottedGeometry.GetLocalSize());
 
 	// Paint Child Widgets (Tracks) 
 	LayerId = SCompoundWidget::OnPaint(Args, AllottedGeometry, MyCullingRect, OutDrawElements, LayerId, InWidgetStyle, ShouldBeEnabled(bParentEnabled));
@@ -169,29 +155,28 @@ int32 SRewindDebuggerTimelines::OnPaint( const FPaintArgs& Args, const FGeometry
 	FDrawTickArgs TickArgs;
 	TickArgs.AllottedGeometry = AllottedGeometry;
 	TickArgs.bOnlyDrawMajorTicks = true;
-	TickArgs.TickColor = FLinearColor(0.2f,0.2f,0.2f,0.2f);
+	TickArgs.TickColor = FLinearColor(0.2f, 0.2f, 0.2f, 0.2f);
 	TickArgs.ClippingRect = MyCullingRect;
 	TickArgs.DrawEffects = ESlateDrawEffect::None;
 	TickArgs.StartLayer = ++LayerId;
 	TickArgs.TickOffset = 0;
 	TickArgs.MajorTickHeight = AllottedGeometry.Size.Y;
 	DrawTicks(OutDrawElements, RangeToScreen, TickArgs);
-    	
+
 	// Paint Cursor
-	const float XPos = RangeToScreen.InputToLocalX( ScrubPosition.Get() );
+	const float XPos = RangeToScreen.InputToLocalX(ScrubPosition.Get());
 
-	TArray<FVector2D> Points {{XPos, 0}, { XPos, AllottedGeometry.Size.Y}};
-	
+	TArray<FVector2D> Points{ {XPos, 0}, { XPos, AllottedGeometry.Size.Y} };
 
 	FSlateDrawElement::MakeLines(
-				OutDrawElements,
-						++LayerId,
-						AllottedGeometry.ToPaintGeometry(),
-						Points,
-						ESlateDrawEffect::None,
-						FLinearColor::White,
-						false
-						);
+		OutDrawElements,
+		++LayerId,
+		AllottedGeometry.ToPaintGeometry(),
+		Points,
+		ESlateDrawEffect::None,
+		FLinearColor::White,
+		false
+	);
 
 	return LayerId;
 }
@@ -210,17 +195,17 @@ static void RestoreTimelineExpansion(TSharedPtr<RewindDebugger::FRewindDebuggerT
 
 void SRewindDebuggerTimelines::RestoreExpansion()
 {
-	for (auto& Track : *DebugComponents)
+	for (auto& Track : *Tracks)
 	{
-		::RestoreTimelineExpansion(Track, ComponentTreeView);
+		::RestoreTimelineExpansion(Track, TreeView);
 	}
 }
 
 void SRewindDebuggerTimelines::Refresh()
 {
-	ComponentTreeView->RebuildList();
+	TreeView->RebuildList();
 
-	if (DebugComponents)
+	if (Tracks)
 	{
 		// make sure any newly added TreeView nodes are created expanded
 		RestoreExpansion();
diff --git a/Engine/Plugins/Animation/GameplayInsights/Source/RewindDebugger/Private/SRewindDebuggerTimelines.h b/Engine/Plugins/Animation/GameplayInsights/Source/RewindDebugger/Private/SRewindDebuggerTimelines.h
--- a/Engine/Plugins/Animation/GameplayInsights/Source/RewindDebugger/Private/SRewindDebuggerTimelines.h
+++ b/Engine/Plugins/Animation/GameplayInsights/Source/RewindDebugger/Private/SRewindDebuggerTimelines.h
@@ -1,56 +1,48 @@
 // Copyright Epic Games, Inc. All Rights Reserved.
 #pragma once
 
-#include "CoreMinimal.h"
-#include "IRewindDebugger.h"
 #include "RewindDebuggerTrack.h"
+#include "SSimpleTimeSlider.h"
 #include "Widgets/DeclarativeSyntaxSupport.h"
 #include "Widgets/SCompoundWidget.h"
 #include "Widgets/Views/STreeView.h"
-#include "SSimpleTimeSlider.h"
 
 
 class SRewindDebuggerTimelines : public SSimpleTimeSlider
 {
-	using FOnSelectionChanged = typename STreeView<TSharedPtr<RewindDebugger::FRewindDebuggerTrack>>::FOnSelectionChanged;
+	using FOnSelectionChanged = STreeView<TSharedPtr<RewindDebugger::FRewindDebuggerTrack>>::FOnSelectionChanged;
 public:
 	SLATE_BEGIN_ARGS(SRewindDebuggerTimelines) {}
-	 
-		SLATE_ARGUMENT( TArray< TSharedPtr< RewindDebugger::FRewindDebuggerTrack > >*, DebugComponents );
-		SLATE_ARGUMENT( TSharedPtr<SScrollBar>, ExternalScrollbar)
-		SLATE_EVENT( FOnSelectionChanged, OnSelectionChanged )
-		SLATE_EVENT( FOnTableViewScrolled, OnScrolled)
-		SLATE_EVENT( FSimpleDelegate, OnExpansionChanged )
-
-	
+
+		SLATE_ARGUMENT(TArray< TSharedPtr< RewindDebugger::FRewindDebuggerTrack > >*, Tracks);
+		SLATE_ARGUMENT(TSharedPtr<SScrollBar>, ExternalScrollbar)
+		SLATE_EVENT(FOnSelectionChanged, OnSelectionChanged)
+		SLATE_EVENT(FOnTableViewScrolled, OnScrolled)
+		SLATE_EVENT(FSimpleDelegate, OnExpansionChanged)
+
+
 		/** The scrub position */
-    	SLATE_ATTRIBUTE(double, ScrubPosition);
-    
-    	/** View time range */
-    	SLATE_ATTRIBUTE(TRange<double>, ViewRange);
-    
-    	/** Clamp time range */
-    	SLATE_ATTRIBUTE(TRange<double>, ClampRange);
-    
+		SLATE_ATTRIBUTE(double, ScrubPosition);
+
+		/** View time range */
+		SLATE_ATTRIBUTE(TRange<double>, ViewRange);
+
+		/** Clamp time range */
+		SLATE_ATTRIBUTE(TRange<double>, ClampRange);
+
 		/** Called when the scrub position changes */
 		SLATE_EVENT(FOnScrubPositionChanged, OnScrubPositionChanged);
-	
+
 		/** Called right before the scrubber begins to move */
 		SLATE_EVENT(FSimpleDelegate, OnBeginScrubberMovement);
-	
+
 		/** Called right after the scrubber handle is released by the user */
 		SLATE_EVENT(FSimpleDelegate, OnEndScrubberMovement);
-	
-    	/** Called when the view range changes */
-    	SLATE_EVENT(FOnRangeChanged, OnViewRangeChanged);
-    	
-	SLATE_END_ARGS()
 
-	/**
-	* Default constructor.
-	*/
-	SRewindDebuggerTimelines();
-	virtual ~SRewindDebuggerTimelines();
+		/** Called when the view range changes */
+		SLATE_EVENT(FOnRangeChanged, OnViewRangeChanged);
+
+	SLATE_END_ARGS()
 
 	/**
 	* Constructs the application.
@@ -62,17 +54,17 @@ public:
 	void Refresh();
 	void RestoreExpansion();
 
-	void SetSelection(TSharedPtr<RewindDebugger::FRewindDebuggerTrack> SelectedItem);
-	void ScrollTo(double ScrollOffset);
-	
-	virtual int32 OnPaint( const FPaintArgs& Args, const FGeometry& AllottedGeometry, const FSlateRect& MyCullingRect, FSlateWindowElementList& OutDrawElements, int32 LayerId, const FWidgetStyle& InWidgetStyle, bool bParentEnabled ) const override;
+	void SetSelection(TSharedPtr<RewindDebugger::FRewindDebuggerTrack> SelectedItem) const;
+	void ScrollTo(double ScrollOffset) const;
+
+	virtual int32 OnPaint(const FPaintArgs& Args, const FGeometry& AllottedGeometry, const FSlateRect& MyCullingRect, FSlateWindowElementList& OutDrawElements, int32 LayerId, const FWidgetStyle& InWidgetStyle, bool bParentEnabled) const override;
 
 private:
-	TSharedRef<ITableRow> ComponentTreeViewGenerateRow(TSharedPtr<RewindDebugger::FRewindDebuggerTrack> InItem, const TSharedRef<STableViewBase>& OwnerTable);
-	void TimelineViewExpansionChanged(TSharedPtr<RewindDebugger::FRewindDebuggerTrack> InItem, bool bShouldBeExpanded);
-	
-	TArray<TSharedPtr<RewindDebugger::FRewindDebuggerTrack>>* DebugComponents;
-	TSharedPtr<STreeView<TSharedPtr<RewindDebugger::FRewindDebuggerTrack>>> ComponentTreeView;
-	
+	TSharedRef<ITableRow> TreeViewGenerateRow(TSharedPtr<RewindDebugger::FRewindDebuggerTrack> InItem, const TSharedRef<STableViewBase>& OwnerTable);
+	void TimelineViewExpansionChanged(TSharedPtr<RewindDebugger::FRewindDebuggerTrack> InItem, bool bShouldBeExpanded) const;
+
+	TArray<TSharedPtr<RewindDebugger::FRewindDebuggerTrack>>* Tracks = nullptr;
+	TSharedPtr<STreeView<TSharedPtr<RewindDebugger::FRewindDebuggerTrack>>> TreeView;
+
 	FSimpleDelegate OnExpansionChanged;
 };
diff --git a/Engine/Plugins/Animation/GameplayInsights/Source/RewindDebugger/Private/SRewindDebuggerTrackTree.cpp b/Engine/Plugins/Animation/GameplayInsights/Source/RewindDebugger/Private/SRewindDebuggerTrackTree.cpp
new file mode 100644
--- /dev/null
+++ b/Engine/Plugins/Animation/GameplayInsights/Source/RewindDebugger/Private/SRewindDebuggerTrackTree.cpp
@@ -0,0 +1,184 @@
+// Copyright Epic Games, Inc. All Rights Reserved.
+
+#include "SRewindDebuggerTrackTree.h"
+#include "ISequencerWidgetsModule.h"
+#include "ObjectTrace.h"
+#include "RewindDebugger.h"
+#include "RewindDebuggerStyle.h"
+#include "Styling/SlateIconFinder.h"
+#include "Widgets/Images/SImage.h"
+#include "Widgets/Images/SLayeredImage.h"
+
+#define LOCTEXT_NAMESPACE "SAnimationInsights"
+
+class SRewindDebuggerComponentTreeTableRow final : public STableRow<TSharedPtr<RewindDebugger::FRewindDebuggerTrack>>
+{
+public:
+
+	SLATE_BEGIN_ARGS(SRewindDebuggerComponentTreeTableRow) {}
+	SLATE_END_ARGS()
+
+	void Construct(const FArguments& InArgs, const TSharedRef<STableViewBase>& InOwnerTableView)
+	{
+		STableRow<TSharedPtr<RewindDebugger::FRewindDebuggerTrack>>::Construct(STableRow<TSharedPtr<RewindDebugger::FRewindDebuggerTrack>>::FArguments(), InOwnerTableView);
+		Style = &FRewindDebuggerStyle::Get().GetWidgetStyle<FTableRowStyle>("RewindDebugger.TableRow");
+
+		SetHover(TAttribute<bool>::CreateLambda([this]()
+			{
+				if (TSharedPtr<ITypedTableView<TSharedPtr<RewindDebugger::FRewindDebuggerTrack>>> TableView = OwnerTablePtr.Pin())
+				{
+					if (const TSharedPtr<RewindDebugger::FRewindDebuggerTrack>* Track = GetItemForThis(TableView.ToSharedRef()))
+					{
+						return (*Track)->GetIsHovered();
+					}
+				}
+
+				return false;
+			}));
+	}
+
+	virtual void OnMouseEnter(const FGeometry& MyGeometry, const FPointerEvent& MouseEvent) override
+	{
+		if (TSharedPtr<ITypedTableView<TSharedPtr<RewindDebugger::FRewindDebuggerTrack>>> TableView = OwnerTablePtr.Pin())
+		{
+			if (const TSharedPtr<RewindDebugger::FRewindDebuggerTrack>* Track = GetItemForThis(TableView.ToSharedRef()))
+			{
+				(*Track)->SetIsTreeHovered(true);
+			}
+		}
+		STableRow<TSharedPtr<RewindDebugger::FRewindDebuggerTrack>>::OnMouseEnter(MyGeometry, MouseEvent);
+	}
+
+	virtual void OnMouseLeave(const FPointerEvent& MouseEvent) override
+	{
+		STableRow<TSharedPtr<RewindDebugger::FRewindDebuggerTrack>>::OnMouseLeave(MouseEvent);
+
+		if (TSharedPtr<ITypedTableView<TSharedPtr<RewindDebugger::FRewindDebuggerTrack>>> TableView = OwnerTablePtr.Pin())
+		{
+			if (const TSharedPtr<RewindDebugger::FRewindDebuggerTrack>* Track = GetItemForThis(TableView.ToSharedRef()))
+			{
+				(*Track)->SetIsTreeHovered(false);
+			}
+		}
+	}
+};
+
+
+TSharedRef<ITableRow> SRewindDebuggerTrackTree::TreeViewGenerateRow(TSharedPtr<RewindDebugger::FRewindDebuggerTrack> InItem, const TSharedRef<STableViewBase>& OwnerTable) const
+{
+	const FSlateIcon ObjectIcon = InItem->GetIcon();
+
+	const TSharedRef<SLayeredImage> LayeredIcons = SNew(SLayeredImage)
+		.DesiredSizeOverride(FVector2D(16, 16))
+		.Image(ObjectIcon.GetIcon());
+
+	if (ObjectIcon.GetOverlayIcon())
+	{
+		LayeredIcons->AddLayer(ObjectIcon.GetOverlayIcon());
+	}
+
+	TSharedRef<SRewindDebuggerComponentTreeTableRow> Row = SNew(SRewindDebuggerComponentTreeTableRow, OwnerTable);
+
+
+	Row->SetContent(
+		SNew(SHorizontalBox)
+		+ SHorizontalBox::Slot().AutoWidth().Padding(2)
+		[
+			LayeredIcons
+		]
+		+ SHorizontalBox::Slot()
+		.FillWidth(1.0f)
+		.Padding(1.0f, 0.0f)
+		.VAlign(VAlign_Center)
+		[
+			SNew(STextBlock)
+				.Text(InItem->GetDisplayName())
+				.Font_Lambda([this, InItem]() { return InItem->GetIsHovered() ? FCoreStyle::GetDefaultFontStyle("Bold", 10) : FCoreStyle::GetDefaultFontStyle("Regular", 10); })
+		]
+	);
+
+	return Row;
+}
+
+void TreeViewGetChildren(TSharedPtr<RewindDebugger::FRewindDebuggerTrack> InItem, TArray<TSharedPtr<RewindDebugger::FRewindDebuggerTrack>>& OutChildren)
+{
+	InItem->IterateSubTracks([&OutChildren](const TSharedPtr<RewindDebugger::FRewindDebuggerTrack>& Track)
+		{
+			if (Track->IsVisible())
+			{
+				OutChildren.Add(Track);
+			}
+		});
+}
+
+void SRewindDebuggerTrackTree::TreeViewExpansionChanged(TSharedPtr<RewindDebugger::FRewindDebuggerTrack> InItem, const bool bShouldBeExpanded) const
+{
+	InItem->SetIsExpanded(bShouldBeExpanded);
+	OnExpansionChanged.ExecuteIfBound();
+}
+
+void SRewindDebuggerTrackTree::SetSelection(const TSharedPtr<RewindDebugger::FRewindDebuggerTrack>& SelectedItem) const
+{
+	TreeView->SetSelection(SelectedItem);
+}
+
+void SRewindDebuggerTrackTree::ScrollTo(const double ScrollOffset) const
+{
+	TreeView->SetScrollOffset(ScrollOffset);
+}
+
+void SRewindDebuggerTrackTree::Construct(const FArguments& InArgs)
+{
+	Tracks = InArgs._Tracks;
+
+	TreeView = SNew(STreeView<TSharedPtr<RewindDebugger::FRewindDebuggerTrack>>)
+		.TreeItemsSource(Tracks)
+		.OnGenerateRow(this, &SRewindDebuggerTrackTree::TreeViewGenerateRow)
+		.OnGetChildren_Static(&TreeViewGetChildren)
+		.OnExpansionChanged(this, &SRewindDebuggerTrackTree::TreeViewExpansionChanged)
+		.SelectionMode(ESelectionMode::Single)
+		.OnSelectionChanged(InArgs._OnSelectionChanged)
+		.OnMouseButtonDoubleClick(InArgs._OnMouseButtonDoubleClick)
+		.ExternalScrollbar(InArgs._ExternalScrollBar)
+		.AllowOverscroll(EAllowOverscroll::No)
+		.OnTreeViewScrolled(InArgs._OnScrolled)
+		.ScrollbarDragFocusCause(EFocusCause::SetDirectly)
+		.OnContextMenuOpening(InArgs._OnContextMenuOpening);
+
+	ChildSlot
+		[
+			TreeView.ToSharedRef()
+		];
+
+	OnExpansionChanged = InArgs._OnExpansionChanged;
+}
+
+static void RestoreExpansion(const TSharedPtr<RewindDebugger::FRewindDebuggerTrack>& Track, TSharedPtr<STreeView<TSharedPtr<RewindDebugger::FRewindDebuggerTrack>>>& TreeView)
+{
+	TreeView->SetItemExpansion(Track, Track->GetIsExpanded());
+	Track->IterateSubTracks([&TreeView](const TSharedPtr<RewindDebugger::FRewindDebuggerTrack>& SubTrack)
+		{
+			RestoreExpansion(SubTrack, TreeView);
+		});
+}
+
+void SRewindDebuggerTrackTree::RestoreExpansion()
+{
+	for (TSharedPtr<RewindDebugger::FRewindDebuggerTrack>& Track : *Tracks)
+	{
+		::RestoreExpansion(Track, TreeView);
+	}
+}
+
+void SRewindDebuggerTrackTree::Refresh()
+{
+	TreeView->RebuildList();
+
+	if (Tracks)
+	{
+		// make sure any newly added TreeView nodes are created expanded
+		RestoreExpansion();
+	}
+}
+
+#undef LOCTEXT_NAMESPACE
diff --git a/Engine/Plugins/Animation/GameplayInsights/Source/RewindDebugger/Private/SRewindDebuggerTrackTree.h b/Engine/Plugins/Animation/GameplayInsights/Source/RewindDebugger/Private/SRewindDebuggerTrackTree.h
new file mode 100644
--- /dev/null
+++ b/Engine/Plugins/Animation/GameplayInsights/Source/RewindDebugger/Private/SRewindDebuggerTrackTree.h
@@ -0,0 +1,47 @@
+// Copyright Epic Games, Inc. All Rights Reserved.
+#pragma once
+
+#include "IRewindDebugger.h"
+#include "RewindDebuggerTrack.h"
+#include "Widgets/DeclarativeSyntaxSupport.h"
+#include "Widgets/SCompoundWidget.h"
+#include "Widgets/Views/STreeView.h"
+
+class SRewindDebuggerTrackTree : public SCompoundWidget
+{
+	using FOnSelectionChanged = STreeView<TSharedPtr<RewindDebugger::FRewindDebuggerTrack>>::FOnSelectionChanged;
+	using FOnMouseButtonDoubleClick = STreeView<TSharedPtr<RewindDebugger::FRewindDebuggerTrack>>::FOnMouseButtonDoubleClick;
+
+public:
+	SLATE_BEGIN_ARGS(SRewindDebuggerTrackTree) {}
+		SLATE_ARGUMENT(TArray< TSharedPtr< RewindDebugger::FRewindDebuggerTrack > >*, Tracks);
+		SLATE_ARGUMENT(TSharedPtr< SScrollBar >, ExternalScrollBar);
+		SLATE_EVENT(FOnSelectionChanged, OnSelectionChanged)
+		SLATE_EVENT(FOnMouseButtonDoubleClick, OnMouseButtonDoubleClick)
+		SLATE_EVENT(FOnContextMenuOpening, OnContextMenuOpening)
+		SLATE_EVENT(FSimpleDelegate, OnExpansionChanged)
+		SLATE_EVENT(FOnTableViewScrolled, OnScrolled)
+	SLATE_END_ARGS()
+
+	/**
+	 * Constructs the application.
+	 *
+	 * @param InArgs The Slate argument list.
+	 */
+	void Construct(const FArguments& InArgs);
+
+	void Refresh();
+	void RestoreExpansion();
+
+	void SetSelection(const TSharedPtr<RewindDebugger::FRewindDebuggerTrack>& SelectedItem) const;
+	void ScrollTo(double ScrollOffset) const;
+
+private:
+	TSharedRef<ITableRow> TreeViewGenerateRow(TSharedPtr<RewindDebugger::FRewindDebuggerTrack> InItem, const TSharedRef<STableViewBase>& OwnerTable) const;
+	void TreeViewExpansionChanged(TSharedPtr<RewindDebugger::FRewindDebuggerTrack> InItem, bool bShouldBeExpanded) const;
+
+	TArray<TSharedPtr<RewindDebugger::FRewindDebuggerTrack>>* Tracks = nullptr;
+	TSharedPtr<STreeView<TSharedPtr<RewindDebugger::FRewindDebuggerTrack>>> TreeView;
+
+	FSimpleDelegate OnExpansionChanged;
+};
diff --git a/Engine/Plugins/Animation/GameplayInsights/Source/RewindDebuggerVLog/Private/RewindDebuggerVLog.cpp b/Engine/Plugins/Animation/GameplayInsights/Source/RewindDebuggerVLog/Private/RewindDebuggerVLog.cpp
--- a/Engine/Plugins/Animation/GameplayInsights/Source/RewindDebuggerVLog/Private/RewindDebuggerVLog.cpp
+++ b/Engine/Plugins/Animation/GameplayInsights/Source/RewindDebuggerVLog/Private/RewindDebuggerVLog.cpp
@@ -22,11 +22,6 @@
 
 TAutoConsoleVariable<int32> CVarRewindDebuggerVLogUseActor(TEXT("a.RewindDebugger.VisualLogs.UseActor"), 0, TEXT("Use actor based debug renderer for visual logs"));
 
-FRewindDebuggerVLog::FRewindDebuggerVLog()
-{
-
-}
-
 void FRewindDebuggerVLog::OnShowDebugInfo(UCanvas* Canvas, APlayerController* Player)
 {
 	ScreenTextY = 60;
@@ -56,19 +51,16 @@ void FRewindDebuggerVLog::OnShowDebugInfo(UCanvas* Canvas, APlayerController* Pl
 				const TraceServices::IFrameProvider& FrameProvider = TraceServices::ReadFrameProvider(*Session);
 				TraceServices::FFrame CurrentFrame;
 
-				if(FrameProvider.GetFrameFromTime(ETraceFrameType::TraceFrameType_Game, CurrentTraceTime, CurrentFrame))
+				if (FrameProvider.GetFrameFromTime(ETraceFrameType::TraceFrameType_Game, CurrentTraceTime, CurrentFrame))
 				{
 					if (const IVisualLoggerProvider* VisualLoggerProvider = Session->ReadProvider<IVisualLoggerProvider>("VisualLoggerProvider"))
 					{
-						AddLogEntries(RewindDebugger->GetDebugComponents(), CurrentFrame.StartTime, CurrentFrame.EndTime, VisualLoggerProvider, Canvas);
+						AddLogEntries(RewindDebugger->GetDebuggedObjects(), CurrentFrame.StartTime, CurrentFrame.EndTime, VisualLoggerProvider, Canvas);
 					}
 				}
 			}
 		}
 	}
-
-
-	
 }
 
 FRewindDebuggerVLog::~FRewindDebuggerVLog()
@@ -80,7 +72,7 @@ void FRewindDebuggerVLog::Initialize()
 {
 	UToolMenu* Menu = UToolMenus::Get()->FindMenu("RewindDebugger.ToolBar");
 
-	FToolMenuSection& NewSection = Menu->AddSection("Visual Logger", LOCTEXT("Visual Logger","Visual Logger"));
+	FToolMenuSection& NewSection = Menu->AddSection("Visual Logger", LOCTEXT("Visual Logger", "Visual Logger"));
 
 	NewSection.AddSeparator("VisualLogger");
 
@@ -89,15 +81,14 @@ void FRewindDebuggerVLog::Initialize()
 		FUIAction(),
 		FNewToolMenuDelegate::CreateRaw(this, &FRewindDebuggerVLog::MakeCategoriesMenu),
 		LOCTEXT("VLog Categories", "VLog Categories")
-		));
-	
+	));
+
 	NewSection.AddEntry(FToolMenuEntry::InitComboButton(
 		"VLog Level",
 		FUIAction(),
 		FNewToolMenuDelegate::CreateRaw(this, &FRewindDebuggerVLog::MakeLogLevelMenu),
 		LOCTEXT("VLog Level", "VLog Level")));
-		
-		
+
 	FVisualLoggerTraceDevice& TraceDevice = FVisualLoggerTraceDevice::Get();
 	TraceDevice.ImmediateRenderDelegate.BindRaw(this, &FRewindDebuggerVLog::ImmediateRender);
 
@@ -117,7 +108,7 @@ void FRewindDebuggerVLog::RenderLogEntry(const FVisualLogEntry& Entry, UCanvas*
 {
 	if (CVarRewindDebuggerVLogUseActor.GetValueOnAnyThread())
 	{
-		// old actor based codepath
+		// old actor based code path
 		if (AVLogRenderingActor* RenderingActor = GetRenderingActor())
 		{
 			RenderingActor->AddLogEntry(Entry);
@@ -126,7 +117,7 @@ void FRewindDebuggerVLog::RenderLogEntry(const FVisualLogEntry& Entry, UCanvas*
 	else
 	{
 		UWorld* World = IRewindDebugger::Instance()->GetWorldToVisualize();
-		FVisualLogEntryRenderer::RenderLogEntry(World,Entry, &MatchCategoryFilters, Canvas, GEngine->GetMediumFont(), MonospaceFont, ScreenTextY);
+		FVisualLogEntryRenderer::RenderLogEntry(World, Entry, &MatchCategoryFilters, Canvas, GEngine->GetMediumFont(), MonospaceFont, ScreenTextY);
 
 	}
 }
@@ -137,7 +128,7 @@ void FRewindDebuggerVLog::ImmediateRender(const UObject* Object, const FVisualLo
 	if (IRewindDebugger* RewindDebugger = IRewindDebugger::Instance())
 	{
 		uint64 ObjectId = FObjectTrace::GetObjectId(Object);
-		if (RewindDebugger->IsContainedByDebugComponent(ObjectId))
+		if (RewindDebugger->IsObjectCurrentlyDebugged(ObjectId))
 		{
 			ImmediateRenderQueue.Add(Entry);
 		}
@@ -157,7 +148,7 @@ void FRewindDebuggerVLog::ToggleCategory(const FName& Category)
 
 }
 
-ELogVerbosity::Type FRewindDebuggerVLog::GetMinLogVerbosity() const
+ELogVerbosity::Type FRewindDebuggerVLog::GetMinLogVerbosity()
 {
 	return static_cast<ELogVerbosity::Type>(URewindDebuggerVLogSettings::Get().DisplayVerbosity);
 }
@@ -168,15 +159,15 @@ void FRewindDebuggerVLog::SetMinLogVerbosity(ELogVerbosity::Type Value)
 
 }
 
-void FRewindDebuggerVLog::MakeLogLevelMenu(UToolMenu* Menu)
+void FRewindDebuggerVLog::MakeLogLevelMenu(UToolMenu* Menu) const
 {
 	FToolMenuSection& Section = Menu->AddSection("Levels");
 
-	for(ELogVerbosity::Type LogVerbosityLevel = ELogVerbosity::All; LogVerbosityLevel > 0; LogVerbosityLevel = static_cast<ELogVerbosity::Type>(LogVerbosityLevel-1))
+	for (ELogVerbosity::Type LogVerbosityLevel = ELogVerbosity::All; LogVerbosityLevel > 0; LogVerbosityLevel = static_cast<ELogVerbosity::Type>(LogVerbosityLevel - 1))
 	{
 		FString Name = ToString(LogVerbosityLevel);
 		FText Label(FText::FromString(Name));
-		
+
 		Section.AddMenuEntry(FName(Name),
 			Label,
 			FText(),
@@ -184,14 +175,14 @@ void FRewindDebuggerVLog::MakeLogLevelMenu(UToolMenu* Menu)
 			FUIAction(
 				FExecuteAction::CreateLambda([LogVerbosityLevel, this]() { SetMinLogVerbosity(LogVerbosityLevel); }),
 				FCanExecuteAction(),
-				FGetActionCheckState::CreateLambda( [LogVerbosityLevel, this]() { return GetMinLogVerbosity() == LogVerbosityLevel ? ECheckBoxState::Checked : ECheckBoxState::Unchecked; } )
+				FGetActionCheckState::CreateLambda([LogVerbosityLevel, this]() { return GetMinLogVerbosity() == LogVerbosityLevel ? ECheckBoxState::Checked : ECheckBoxState::Unchecked; })
 			),
 			EUserInterfaceActionType::Check
-			);
+		);
 	}
 }
 
-void FRewindDebuggerVLog::MakeCategoriesMenu(UToolMenu* Menu)
+void FRewindDebuggerVLog::MakeCategoriesMenu(UToolMenu* Menu) const
 {
 	FToolMenuSection& Section = Menu->AddSection("Categories");
 	IRewindDebugger* RewindDebugger = IRewindDebugger::Instance();
@@ -201,41 +192,42 @@ void FRewindDebuggerVLog::MakeCategoriesMenu(UToolMenu* Menu)
 		if (const IVisualLoggerProvider* VisualLoggerProvider = Session->ReadProvider<IVisualLoggerProvider>("VisualLoggerProvider"))
 		{
 			VisualLoggerProvider->EnumerateCategories([&Section, this](const FName& Category)
-			{
-				Section.AddMenuEntry(Category,
-					FText::FromName(Category),
-					FText(),
-					FSlateIcon(),
-					FUIAction(
-						FExecuteAction::CreateLambda([Category, this]() { ToggleCategory(Category); }),
-						FCanExecuteAction(),
-						FGetActionCheckState::CreateLambda( [Category, this]() { return IsCategoryActive(Category) ? ECheckBoxState::Checked : ECheckBoxState::Unchecked; } )
-					),
-					EUserInterfaceActionType::Check
+				{
+					Section.AddMenuEntry(Category,
+						FText::FromName(Category),
+						FText(),
+						FSlateIcon(),
+						FUIAction(
+							FExecuteAction::CreateLambda([Category, this]() { ToggleCategory(Category); }),
+							FCanExecuteAction(),
+							FGetActionCheckState::CreateLambda([Category, this]() { return IsCategoryActive(Category) ? ECheckBoxState::Checked : ECheckBoxState::Unchecked; })
+						),
+						EUserInterfaceActionType::Check
 					);
-			});
+				});
 		}
 	}
 }
 
 void FRewindDebuggerVLog::AddLogEntries(const TArray<TSharedPtr<FDebugObjectInfo>>& Components, float StartTime, float EndTime, const IVisualLoggerProvider* VisualLoggerProvider, UCanvas* Canvas)
 {
-	for(const TSharedPtr<FDebugObjectInfo>& ComponentInfo : Components)
+	for (const TSharedPtr<FDebugObjectInfo>& ComponentInfo : Components)
 	{
-		if (!ObjectsVisited.Contains(ComponentInfo->ObjectId))
+		const uint64 ObjectId = ComponentInfo->ObjectId;
+		if (!ObjectsVisited.Contains(ObjectId))
 		{
-			ObjectsVisited.Add(ComponentInfo->ObjectId);
-			VisualLoggerProvider->ReadVisualLogEntryTimeline(ComponentInfo->ObjectId, [this,StartTime, EndTime, Canvas](const IVisualLoggerProvider::VisualLogEntryTimeline &TimelineData)
-			{
-				TimelineData.EnumerateEvents(StartTime, EndTime, [this, StartTime, EndTime, Canvas](double InStartTime, double InEndTime, uint32 InDepth, const FVisualLogEntry& LogEntry)
+			ObjectsVisited.Add(ObjectId);
+			VisualLoggerProvider->ReadVisualLogEntryTimeline(ObjectId, [this, StartTime, EndTime, Canvas](const IVisualLoggerProvider::VisualLogEntryTimeline& TimelineData)
 				{
-					if (InStartTime >= StartTime && InStartTime <= EndTime)
-					{
-						RenderLogEntry(LogEntry, Canvas);
-					}
-					return TraceServices::EEventEnumerate::Continue;
+					TimelineData.EnumerateEvents(StartTime, EndTime, [this, StartTime, EndTime, Canvas](double InStartTime, double InEndTime, uint32 InDepth, const FVisualLogEntry& LogEntry)
+						{
+							if (InStartTime >= StartTime && InStartTime <= EndTime)
+							{
+								RenderLogEntry(LogEntry, Canvas);
+							}
+							return TraceServices::EEventEnumerate::Continue;
+						});
 				});
-			});
 		}
 
 		AddLogEntries(ComponentInfo->Children, StartTime, EndTime, VisualLoggerProvider, Canvas);
diff --git a/Engine/Plugins/Animation/GameplayInsights/Source/RewindDebuggerVLog/Private/RewindDebuggerVLog.h b/Engine/Plugins/Animation/GameplayInsights/Source/RewindDebuggerVLog/Private/RewindDebuggerVLog.h
--- a/Engine/Plugins/Animation/GameplayInsights/Source/RewindDebuggerVLog/Private/RewindDebuggerVLog.h
+++ b/Engine/Plugins/Animation/GameplayInsights/Source/RewindDebuggerVLog/Private/RewindDebuggerVLog.h
@@ -13,21 +13,23 @@ class UFont;
 class FRewindDebuggerVLog : public IRewindDebuggerExtension
 {
 public:
-	FRewindDebuggerVLog();
 	void OnShowDebugInfo(UCanvas* Canvas, APlayerController* Player);
 	virtual ~FRewindDebuggerVLog();
 
-	virtual FString GetName() { return TEXT("FRewindDebuggerVLog"); }
+	virtual FString GetName()
+	{
+		return TEXT("FRewindDebuggerVLog");
+	}
 
 	void Initialize();
-	void MakeCategoriesMenu(UToolMenu* Menu);
-	void MakeLogLevelMenu(UToolMenu* Menu);
-	void ToggleCategory(const FName& Category);
-	bool IsCategoryActive(const FName& Category);
-
-	ELogVerbosity::Type GetMinLogVerbosity() const;
-	void SetMinLogVerbosity(ELogVerbosity::Type Value);
-	
+	void MakeCategoriesMenu(UToolMenu* Menu) const;
+	void MakeLogLevelMenu(UToolMenu* Menu) const;
+	static void ToggleCategory(const FName& Category);
+	static bool IsCategoryActive(const FName& Category);
+
+	static ELogVerbosity::Type GetMinLogVerbosity();
+	static void SetMinLogVerbosity(ELogVerbosity::Type Value);
+
 	virtual void Update(float DeltaTime, IRewindDebugger* RewindDebugger) override;
 
 private:
@@ -40,11 +42,10 @@ private:
 	TWeakObjectPtr<AVLogRenderingActor> VLogActor;
 
 	TSet<uint64> ObjectsVisited;
-	int32 ScreenTextY;
+	int32 ScreenTextY = 0;
 
 	FDelegateHandle DelegateHandle;
 	UFont* MonospaceFont = nullptr;
 
 	TArray<FVisualLogEntry> ImmediateRenderQueue;
 };
-
diff --git a/Engine/Source/Editor/RewindDebuggerInterface/Public/IRewindDebugger.h b/Engine/Source/Editor/RewindDebuggerInterface/Public/IRewindDebugger.h
--- a/Engine/Source/Editor/RewindDebuggerInterface/Public/IRewindDebugger.h
+++ b/Engine/Source/Editor/RewindDebuggerInterface/Public/IRewindDebugger.h
@@ -1,42 +1,39 @@
 // Copyright Epic Games, Inc. All Rights Reserved.
 #pragma once
 
+#if UE_ENABLE_INCLUDE_ORDER_DEPRECATED_IN_5_7
 #include "CoreMinimal.h"
+#endif // UE_ENABLE_INCLUDE_ORDER_DEPRECATED_IN_5_7
 #include "TraceServices/Model/Frames.h"
 #include "IRewindDebugger.generated.h"
 
-// IRewindDebugger
-//
-// Public interface to rewind debugger
-
 REWINDDEBUGGERINTERFACE_API DECLARE_LOG_CATEGORY_EXTERN(LogRewindDebugger, Log, All);
 
-namespace TraceServices
+namespace TraceServices { class IAnalysisSession; }
+namespace RewindDebugger
 {
-	class IAnalysisSession;
+	class FRewindDebuggerTrack;
+	class FPropertiesTrack;
 }
+class IGameplayProvider;
+struct FObjectInfo;
 
 struct FDebugObjectInfo
 {
-	FDebugObjectInfo(uint64 Id, const FString& Name): ObjectId(Id), ObjectName(Name), bExpanded(true)
+	FDebugObjectInfo() = default;
+	FDebugObjectInfo(const uint64 Id, const FString& Name): ObjectId(Id), ObjectName(Name), bExpanded(true)
 	{
 	}
 
-	uint64 ObjectId;
+	uint64 ObjectId = 0;
 	FString ObjectName;
-	bool bExpanded;
+	bool bExpanded = false;
 
 	TArray<TSharedPtr<FDebugObjectInfo>> Children;
 };
 
-namespace RewindDebugger
-{
-	class FRewindDebuggerTrack;
-	class FPropertiesTrack;
-}
-
 UCLASS()
-class REWINDDEBUGGERINTERFACE_API UComponentContextMenuContext : public UObject
+class REWINDDEBUGGERINTERFACE_API URewindDebuggerTrackContextMenuContext : public UObject
 {
 	GENERATED_BODY()
 public:
@@ -45,8 +42,14 @@ public:
 	TSharedPtr<RewindDebugger::FRewindDebuggerTrack> SelectedTrack;
 };
 
-class SPropertiesDebugViewBase;
+class REWINDDEBUGGERINTERFACE_API UE_DEPRECATED(5.7, "Use URewindDebuggerTrackContextMenuContext instead") UComponentContextMenuContext : public URewindDebuggerTrackContextMenuContext
+{
+};
 
+/**
+ * IRewindDebugger
+ * Public interface to rewind debugger
+ */
 class REWINDDEBUGGERINTERFACE_API IRewindDebugger
 {
 public:
@@ -54,64 +57,77 @@ public:
 
 	virtual ~IRewindDebugger();
 
-	// get the time the debugger is scrubbed to, in seconds since the capture started (or the recording duration while the game is running)
+	/** @return the time the debugger is scrubbed to, in seconds since the capture started (or the recording duration while the game is running) */
 	virtual double CurrentTraceTime() const = 0;
-	
-	// get the time the debugger is scrubbed to, in seconds since the recording started
+
+	/** @return the time the debugger is scrubbed to, in seconds since the recording started */
 	virtual double GetScrubTime() const = 0;
 
-	// current visible range in trace/profiler units (same units as CurrentTraceTime)
+	/** @return the current visible range in trace/profiler units (same units as CurrentTraceTime) */
 	virtual const TRange<double>& GetCurrentTraceRange() const = 0;
-	
-	// current visible range in Rewind Debugger recording time units
+
+	/** @return the current visible range in Rewind Debugger recording time units */
 	virtual const TRange<double>& GetCurrentViewRange() const = 0;
 
-	// get the current analysis session
+	/** @return the current analysis session */
 	virtual const TraceServices::IAnalysisSession* GetAnalysisSession() const = 0;
 
-	// get insights id for the selected target actor
+	/** @return insights id for the selected target actor */
 	virtual uint64 GetTargetActorId() const = 0;
 
-	// get a list of all components of the selected target actor (with the actor as the first element in the list)
-	virtual TArray<TSharedPtr<FDebugObjectInfo>>& GetDebugComponents() = 0;
+	/** @return list of all components of the selected target actor (with the actor as the first element in the list)*/
+	virtual TArray<TSharedPtr<FDebugObjectInfo>>& GetDebuggedObjects() = 0;
+
+	UE_DEPRECATED(5.7, "Use GetDebuggedObjects instead")
+	virtual TArray<TSharedPtr<FDebugObjectInfo>>& GetDebugComponents() final
+	{
+		return GetDebuggedObjects();
+	}
+
+	/** @return Whether the given object id is one of the object currently debugged, or one of its children. */
+	virtual bool IsObjectCurrentlyDebugged(uint64 InObjectId) const = 0;
 
-	// check if the given object id is a debug component, or a child or a debug component
-	virtual bool IsContainedByDebugComponent(uint64 ObjectId) const = 0;
-	
-	// returns the currently selected debug component
-	virtual TSharedPtr<FDebugObjectInfo> GetSelectedComponent() const = 0;
+	UE_DEPRECATED(5.7, "Use IsObjectCurrentlyDebugged instead")
+	virtual bool IsContainedByDebugComponent(uint64 InObjectId) const final
+	{
+		return IsObjectCurrentlyDebugged(InObjectId);
+	}
+
+	/** @return the currently selected debug object */
+	virtual TSharedPtr<FDebugObjectInfo> GetSelectedObject() const = 0;
 
-	// returns the currently selected track
+	/** @return the currently selected track */
 	virtual TSharedPtr<RewindDebugger::FRewindDebuggerTrack> GetSelectedTrack() const = 0;
-	
-	// get posiotion of the selected target actor (returns true if position is valid)
+
+	/** @return position of the selected target actor (returns true if position is valid) */
 	virtual bool GetTargetActorPosition(FVector& OutPosition) const = 0;
 
-	// get the world that the debugger is replaying in
+	/** @return the world that the debugger is replaying in */
 	virtual UWorld* GetWorldToVisualize() const = 0;
 
-	// returns true if recording is active
+	/** @return whether the recording is active */
 	virtual bool IsRecording() const = 0;
 
-	// returns true if we have loaded a trace file from disk
+	/** @return whether a trace file is loaded from disk */
 	virtual bool IsTraceFileLoaded() const = 0;
 
-	// returns true if PIE is running and not paused
+	/** @return PIE is running and not paused */
 	virtual bool IsPIESimulating() const = 0;
 
-	// returns the length of the current recording
+	/** @return the length of the current recording */
 	virtual double GetRecordingDuration() const = 0;
 
-	// opens the Rewind Debugger details panel tab
+	/** Opens the Rewind Debugger details panel tab */
 	virtual void OpenDetailsPanel() = 0;
-	
-	virtual const struct FObjectInfo* FindOwningActorInfo(const class IGameplayProvider* GameplayProvider, uint64 ObjectId) const = 0;
+
+	/** @return the object information of the first object inheriting from AActor in the outer chain of ObjectId */
+	virtual const FObjectInfo* FindOwningActorInfo(const IGameplayProvider* GameplayProvider, uint64 ObjectId) const = 0;
 
 	virtual bool ShouldDisplayWorld(uint64 WorldId) = 0;
 
-	// get the current IRewindDebugger instance
+	/** @return the current IRewindDebugger instance */
 	static IRewindDebugger* Instance();
-	
+
 protected:
 	static IRewindDebugger* InternalInstance;
 };
diff --git a/Engine/Source/Editor/RewindDebuggerInterface/Public/IRewindDebuggerTrackCreator.h b/Engine/Source/Editor/RewindDebuggerInterface/Public/IRewindDebuggerTrackCreator.h
--- a/Engine/Source/Editor/RewindDebuggerInterface/Public/IRewindDebuggerTrackCreator.h
+++ b/Engine/Source/Editor/RewindDebuggerInterface/Public/IRewindDebuggerTrackCreator.h
@@ -2,9 +2,12 @@
 
 #pragma once
 
+#if UE_ENABLE_INCLUDE_ORDER_DEPRECATED_IN_5_7
 #include "CoreMinimal.h"
-#include "IRewindDebuggerView.h"
 #include "Textures/SlateIcon.h"
+#endif // UE_ENABLE_INCLUDE_ORDER_DEPRECATED_IN_5_7
+
+#include "IRewindDebuggerView.h"
 
 namespace TraceServices
 {
@@ -15,7 +18,7 @@ class IRewindDebugger;
 
 namespace RewindDebugger
 {
-	
+
 class FRewindDebuggerTrack;
 
 struct FRewindDebuggerTrackType
@@ -24,73 +27,92 @@ struct FRewindDebuggerTrackType
 	FText DisplayName;
 };
 
-// Interface class which creates tracks
+/**
+ * Interface class which creates tracks
+ */
 class REWINDDEBUGGERINTERFACE_API IRewindDebuggerTrackCreator : public IModularFeature
 {
-	public:
-	virtual ~IRewindDebuggerTrackCreator() {};
-	
+public:
 	static const FName ModularFeatureName;
 
-	// returns the type of UObject this track creator can create tracks for
+	virtual ~IRewindDebuggerTrackCreator() {}
+
+	/** @return name of the type of UObject this track creator can create tracks for */
 	FName GetTargetTypeName() const
 	{
 		return GetTargetTypeNameInternal();
 	}
 
-	// returns an identifying Name for this type of track
+	/** @return identifying Name for this type of track */
 	FName GetName() const
 	{
 		return GetNameInternal();
 	}
 
-	// returns an integer.  Higher values will show higher in the track list (default is 0)
+	/** @return integer representing the track priority.  Higher values will show higher in the track list (default is 0) */
 	int32 GetSortOrderPriority() const
 	{
 		return GetSortOrderPriorityInternal();
 	}
 
-	// optional additional filter, to prevent debug views from being listed if they have no data
+	/**
+	 * Optional filter to prevent tracks from being listed if there is no data associated to the provided object identifier.
+	 * @return Whether the track should be created to represent data associated to the provided object identifier.
+	 */
 	bool HasDebugInfo(uint64 ObjectId) const
 	{
 		return HasDebugInfoInternal(ObjectId);
 	};
-	
-	// Create a track which will be shown in the timeline view and tree view, as a child track of the Object
+
+	/**
+	 * Creates a track which will be shown in the timeline view and tree view, as a child track of the Object
+	 * @return the created child track
+	 */
 	TSharedPtr<RewindDebugger::FRewindDebuggerTrack> CreateTrack(uint64 ObjectId) const
 	{
 		return CreateTrackInternal(ObjectId);
 	}
-	
+
 	void GetTrackTypes(TArray<FRewindDebuggerTrackType>& Types) const
 	{
 		return GetTrackTypesInternal(Types);
 	};
 
-	private:
+private:
 
-	// get the UObject type name this Creator will create child tracks for
-	virtual FName GetTargetTypeNameInternal() const { return "Object"; }
+	/** get the UObject type name this Creator will create child tracks for */
+	virtual FName GetTargetTypeNameInternal() const
+	{
+		return "Object";
+	}
 
-	// get the Name (unique identifier) for this Creator
-	virtual FName GetNameInternal() const { return FName(); }
+	/** get the Name (unique identifier) for this Creator */
+	virtual FName GetNameInternal() const
+	{
+		return FName();
+	}
 
-	// An integer to override sort order for tracks created by this Creator (Higher priority will make tracks appear higher in the list)
-	virtual int32 GetSortOrderPriorityInternal() const { return 0; }
+	/** An integer to override sort order for tracks created by this Creator (Higher priority will make tracks appear higher in the list) */
+	virtual int32 GetSortOrderPriorityInternal() const
+	{
+		return 0;
+	}
 
-	// Add track types that this Creator will create to the track type list
-	virtual void GetTrackTypesInternal(TArray<FRewindDebuggerTrackType>& Types) const { };
+	/** Add track types that this Creator will create to the track type list */
+	virtual void GetTrackTypesInternal(TArray<FRewindDebuggerTrackType>& Types) const
+	{
+	};
 
-	// Returns true if this creator's track should be shown for this Object, false if there is no data and they should be hidden.
-	virtual bool HasDebugInfoInternal(uint64 ObjectId) const
+	/** Returns true if this creator's track should be shown for this Object, false if there is no data and they should be hidden. */
+	virtual bool HasDebugInfoInternal(uint64 InObjectId) const
 	{
 		return true;
 	};
-	
-	virtual TSharedPtr<FRewindDebuggerTrack> CreateTrackInternal(uint64 ObjectId) const
+
+	virtual TSharedPtr<FRewindDebuggerTrack> CreateTrackInternal(uint64 InObjectId) const
 	{
 		return TSharedPtr<FRewindDebuggerTrack>();
 	}
 };
-	
-}
\ No newline at end of file
+
+} // namespace RewindDebugger
diff --git a/Engine/Source/Editor/RewindDebuggerInterface/Public/RewindDebuggerTrack.h b/Engine/Source/Editor/RewindDebuggerInterface/Public/RewindDebuggerTrack.h
--- a/Engine/Source/Editor/RewindDebuggerInterface/Public/RewindDebuggerTrack.h
+++ b/Engine/Source/Editor/RewindDebuggerInterface/Public/RewindDebuggerTrack.h
@@ -1,7 +1,10 @@
 // Copyright Epic Games, Inc. All Rights Reserved.
 #pragma once
 
+#if UE_ENABLE_INCLUDE_ORDER_DEPRECATED_IN_5_7
 #include "CoreMinimal.h"
+#endif // UE_ENABLE_INCLUDE_ORDER_DEPRECATED_IN_5_7
+
 #include "IRewindDebugger.h"
 #include "Widgets/SWidget.h"
 #include "Textures/SlateIcon.h"
@@ -15,7 +18,7 @@ struct FToolMenuSection;
 
 namespace RewindDebugger
 {
-	
+
 class FRewindDebuggerTrack
 {
 public:
@@ -27,128 +30,174 @@ public:
 	{
 	}
 
-	bool GetIsExpanded()
+	bool GetIsExpanded() const
 	{
 		return bExpanded;
 	}
 
-	void SetIsExpanded(bool bIsExpanded)
+	void SetIsExpanded(const bool bIsExpanded)
 	{
 		bExpanded = bIsExpanded;
 	}
-	
-	bool GetIsSelected()
+
+	bool GetIsSelected() const
 	{
 		return bSelected;
 	}
 
-	void SetIsSelected(bool bIsSelected)
+	void SetIsSelected(const bool bIsSelected)
 	{
 		bSelected = bIsSelected;
 	}
-    		
-	bool GetIsTreeHovered()
+
+	bool GetIsTreeHovered() const
 	{
 		return bTreeHovered;
 	}
-	
-	void SetIsTreeHovered(bool bIsHovered)
-   	{
-   		bTreeHovered = bIsHovered;
-   	}	
-		
-	bool GetIsTrackHovered()
+
+	void SetIsTreeHovered(const bool bIsHovered)
+	{
+		bTreeHovered = bIsHovered;
+	}
+
+	bool GetIsTrackHovered() const
 	{
 		return bTrackHovered;
 	}
-	
-	void SetIsTrackHovered(bool bIsHovered)
-   	{
-   		bTrackHovered = bIsHovered;
-   	}
-	
-	bool GetIsHovered()
+
+	void SetIsTrackHovered(const bool bIsHovered)
+	{
+		bTrackHovered = bIsHovered;
+	}
+
+	bool GetIsHovered() const
 	{
 		return bTrackHovered || bTreeHovered;
 	}
 
-	// Update should do work to compute children etc for the current time range.  Return true if children have changed.
+	/** Update should do work to compute children etc. for the current time range.  Return true if children have changed. */
 	bool Update()
 	{
 		return UpdateInternal();
 	}
 
-	// Get a widget to show in the timeline view for this track
+	/** Get a widget to show in the timeline view for this track */
 	TSharedPtr<SWidget> GetTimelineView()
 	{
 		return GetTimelineViewInternal();
 	}
 
-	// Get a widget to show in the details tab, when this track is selected
+	/** Get a widget to show in the details tab, when this track is selected */
 	TSharedPtr<SWidget> GetDetailsView()
 	{
 		return GetDetailsViewInternal();
 	}
-	
-	// unique name for track (must match creator name if track is created by an IRewindDebuggerViewCreator) 
+
+	/** unique name for track (must match creator name if track is created by an IRewindDebuggerViewCreator)  */
 	FName GetName() const
 	{
 		return GetNameInternal();
 	}
 
-	// icon to display in the tree view
+	/** icon to display in the tree view */
 	FSlateIcon GetIcon()
 	{
 		return GetIconInternal();
 	}
 
-	// display name for track in Tree View
+	/** display name for track in Tree View */
 	FText GetDisplayName() const
 	{
 		return GetDisplayNameInternal();
 	}
 
-	// insights object id for an object associated with this track
+	/** insights object id for an object associated with this track */
 	uint64 GetObjectId() const
 	{
 		return GetObjectIdInternal();
 	}
 
-	// iterate over all sub-tracks of this track and call Iterator function
+	/** iterate over all sub-tracks of this track and call Iterator function */
 	void IterateSubTracks(TFunction<void(TSharedPtr<FRewindDebuggerTrack> SubTrack)> IteratorFunction)
 	{
 		IterateSubTracksInternal(IteratorFunction);
 	}
 
-	// returns true for tracks that contain debug data (used for filtering out parts of the hierarchy with no useful information in them)
+	/** returns true for tracks that contain debug data (used for filtering out parts of the hierarchy with no useful information in them) */
 	bool HasDebugData() const
 	{
 		return HasDebugDataInternal();
 	}
 
-	// Called when a track is double clicked.  Returns true if the track handled the double click
+	/** Called when a track is double-clicked.  Returns true if the track handled the double click */
 	bool HandleDoubleClick()
 	{
 		return HandleDoubleClickInternal();
 	}
 
-	bool IsVisible() const { return bVisible; }
-	void SetIsVisible(bool bInIsVisible) { bVisible = bInIsVisible; }
+	bool IsVisible() const
+	{
+		return bVisible;
+	}
+
+	void SetIsVisible(const bool bInIsVisible)
+	{
+		bVisible = bInIsVisible;
+	}
 
-	// Called to generate context menu for the current selected track
+	/** Called to generate context menu for the current selected track */
 	virtual void BuildContextMenu(FToolMenuSection& InMenuSection) {};
 
 private:
 
-	virtual bool UpdateInternal() { return false; }
-	virtual TSharedPtr<SWidget> GetTimelineViewInternal() { return TSharedPtr<SWidget>(); }
-	virtual TSharedPtr<SWidget> GetDetailsViewInternal() { return TSharedPtr<SWidget>(); }
-	virtual FSlateIcon GetIconInternal() { return FSlateIcon(); }
-	virtual FName GetNameInternal() const { return ""; }
-	virtual FText GetDisplayNameInternal() const { return FText(); }
-	virtual uint64 GetObjectIdInternal() const { return 0; }
-	virtual bool HasDebugDataInternal() const { return true; }
-	virtual void IterateSubTracksInternal(TFunction<void(TSharedPtr<FRewindDebuggerTrack> SubTrack)> IteratorFunction) { }
+	virtual bool UpdateInternal()
+	{
+		return false;
+	}
+
+	virtual TSharedPtr<SWidget> GetTimelineViewInternal()
+	{
+		return TSharedPtr<SWidget>();
+	}
+
+	virtual TSharedPtr<SWidget> GetDetailsViewInternal()
+	{
+		return TSharedPtr<SWidget>();
+	}
+
+	virtual FSlateIcon GetIconInternal()
+	{
+		return FSlateIcon();
+	}
+
+	virtual FName GetNameInternal() const
+	{
+		return "";
+	}
+
+	virtual FText GetDisplayNameInternal() const
+	{
+		return FText();
+	}
+
+	virtual uint64 GetObjectIdInternal() const
+	{
+		return {};
+	}
+
+	virtual uint64 GetSubObjectIdInternal() const
+	{
+		return INDEX_NONE;
+	}
+
+	virtual bool HasDebugDataInternal() const
+	{
+		return true;
+	}
+
+	virtual void IterateSubTracksInternal(TFunction<void(TSharedPtr<FRewindDebuggerTrack> SubTrack)> IteratorFunction)
+	{
+	}
 
 	virtual bool HandleDoubleClickInternal()
 	{
@@ -162,5 +211,5 @@ private:
 	bool bExpanded : 1 = true;
 	bool bVisible : 1 = true;
 };
-	
-}
\ No newline at end of file
+
+} // namespace RewindDebugger
\ No newline at end of file
diff --git a/Engine/Source/Runtime/Engine/Private/ObjectTrace.cpp b/Engine/Source/Runtime/Engine/Private/ObjectTrace.cpp
--- a/Engine/Source/Runtime/Engine/Private/ObjectTrace.cpp
+++ b/Engine/Source/Runtime/Engine/Private/ObjectTrace.cpp
@@ -1,12 +1,11 @@
 // Copyright Epic Games, Inc. All Rights Reserved.
 
 #include "ObjectTrace.h"
-#include "Misc/ScopeRWLock.h"
 
 #include UE_INLINE_GENERATED_CPP_BY_NAME(ObjectTrace)
 
 #if OBJECT_TRACE_ENABLED
-
+#include "Misc/ScopeRWLock.h"
 #include "SceneView.h"
 #if WITH_EDITOR
 #include "Editor.h"
@@ -38,7 +37,6 @@ UE_TRACE_EVENT_BEGIN(Object, ObjectLifetimeBegin2)
 	UE_TRACE_EVENT_FIELD(uint64, Id)
 UE_TRACE_EVENT_END()
 
-
 UE_TRACE_EVENT_BEGIN(Object, ObjectTransform)
 	UE_TRACE_EVENT_FIELD(uint64, Cycle)
 	UE_TRACE_EVENT_FIELD(double, RecordingTime)
@@ -138,21 +136,21 @@ static FDelegateHandle WorldTickStartHandle;
 
 static void TickObjectTraceWorldSubsystem(UWorld* InWorld, ELevelTick InTickType, float InDeltaSeconds)
 {
-	if(InTickType == LEVELTICK_All)
+	if (InTickType == LEVELTICK_All)
 	{
 		if (!InWorld->IsPaused())
 		{
-			if(UObjectTraceWorldSubsystem* Subsystem = UWorld::GetSubsystem<UObjectTraceWorldSubsystem>(InWorld))
+			if (UObjectTraceWorldSubsystem* Subsystem = UWorld::GetSubsystem<UObjectTraceWorldSubsystem>(InWorld))
 			{
 				Subsystem->FrameIndex++;
 				Subsystem->ElapsedTime += InDeltaSeconds;
 
 				UE_TRACE_LOG(Object, RecordingInfo, ObjectChannel)
-		            << RecordingInfo.WorldId(FObjectTrace::GetObjectId(InWorld))
+					<< RecordingInfo.WorldId(FObjectTrace::GetObjectId(InWorld))
 					<< RecordingInfo.Cycle(FPlatformTime::Cycles64())
 					<< RecordingInfo.ElapsedTime(Subsystem->ElapsedTime)
 					<< RecordingInfo.FrameIndex(Subsystem->FrameIndex)
-					<< RecordingInfo.RecordingIndex(Subsystem->RecordingIndex); 
+					<< RecordingInfo.RecordingIndex(Subsystem->RecordingIndex);
 			}
 		}
 	}
@@ -191,7 +189,7 @@ void FObjectTrace::Destroy()
 void FObjectTrace::Reset()
 {
 	GObjectIdAnnotations.RemoveAllAnnotations();
-	
+
 	UE_AUTORTFM_OPEN
 	{
 		FRWScopeLock ScopeLock(GObjectTracedSetLock, SLT_Write);
@@ -210,7 +208,7 @@ uint64 FObjectTrace::GetObjectId(const UObject* InObject)
 		static uint64 CurrentId = 1;
 
 		FObjectIdAnnotation Annotation = GObjectIdAnnotations.GetAnnotation(InObjectInner);
-		if(Annotation.Id == 0)
+		if (Annotation.Id == 0)
 		{
 			Annotation.Id = CurrentId++;
 			GObjectIdAnnotations.AddAnnotation(InObjectInner, MoveTemp(Annotation));
@@ -221,11 +219,11 @@ uint64 FObjectTrace::GetObjectId(const UObject* InObject)
 
 	uint64 Id = 0;
 	uint64 OuterId = 0;
-	if(InObject)
+	if (InObject)
 	{
 		Id = GetObjectIdInner(InObject);
 
-		if(const UObject* Outer = InObject->GetOuter())
+		if (const UObject* Outer = InObject->GetOuter())
 		{
 			OuterId = GetObjectIdInner(Outer);
 		}
@@ -243,13 +241,13 @@ UObject* FObjectTrace::GetObjectFromId(uint64 Id)
 	{
 		return nullptr;
 	}
-	
+
 	return GObjectIdAnnotations.Find(FindAnnotation);
 }
 
 void FObjectTrace::ResetWorldElapsedTime(const UWorld* World)
 {
-	if(UObjectTraceWorldSubsystem* WorldSubsystem = UWorld::GetSubsystem<UObjectTraceWorldSubsystem>(World))
+	if (UObjectTraceWorldSubsystem* WorldSubsystem = UWorld::GetSubsystem<UObjectTraceWorldSubsystem>(World))
 	{
 		WorldSubsystem->ElapsedTime = 0;
 	}
@@ -257,7 +255,7 @@ void FObjectTrace::ResetWorldElapsedTime(const UWorld* World)
 
 double FObjectTrace::GetWorldElapsedTime(const UWorld* World)
 {
-	if(UObjectTraceWorldSubsystem* WorldSubsystem = UWorld::GetSubsystem<UObjectTraceWorldSubsystem>(World))
+	if (const UObjectTraceWorldSubsystem* WorldSubsystem = UWorld::GetSubsystem<UObjectTraceWorldSubsystem>(World))
 	{
 		return WorldSubsystem->ElapsedTime;
 	}
@@ -266,11 +264,11 @@ double FObjectTrace::GetWorldElapsedTime(const UWorld* World)
 
 double FObjectTrace::GetObjectWorldElapsedTime(const UObject* InObject)
 {
-	if(InObject != nullptr)
+	if (InObject != nullptr)
 	{
-		if(UWorld* World = InObject->GetWorld())
+		if (const UWorld* World = InObject->GetWorld())
 		{
-			if(UObjectTraceWorldSubsystem* WorldSubsystem = UWorld::GetSubsystem<UObjectTraceWorldSubsystem>(World))
+			if (const UObjectTraceWorldSubsystem* WorldSubsystem = UWorld::GetSubsystem<UObjectTraceWorldSubsystem>(World))
 			{
 				return WorldSubsystem->ElapsedTime;
 			}
@@ -282,15 +280,15 @@ double FObjectTrace::GetObjectWorldElapsedTime(const UObject* InObject)
 
 void FObjectTrace::SetWorldRecordingIndex(const UWorld* World, uint16 Index)
 {
-	if(UObjectTraceWorldSubsystem* WorldSubsystem = UWorld::GetSubsystem<UObjectTraceWorldSubsystem>(World))
+	if (UObjectTraceWorldSubsystem* WorldSubsystem = UWorld::GetSubsystem<UObjectTraceWorldSubsystem>(World))
 	{
 		WorldSubsystem->RecordingIndex = Index;
 	}
 }
 
-uint16 FObjectTrace::GetWorldRecordingIndex(const UWorld* World)
+uint16 FObjectTrace::GetWorldRecordingIndex(const UWorld* InWorld)
 {
-	if(UObjectTraceWorldSubsystem* WorldSubsystem = UWorld::GetSubsystem<UObjectTraceWorldSubsystem>(World))
+	if (const UObjectTraceWorldSubsystem* WorldSubsystem = UWorld::GetSubsystem<UObjectTraceWorldSubsystem>(InWorld))
 	{
 		return WorldSubsystem->RecordingIndex;
 	}
@@ -299,11 +297,11 @@ uint16 FObjectTrace::GetWorldRecordingIndex(const UWorld* World)
 
 uint16 FObjectTrace::GetObjectWorldRecordingIndex(const UObject* InObject)
 {
-	if(InObject != nullptr)
+	if (InObject != nullptr)
 	{
-		if(UWorld* World = InObject->GetWorld())
+		if (const UWorld* World = InObject->GetWorld())
 		{
-			if(UObjectTraceWorldSubsystem* WorldSubsystem = UWorld::GetSubsystem<UObjectTraceWorldSubsystem>(World))
+			if (const UObjectTraceWorldSubsystem* WorldSubsystem = UWorld::GetSubsystem<UObjectTraceWorldSubsystem>(World))
 			{
 				return WorldSubsystem->RecordingIndex;
 			}
@@ -315,11 +313,11 @@ uint16 FObjectTrace::GetObjectWorldRecordingIndex(const UObject* InObject)
 
 uint16 FObjectTrace::GetObjectWorldTickCounter(const UObject* InObject)
 {
-	if(InObject != nullptr)
+	if (InObject != nullptr)
 	{
-		if(UWorld* World = InObject->GetWorld())
+		if (const UWorld* World = InObject->GetWorld())
 		{
-			if(UObjectTraceWorldSubsystem* WorldSubsystem = UWorld::GetSubsystem<UObjectTraceWorldSubsystem>(World))
+			if (const UObjectTraceWorldSubsystem* WorldSubsystem = UWorld::GetSubsystem<UObjectTraceWorldSubsystem>(World))
 			{
 				return WorldSubsystem->FrameIndex;
 			}
@@ -331,12 +329,12 @@ uint16 FObjectTrace::GetObjectWorldTickCounter(const UObject* InObject)
 
 void FObjectTrace::OutputClass(const UClass* InClass)
 {
-	bool bChannelEnabled = UE_TRACE_CHANNELEXPR_IS_ENABLED(ObjectChannel);
+	const bool bChannelEnabled = UE_TRACE_CHANNELEXPR_IS_ENABLED(ObjectChannel);
 	if (!bChannelEnabled || InClass == nullptr)
 	{
 		return;
 	}
-	uint64 ObjectId = GetObjectId(InClass);
+	const uint64 ObjectId = GetObjectId(InClass);
 	if (!AddObjectToTracedSet(ObjectId))
 	{
 		return;
@@ -344,9 +342,9 @@ void FObjectTrace::OutputClass(const UClass* InClass)
 
 	OutputClass(InClass->GetSuperClass());
 
-	FString ClassPathName = InClass->GetPathName();
+	const FString ClassPathName = InClass->GetPathName();
 	TCHAR ClassName[FName::StringBufferSize];
-	uint32 ClassNameLength = InClass->GetFName().ToString(ClassName);
+	const uint32 ClassNameLength = InClass->GetFName().ToString(ClassName);
 
 	UE_TRACE_LOG(Object, Class, ObjectChannel)
 		<< Class.Id(ObjectId)
@@ -357,7 +355,7 @@ void FObjectTrace::OutputClass(const UClass* InClass)
 
 void FObjectTrace::OutputView(const UObject* InPlayer, const FSceneView* InView)
 {
-	bool bChannelEnabled = UE_TRACE_CHANNELEXPR_IS_ENABLED(ObjectChannel);
+	const bool bChannelEnabled = UE_TRACE_CHANNELEXPR_IS_ENABLED(ObjectChannel);
 	if (!bChannelEnabled || InPlayer == nullptr)
 	{
 		return;
@@ -369,10 +367,10 @@ void FObjectTrace::OutputView(const UObject* InPlayer, const FSceneView* InView)
 	}
 
 	const FIntRect& ViewRect = InView->CameraConstrainedViewRect;
-	float AspectRatio = (float)ViewRect.Width()/(float)ViewRect.Height();
+	const float AspectRatio = static_cast<float>(ViewRect.Width()) / static_cast<float>(ViewRect.Height());
 
-	FMatrix ProjMatrix = InView->ViewMatrices.GetProjectionMatrix();
-	float Fov = atan(1.0 / ProjMatrix.M[0][0]) * 2.0 * 180.0 / UE_DOUBLE_PI;
+	const FMatrix ProjMatrix = InView->ViewMatrices.GetProjectionMatrix();
+	const float Fov = atan(1.0 / ProjMatrix.M[0][0]) * 2.0 * 180.0 / UE_DOUBLE_PI;
 
 	UE_TRACE_LOG(Object, View, ObjectChannel)
 		<< View.Cycle(FPlatformTime::Cycles64())
@@ -389,13 +387,13 @@ void FObjectTrace::OutputView(const UObject* InPlayer, const FSceneView* InView)
 
 void FObjectTrace::OutputObject(const UObject* InObject)
 {
-	bool bChannelEnabled = UE_TRACE_CHANNELEXPR_IS_ENABLED(ObjectChannel);
+	const bool bChannelEnabled = UE_TRACE_CHANNELEXPR_IS_ENABLED(ObjectChannel);
 	if (!bChannelEnabled || InObject == nullptr)
 	{
 		return;
 	}
 
-	if(InObject->HasAnyFlags(RF_ClassDefaultObject))
+	if (InObject->HasAnyFlags(RF_ClassDefaultObject))
 	{
 		return;
 	}
@@ -404,8 +402,8 @@ void FObjectTrace::OutputObject(const UObject* InObject)
 	{
 		return;
 	}
-	
-	uint64 ObjectId = GetObjectId(InObject);
+
+	const uint64 ObjectId = GetObjectId(InObject);
 	if (!AddObjectToTracedSet(ObjectId))
 	{
 		return;
@@ -417,8 +415,8 @@ void FObjectTrace::OutputObject(const UObject* InObject)
 	TRACE_CLASS(InObject->GetClass());
 
 	TCHAR ObjectName[FName::StringBufferSize];
-	uint32 ObjectNameLength = InObject->GetFName().ToString(ObjectName);
-	FString ObjectPathName = InObject->GetPathName();
+	const uint32 ObjectNameLength = InObject->GetFName().ToString(ObjectName);
+	const FString ObjectPathName = InObject->GetPathName();
 
 	UE_TRACE_LOG(Object, Object, ObjectChannel)
 		<< Object.Id(ObjectId)
@@ -430,18 +428,18 @@ void FObjectTrace::OutputObject(const UObject* InObject)
 	UE_TRACE_LOG(Object, ObjectLifetimeBegin2, ObjectChannel)
 		<< ObjectLifetimeBegin2.Cycle(FPlatformTime::Cycles64())
 		<< ObjectLifetimeBegin2.RecordingTime(FObjectTrace::GetWorldElapsedTime(InObject->GetWorld()))
-		<< ObjectLifetimeBegin2.Id(GetObjectId(InObject));
+		<< ObjectLifetimeBegin2.Id(ObjectId);
 }
 
 void FObjectTrace::OutputObjectTransform(const UObject* InObject, const FTransform& InTransform, ETeleportType InTeleportType)
 {
-	bool bChannelEnabled = UE_TRACE_CHANNELEXPR_IS_ENABLED(ObjectChannel);
+	const bool bChannelEnabled = UE_TRACE_CHANNELEXPR_IS_ENABLED(ObjectChannel);
 	if (!bChannelEnabled || InObject == nullptr)
 	{
 		return;
 	}
 
-	if(InObject->HasAnyFlags(RF_ClassDefaultObject))
+	if (InObject->HasAnyFlags(RF_ClassDefaultObject))
 	{
 		return;
 	}
@@ -463,13 +461,13 @@ void FObjectTrace::OutputObjectTransform(const UObject* InObject, const FTransfo
 
 void FObjectTrace::OutputObjectEvent(const UObject* InObject, const TCHAR* InEvent)
 {
-	bool bChannelEnabled = UE_TRACE_CHANNELEXPR_IS_ENABLED(ObjectChannel);
+	const bool bChannelEnabled = UE_TRACE_CHANNELEXPR_IS_ENABLED(ObjectChannel);
 	if (!bChannelEnabled || InObject == nullptr)
 	{
 		return;
 	}
 
-	if(InObject->HasAnyFlags(RF_ClassDefaultObject))
+	if (InObject->HasAnyFlags(RF_ClassDefaultObject))
 	{
 		return;
 	}
@@ -500,7 +498,7 @@ void FObjectTrace::OutputObjectLifetimeEnd(const UObject* InObject)
 		return;
 	}
 
-	if(InObject->HasAnyFlags(RF_ClassDefaultObject))
+	if (InObject->HasAnyFlags(RF_ClassDefaultObject))
 	{
 		return;
 	}
@@ -520,7 +518,7 @@ void FObjectTrace::OutputObjectLifetimeEnd(const UObject* InObject)
 
 void FObjectTrace::OutputPawnPossess(const UObject* InController, const UObject* InPawn)
 {
-	bool bChannelEnabled = UE_TRACE_CHANNELEXPR_IS_ENABLED(ObjectChannel);
+	const bool bChannelEnabled = UE_TRACE_CHANNELEXPR_IS_ENABLED(ObjectChannel);
 	if (!bChannelEnabled || InController == nullptr)
 	{
 		return;
@@ -531,7 +529,7 @@ void FObjectTrace::OutputPawnPossess(const UObject* InController, const UObject*
 	{
 		TRACE_OBJECT(InPawn);
 	}
-	
+
 	UE_TRACE_LOG(Object, PawnPossess, ObjectChannel)
 		<< PawnPossess.Cycle(FPlatformTime::Cycles64())
 		<< PawnPossess.ControllerId(GetObjectId(InController))
@@ -540,7 +538,7 @@ void FObjectTrace::OutputPawnPossess(const UObject* InController, const UObject*
 
 void FObjectTrace::OutputWorld(const UWorld* InWorld)
 {
-	bool bChannelEnabled = UE_TRACE_CHANNELEXPR_IS_ENABLED(ObjectChannel);
+	const bool bChannelEnabled = UE_TRACE_CHANNELEXPR_IS_ENABLED(ObjectChannel);
 	if (!bChannelEnabled || InWorld == nullptr)
 	{
 		return;
@@ -560,8 +558,8 @@ void FObjectTrace::OutputWorld(const UWorld* InWorld)
 	UE_TRACE_LOG(Object, World, ObjectChannel)
 		<< World.Id(GetObjectId(InWorld))
 		<< World.PIEInstanceId(InWorld->GetOutermost()->GetPIEInstanceID())
-		<< World.Type((uint8)InWorld->WorldType)
-		<< World.NetMode((uint8)InWorld->GetNetMode())
+		<< World.Type(static_cast<uint8>(InWorld->WorldType))
+		<< World.NetMode(static_cast<uint8>(InWorld->GetNetMode()))
 		<< World.IsSimulating(bIsSimulating);
 
 	// Trace object AFTER world info so we don't risk world info not being present in the trace
diff --git a/Engine/Source/Runtime/Engine/Public/ObjectTrace.h b/Engine/Source/Runtime/Engine/Public/ObjectTrace.h
--- a/Engine/Source/Runtime/Engine/Public/ObjectTrace.h
+++ b/Engine/Source/Runtime/Engine/Public/ObjectTrace.h
@@ -10,7 +10,7 @@
 
 #include "ObjectTrace.generated.h"
 
-// World subsystem used to track world info
+/** World subsystem used to track world info */
 UCLASS()
 class UObjectTraceWorldSubsystem : public UWorldSubsystem
 {
@@ -28,11 +28,11 @@ class UObjectTraceWorldSubsystem : public UWorldSubsystem
 #endif
 
 public:
-	// The frame index incremented each tick
+	/** The frame index incremented each tick */
 	uint16 FrameIndex;
-	// Trace Recording identifier  (set by RewindDebugger or 0)
+	/** Trace Recording identifier  (set by RewindDebugger or 0) */
 	uint16 RecordingIndex;
-	// Elapsed time since recording started  (or since start of game if RewindDebugger didn't start the trace)
+	/** Elapsed time since recording started  (or since start of game if RewindDebugger didn't start the trace) */
 	double ElapsedTime;
 };
 
@@ -50,7 +50,7 @@ struct FObjectTrace
 
 	/** Shut down object tracing */
 	ENGINE_API static void Destroy();
-	
+
 	/** Reset Caches so a new trace can be started*/
 	ENGINE_API static void Reset();
 
@@ -62,19 +62,19 @@ struct FObjectTrace
 
 	/** Helper function to output an object transform */
 	ENGINE_API static void OutputObjectTransform(const UObject* InObject, const FTransform& InTransform, ETeleportType InTeleportType);
-	
-	/** Helper function to output object create event */
+
+	/** Helper function to output object creation event */
 	ENGINE_API static void OutputObjectLifetimeBegin(const UObject* InObject);
-	
-	/** Helper function to output object create event */
+
+	/** Helper function to output object destruction event */
 	ENGINE_API static void OutputObjectLifetimeEnd(const UObject* InObject);
 
 	/** Helper function to output camera information for a player */
-	ENGINE_API static void OutputView(const UObject* LocalPlayer, const FSceneView* View);
+	ENGINE_API static void OutputView(const UObject* InPlayer, const FSceneView* InView);
 
 	/** Helper function to output an object event */
 	ENGINE_API static void OutputObjectEvent(const UObject* InObject, const TCHAR* InEvent);
-	
+
 	/** Helper function to output controller attach event */
 	ENGINE_API static void OutputPawnPossess(const UObject* InController, const UObject* InPawn);
 
@@ -107,7 +107,6 @@ struct FObjectTrace
 
 	/** Helper function to output a world */
 	ENGINE_API static void OutputWorld(const UWorld* InWorld);
-
 };
 
 #define TRACE_CLASS(Class) \
@@ -122,19 +121,19 @@ struct FObjectTrace
 #if TRACE_FILTERING_ENABLED
 
 #define TRACE_OBJECT_EVENT(Object, Event) \
-	if(CAN_TRACE_OBJECT(Object)) { UNCONDITIONAL_TRACE_OBJECT_EVENT(Object, Event); }
+	if (CAN_TRACE_OBJECT(Object)) { UNCONDITIONAL_TRACE_OBJECT_EVENT(Object, Event); }
 
 #define TRACE_OBJECT_TRANSFORM(Object, Transform, Teleport) \
-	if(CAN_TRACE_OBJECT(Object)) { FObjectTrace::OutputObjectTransform(Object, Transform, Teleport); }
-	
+	if (CAN_TRACE_OBJECT(Object)) { FObjectTrace::OutputObjectTransform(Object, Transform, Teleport); }
+
 #define TRACE_OBJECT_LIFETIME_BEGIN(Object) \
-	if(CAN_TRACE_OBJECT(Object)) { FObjectTrace::OutputObjectLifetimeBegin(Object); }
+	if (CAN_TRACE_OBJECT(Object)) { FObjectTrace::OutputObjectLifetimeBegin(Object); }
 
 #define TRACE_OBJECT_LIFETIME_END(Object) \
-	if(CAN_TRACE_OBJECT(Object)) { FObjectTrace::OutputObjectLifetimeEnd(Object); }
+	if (CAN_TRACE_OBJECT(Object)) { FObjectTrace::OutputObjectLifetimeEnd(Object); }
 
 #define TRACE_PAWN_POSSESS(Controller, Pawn)\
-	if(CAN_TRACE_OBJECT(Controller) && (Pawn==nullptr || (CAN_TRACE_OBJECT(Pawn)))) FObjectTrace::OutputPawnPossess(Controller, Pawn);
+	if (CAN_TRACE_OBJECT(Controller) && (Pawn==nullptr || (CAN_TRACE_OBJECT(Pawn)))) FObjectTrace::OutputPawnPossess(Controller, Pawn);
 
 #else
 
@@ -167,14 +166,14 @@ struct FObjectTrace
 	static UObject* GetObjectFromId(uint64 id) { return nullptr; }
 };
 
-#define TRACE_CLASS(Class)
-#define TRACE_OBJECT(Object)
-#define TRACE_OBJECT_TRANSFORM(Object, Transform, Teleport)
-#define TRACE_OBJECT_EVENT(Object, Event)
-#define TRACE_WORLD(World)
-#define TRACE_PAWN_POSSESS(Controller, Pawn)
-#define TRACE_VIEW(Player, View)
-#define TRACE_OBJECT_LIFETIME_BEGIN(Object)
-#define TRACE_OBJECT_LIFETIME_END(Object)
-
-#endif
+#define TRACE_CLASS(...)
+#define TRACE_OBJECT(...)
+#define TRACE_OBJECT_TRANSFORM(...)
+#define TRACE_OBJECT_EVENT(...)
+#define TRACE_WORLD(...)
+#define TRACE_PAWN_POSSESS(...)
+#define TRACE_VIEW(...)
+#define TRACE_OBJECT_LIFETIME_BEGIN(...)
+#define TRACE_OBJECT_LIFETIME_END(...)
+
+#endif
\ No newline at end of file
