[RewindDebugger] added caching to get FText for log categories

diff --git a/Engine/Plugins/Animation/GameplayInsights/Source/RewindDebuggerVLog/Private/VisualLogTrack.cpp b/Engine/Plugins/Animation/GameplayInsights/Source/RewindDebuggerVLog/Private/VisualLogTrack.cpp
--- a/Engine/Plugins/Animation/GameplayInsights/Source/RewindDebuggerVLog/Private/VisualLogTrack.cpp
+++ b/Engine/Plugins/Animation/GameplayInsights/Source/RewindDebuggerVLog/Private/VisualLogTrack.cpp
@@ -12,6 +12,23 @@
 namespace RewindDebugger
 {
 
+static TArray<TPair<FName, FText>> CategoryNameAsTextMappings;
+
+FText& GetCategoryNameAsText(FName InCategory)
+{
+	TPair<FName, FText>* CachedEntry = CategoryNameAsTextMappings.FindByPredicate([InCategory](const TPair<FName, FText>& Pair)
+		{
+			return Pair.Key == InCategory;
+		});
+
+	if (CachedEntry == nullptr)
+	{
+		CachedEntry = &CategoryNameAsTextMappings.Add_GetRef({ InCategory, FText::FromName(InCategory) });
+	}
+
+	return CachedEntry->Value;
+}
+
 UVLogDetailsObject* FVisualLogCategoryTrack::InitializeDetailsObject()
 {
 	UVLogDetailsObject* DetailsObject = NewObject<UVLogDetailsObject>();
@@ -25,11 +42,10 @@ FVisualLogCategoryTrack::FVisualLogCategoryTrack(uint64 InObjectId, const FName&
 	ObjectId(InObjectId),
 	Category(InCategory)
 {
-	TrackName = FText::FromName(Category);
+	TrackName = GetCategoryNameAsText(Category);
 	EventData = MakeShared<SEventTimelineView::FTimelineEventData>();
 	Icon = FSlateIcon("EditorStyle", "Sequencer.Tracks.Event", "Sequencer.Tracks.Event");
 
-	
 	FPropertyEditorModule& PropertyEditorModule = FModuleManager::GetModuleChecked<FPropertyEditorModule>("PropertyEditor");
 	FDetailsViewArgs DetailsViewArgs;
 	DetailsViewArgs.NameAreaSettings = FDetailsViewArgs::HideNameArea;
@@ -52,12 +68,12 @@ TSharedPtr<SEventTimelineView::FTimelineEventData> FVisualLogCategoryTrack::GetE
 	{
 		EventData = MakeShared<SEventTimelineView::FTimelineEventData>();
 	}
-	
+
 	EventUpdateRequested++;
-	
+
 	return EventData;
 }
-	
+
 static FLinearColor MakeNotifyColor(uint32 InSeed, bool bInLine = false)
 {
 	FRandomStream Stream(InSeed);
@@ -66,7 +82,7 @@ static FLinearColor MakeNotifyColor(uint32 InSeed, bool bInLine = false)
 	return FLinearColor::MakeFromHSV8(Hue, SatVal, SatVal);
 }
 
-TSharedPtr<SWidget> FVisualLogCategoryTrack::GetDetailsViewInternal() 
+TSharedPtr<SWidget> FVisualLogCategoryTrack::GetDetailsViewInternal()
 {
 	return DetailsView;
 }
@@ -74,47 +90,47 @@ TSharedPtr<SWidget> FVisualLogCategoryTrack::GetDetailsViewInternal()
 bool FVisualLogCategoryTrack::UpdateInternal()
 {
 	IRewindDebugger* RewindDebugger = IRewindDebugger::Instance();
-	
+
 	TRange<double> TraceTimeRange = RewindDebugger->GetCurrentTraceRange();
 	double StartTime = TraceTimeRange.GetLowerBoundValue();
 	double EndTime = TraceTimeRange.GetUpperBoundValue();
-	
+
 	const TraceServices::IAnalysisSession* AnalysisSession = RewindDebugger->GetAnalysisSession();
-	
+
 	if (const FVisualLoggerProvider* VisLogProvider = AnalysisSession->ReadProvider<FVisualLoggerProvider>(FVisualLoggerProvider::ProviderName))
 	{
-		if(EventUpdateRequested > 10)
+		if (EventUpdateRequested > 10)
 		{
 			TRACE_CPUPROFILER_EVENT_SCOPE(FVisualLogTrack::UpdateEventPointsInternal);
 			EventUpdateRequested = 0;
-			
-			EventData->Points.SetNum(0,EAllowShrinking::No);
+
+			EventData->Points.SetNum(0, EAllowShrinking::No);
 			EventData->Windows.SetNum(0);
 
 			TraceServices::FAnalysisSessionReadScope SessionReadScope(*AnalysisSession);
-			
+
 			VisLogProvider->ReadVisualLogEntryTimeline(ObjectId, [this, StartTime, EndTime, VisLogProvider, AnalysisSession](const FVisualLoggerProvider::VisualLogEntryTimeline& InTimeline)
-			{
-				InTimeline.EnumerateEvents(StartTime, EndTime, [this, StartTime, EndTime, VisLogProvider, AnalysisSession](double InStartTime, double InEndTime, uint32 InDepth, const FVisualLogEntry& InMessage)
 				{
-                	for(const FVisualLogShapeElement& Element : InMessage.ElementsToDraw)
-                	{
-						if (Element.Category == Category)
-						{
-							EventData->Points.Add({InMessage.TimeStamp,FText::FromName(Element.Category), FText::FromString(Element.Description),Element.GetFColor()});
-						}
-                	}
-                	for(const FVisualLogLine& Line : InMessage.LogLines)
-                	{
-						if (Line.Category == Category)
+					InTimeline.EnumerateEvents(StartTime, EndTime, [this, StartTime, EndTime, VisLogProvider, AnalysisSession](double InStartTime, double InEndTime, uint32 InDepth, const FVisualLogEntry& InMessage)
 						{
-							// No Description from log lines.
-							EventData->Points.Add({InMessage.TimeStamp, FText::FromName(Line.Category), FText::FromString(""), Line.Color});
-						}
-                	}
-					return TraceServices::EEventEnumerate::Continue;
+							for (const FVisualLogShapeElement& Element : InMessage.ElementsToDraw)
+							{
+								if (Element.Category == Category)
+								{
+									EventData->Points.Add({ InMessage.TimeStamp,GetCategoryNameAsText(Element.Category), FText::FromString(Element.Description),Element.GetFColor() });
+								}
+							}
+							for (const FVisualLogLine& Line : InMessage.LogLines)
+							{
+								if (Line.Category == Category)
+								{
+									// No Description from log lines.
+									EventData->Points.Add({ InMessage.TimeStamp, GetCategoryNameAsText(Line.Category), FText::GetEmpty(), Line.Color });
+								}
+							}
+							return TraceServices::EEventEnumerate::Continue;
+						});
 				});
-			});
 		}
 
 
@@ -130,45 +146,45 @@ bool FVisualLogCategoryTrack::UpdateInternal()
 				DetailsObject = InitializeDetailsObject();
 			}
 
-			DetailsObject->VisualLogDetails.SetNum(0,EAllowShrinking::No);
-			
+			DetailsObject->VisualLogDetails.SetNum(0, EAllowShrinking::No);
+
 			const TraceServices::IFrameProvider& FramesProvider = TraceServices::ReadFrameProvider(*AnalysisSession);
 			TraceServices::FAnalysisSessionReadScope SessionReadScope(*AnalysisSession);
-			
+
 			TraceServices::FFrame MarkerFrame;
-			if(FramesProvider.GetFrameFromTime(ETraceFrameType::TraceFrameType_Game, CurrentScrubTime, MarkerFrame))
+			if (FramesProvider.GetFrameFromTime(ETraceFrameType::TraceFrameType_Game, CurrentScrubTime, MarkerFrame))
 			{
 				VisLogProvider->ReadVisualLogEntryTimeline(ObjectId, [this, DetailsObject, &MarkerFrame, EndTime, VisLogProvider, AnalysisSession](const FVisualLoggerProvider::VisualLogEntryTimeline& InTimeline)
-				{
-					InTimeline.EnumerateEvents(MarkerFrame.StartTime, MarkerFrame.EndTime, [this, DetailsObject, VisLogProvider, AnalysisSession](double InStartTime, double InEndTime, uint32 InDepth, const FVisualLogEntry& InMessage)
 					{
-						for(const FVisualLogShapeElement& Element : InMessage.ElementsToDraw)
-						{
-							if (Element.Category == Category)
+						InTimeline.EnumerateEvents(MarkerFrame.StartTime, MarkerFrame.EndTime, [this, DetailsObject, VisLogProvider, AnalysisSession](double InStartTime, double InEndTime, uint32 InDepth, const FVisualLogEntry& InMessage)
 							{
-								DetailsObject->VisualLogDetails.Add({Element.Category, Element.Description});
-							}
-						}
-						for(const FVisualLogLine& Line : InMessage.LogLines)
-						{
-							if (Line.Category == Category)
-							{
-								DetailsObject->VisualLogDetails.Add({Line.Category, Line.Line});
-							}
-						}
-						return TraceServices::EEventEnumerate::Continue;
+								for (const FVisualLogShapeElement& Element : InMessage.ElementsToDraw)
+								{
+									if (Element.Category == Category)
+									{
+										DetailsObject->VisualLogDetails.Add({ Element.Category, Element.Description });
+									}
+								}
+								for (const FVisualLogLine& Line : InMessage.LogLines)
+								{
+									if (Line.Category == Category)
+									{
+										DetailsObject->VisualLogDetails.Add({ Line.Category, Line.Line });
+									}
+								}
+								return TraceServices::EEventEnumerate::Continue;
+							});
 					});
-				});
 			}
 		}
 	}
-	
+
 	bool bChanged = false;
 	return bChanged;
-	
+
 }
 
-TSharedPtr<SWidget> FVisualLogTrack::GetDetailsViewInternal() 
+TSharedPtr<SWidget> FVisualLogTrack::GetDetailsViewInternal()
 {
 	return nullptr;
 }
@@ -184,18 +200,18 @@ static const FName VisualLogName("Visual Logging");
 
 FName FVisualLogTrackCreator::GetTargetTypeNameInternal() const
 {
- 	static const FName AnimInstanceName("Object");
+	static const FName AnimInstanceName("Object");
 	return AnimInstanceName;
 }
-	
+
 FName FVisualLogTrackCreator::GetNameInternal() const
 {
 	return VisualLogName;
 }
-	
+
 void FVisualLogTrackCreator::GetTrackTypesInternal(TArray<FRewindDebuggerTrackType>& Types) const
 {
-	Types.Add({VisualLogName, LOCTEXT("Visual Logging", "Visual Logging")});
+	Types.Add({ VisualLogName, LOCTEXT("Visual Logging", "Visual Logging") });
 }
 
 TSharedPtr<RewindDebugger::FRewindDebuggerTrack> FVisualLogTrackCreator::CreateTrackInternal(uint64 ObjectId) const
@@ -203,7 +219,7 @@ TSharedPtr<RewindDebugger::FRewindDebuggerTrack> FVisualLogTrackCreator::CreateT
 	return MakeShared<RewindDebugger::FVisualLogTrack>(ObjectId);
 }
 
-		
+
 FVisualLogTrack::FVisualLogTrack(uint64 InObjectId) :
 	ObjectId(InObjectId)
 {
@@ -215,69 +231,69 @@ bool FVisualLogTrack::UpdateInternal()
 {
 	TRACE_CPUPROFILER_EVENT_SCOPE(FVisualLogTrack::UpdateInternal);
 	IRewindDebugger* RewindDebugger = IRewindDebugger::Instance();
-	
+
 	TRange<double> TraceTimeRange = RewindDebugger->GetCurrentTraceRange();
 	double StartTime = TraceTimeRange.GetLowerBoundValue();
 	double EndTime = TraceTimeRange.GetUpperBoundValue();
-	
+
 	const TraceServices::IAnalysisSession* AnalysisSession = RewindDebugger->GetAnalysisSession();
 	const FVisualLoggerProvider* VisLogProvider = AnalysisSession->ReadProvider<FVisualLoggerProvider>(FVisualLoggerProvider::ProviderName);
-	
+
 	bool bChanged = false;
-	
-	if(VisLogProvider)
+
+	if (VisLogProvider)
 	{
 		TArray<FName> UniqueTrackIds;
-		
+
 		TraceServices::FAnalysisSessionReadScope SessionReadScope(*AnalysisSession);
-	
+
 		VisLogProvider->ReadVisualLogEntryTimeline(ObjectId, [this, StartTime, EndTime, VisLogProvider, AnalysisSession, &UniqueTrackIds](const FVisualLoggerProvider::VisualLogEntryTimeline& InTimeline)
-		{
-			InTimeline.EnumerateEvents(StartTime, EndTime, [this, StartTime, EndTime, VisLogProvider, AnalysisSession, &UniqueTrackIds](double InStartTime, double InEndTime, uint32 InDepth, const FVisualLogEntry& InMessage)
 			{
-				if(InMessage.LogLines.Num() > 0)
-				{
-					for(const FVisualLogLine& Line : InMessage.LogLines)
+				InTimeline.EnumerateEvents(StartTime, EndTime, [this, StartTime, EndTime, VisLogProvider, AnalysisSession, &UniqueTrackIds](double InStartTime, double InEndTime, uint32 InDepth, const FVisualLogEntry& InMessage)
 					{
-						UniqueTrackIds.AddUnique(Line.Category);
-					}
-				}
-				for(const FVisualLogShapeElement& Element : InMessage.ElementsToDraw)
-				{
-					UniqueTrackIds.AddUnique(Element.Category);
-				}
-				return TraceServices::EEventEnumerate::Continue;
+						if (InMessage.LogLines.Num() > 0)
+						{
+							for (const FVisualLogLine& Line : InMessage.LogLines)
+							{
+								UniqueTrackIds.AddUnique(Line.Category);
+							}
+						}
+						for (const FVisualLogShapeElement& Element : InMessage.ElementsToDraw)
+						{
+							UniqueTrackIds.AddUnique(Element.Category);
+						}
+						return TraceServices::EEventEnumerate::Continue;
+					});
 			});
-		});
-	
+
 		UniqueTrackIds.StableSort([](const FName& A, const FName& B) { return A.ToString() < B.ToString(); });
 		const int32 TrackCount = UniqueTrackIds.Num();
-	
+
 		if (Children.Num() != TrackCount)
 			bChanged = true;
-	
+
 		Children.SetNum(UniqueTrackIds.Num());
-		for(int i = 0; i < TrackCount; i++)
+		for (int i = 0; i < TrackCount; i++)
 		{
 			if (!Children[i].IsValid() || !(Children[i].Get()->GetName() == UniqueTrackIds[i]))
 			{
 				Children[i] = MakeShared<FVisualLogCategoryTrack>(ObjectId, UniqueTrackIds[i]);
 				bChanged = true;
 			}
-	
+
 			if (Children[i]->Update())
 			{
 				bChanged = true;
 			}
 		}
 	}
-	
+
 	return bChanged;
 }
-	
+
 void FVisualLogTrack::IterateSubTracksInternal(TFunction<void(TSharedPtr<FRewindDebuggerTrack> SubTrack)> IteratorFunction)
 {
-	for(TSharedPtr<FVisualLogCategoryTrack>& Track : Children)
+	for (TSharedPtr<FVisualLogCategoryTrack>& Track : Children)
 	{
 		IteratorFunction(Track);
 	}
@@ -287,20 +303,19 @@ bool FVisualLogTrackCreator::HasDebugInfoInternal(uint64 ObjectId) const
 {
 	TRACE_CPUPROFILER_EVENT_SCOPE(FVisualLogTrack::HasDebugInfoInternal);
 	const TraceServices::IAnalysisSession* AnalysisSession = IRewindDebugger::Instance()->GetAnalysisSession();
-	
+
 	TraceServices::FAnalysisSessionReadScope SessionReadScope(*AnalysisSession);
 	bool bHasData = false;
 	if (const FVisualLoggerProvider* VLogProvider = AnalysisSession->ReadProvider<FVisualLoggerProvider>(FVisualLoggerProvider::ProviderName))
 	{
-		VLogProvider->ReadVisualLogEntryTimeline(ObjectId,  [&bHasData](const FVisualLoggerProvider::VisualLogEntryTimeline& InTimeline)
-	 	{
-	 		bHasData = true;
-	 	});
+		VLogProvider->ReadVisualLogEntryTimeline(ObjectId, [&bHasData](const FVisualLoggerProvider::VisualLogEntryTimeline& InTimeline)
+			{
+				bHasData = true;
+			});
 	}
 	return bHasData;
 }
 
-	
 }
 
 #undef LOCTEXT_NAMESPACE
\ No newline at end of file
