[RewindDebugger] replaced single uint64 identifier for traced object by a struct allowing more complex scenarios to add non-UObject based entries. The traced object is now a pair of uint64 for the main UObject and another uint64 to represent a child element.

diff --git a/Engine/Plugins/Animation/GameplayInsights/Source/GameplayInsights/Private/AnimCurvesTrack.cpp b/Engine/Plugins/Animation/GameplayInsights/Source/GameplayInsights/Private/AnimCurvesTrack.cpp
--- a/Engine/Plugins/Animation/GameplayInsights/Source/GameplayInsights/Private/AnimCurvesTrack.cpp
+++ b/Engine/Plugins/Animation/GameplayInsights/Source/GameplayInsights/Private/AnimCurvesTrack.cpp
@@ -233,12 +233,12 @@ void FAnimationCurvesTrackCreator::GetTrackTypesInternal(TArray<FRewindDebuggerT
 	Types.Add({AnimationCurvesName, LOCTEXT("Curves", "Curves")});
 }
 
-TSharedPtr<RewindDebugger::FRewindDebuggerTrack> FAnimationCurvesTrackCreator::CreateTrackInternal(uint64 ObjectId) const
+TSharedPtr<RewindDebugger::FRewindDebuggerTrack> FAnimationCurvesTrackCreator::CreateTrackInternal(const FObjectId& InObjectId) const
 {
-	return MakeShared<RewindDebugger::FAnimCurvesTrack>(ObjectId);
+	return MakeShared<RewindDebugger::FAnimCurvesTrack>(InObjectId.GetMainId());
 }
 
-bool FAnimationCurvesTrackCreator::HasDebugInfoInternal(uint64 ObjectId) const
+bool FAnimationCurvesTrackCreator::HasDebugInfoInternal(const FObjectId& InObjectId) const
 {
 	const TraceServices::IAnalysisSession* AnalysisSession = IRewindDebugger::Instance()->GetAnalysisSession();
 	
@@ -246,7 +246,7 @@ bool FAnimationCurvesTrackCreator::HasDebugInfoInternal(uint64 ObjectId) const
 	bool bHasData = false;
 	if (const FAnimationProvider* AnimationProvider = AnalysisSession->ReadProvider<FAnimationProvider>(FAnimationProvider::ProviderName))
 	{
-		AnimationProvider->ReadSkeletalMeshPoseTimeline(ObjectId, [&bHasData](const FAnimationProvider::SkeletalMeshPoseTimeline& InTimeline, bool bInHasCurves)
+		AnimationProvider->ReadSkeletalMeshPoseTimeline(InObjectId.GetMainId(), [&bHasData](const FAnimationProvider::SkeletalMeshPoseTimeline& InTimeline, bool bInHasCurves)
 		{
 			bHasData = bInHasCurves;
 		});
diff --git a/Engine/Plugins/Animation/GameplayInsights/Source/GameplayInsights/Private/AnimCurvesTrack.h b/Engine/Plugins/Animation/GameplayInsights/Source/GameplayInsights/Private/AnimCurvesTrack.h
--- a/Engine/Plugins/Animation/GameplayInsights/Source/GameplayInsights/Private/AnimCurvesTrack.h
+++ b/Engine/Plugins/Animation/GameplayInsights/Source/GameplayInsights/Private/AnimCurvesTrack.h
@@ -74,8 +74,8 @@ private:
 	virtual FName GetTargetTypeNameInternal() const override;
 	virtual FName GetNameInternal() const override;
 	virtual void GetTrackTypesInternal(TArray<FRewindDebuggerTrackType>& Types) const override;
-	virtual TSharedPtr<FRewindDebuggerTrack> CreateTrackInternal(uint64 ObjectId) const override;
-	virtual bool HasDebugInfoInternal(uint64 ObjectId) const override;
+	virtual TSharedPtr<FRewindDebuggerTrack> CreateTrackInternal(const FObjectId& InObjectId) const override;
+	virtual bool HasDebugInfoInternal(const FObjectId& InObjectId) const override;
 };
 
 }
\ No newline at end of file
diff --git a/Engine/Plugins/Animation/GameplayInsights/Source/GameplayInsights/Private/AnimNodesTrack.cpp b/Engine/Plugins/Animation/GameplayInsights/Source/GameplayInsights/Private/AnimNodesTrack.cpp
--- a/Engine/Plugins/Animation/GameplayInsights/Source/GameplayInsights/Private/AnimNodesTrack.cpp
+++ b/Engine/Plugins/Animation/GameplayInsights/Source/GameplayInsights/Private/AnimNodesTrack.cpp
@@ -220,7 +220,7 @@ void FAnimNodesTrack::BuildContextMenu(FMenuBuilder& MenuBuilder)
 								if(const FObjectInfo* AnimInstanceObjectInfo = GameplayProvider->FindObjectInfo(GetGameplayTrack().GetObjectId()))
 								{
 									// find a counterpart skeletal mesh pose track
-									TSharedPtr<FSkeletalMeshPoseTrack> SkeletalMeshPoseTrack = SharedData.FindSkeletalMeshPoseTrack(AnimInstanceObjectInfo->OuterId);
+									TSharedPtr<FSkeletalMeshPoseTrack> SkeletalMeshPoseTrack = SharedData.FindSkeletalMeshPoseTrack(AnimInstanceObjectInfo->GetOuterUObjectId());
 									if(SkeletalMeshPoseTrack.IsValid())
 									{
 										SkeletalMeshPoseTrack->MarkPotentiallyDebugged();
@@ -251,10 +251,10 @@ void FAnimNodesTrack::BuildContextMenu(FMenuBuilder& MenuBuilder)
 							{
 								TraceServices::FAnalysisSessionReadScope SessionReadScope(SharedData.GetAnalysisSession());
 
-								if(const FObjectInfo* AnimInstanceObjectInfo = GameplayProvider->FindObjectInfo(GetGameplayTrack().GetObjectId()))
+								if(const FObjectInfo* AnimInstanceObjectInfo = GameplayProvider->FindObjectInfo(GetGameplayTrack().GetAssociatedObjectId()))
 								{
 									// find a counterpart skeletal mesh pose track
-									TSharedPtr<FSkeletalMeshPoseTrack> SkeletalMeshPoseTrack = SharedData.FindSkeletalMeshPoseTrack(AnimInstanceObjectInfo->OuterId);
+									TSharedPtr<FSkeletalMeshPoseTrack> SkeletalMeshPoseTrack = SharedData.FindSkeletalMeshPoseTrack(AnimInstanceObjectInfo->GetOuterUObjectId());
 									if(SkeletalMeshPoseTrack.IsValid())
 									{
 										if(USkeletalMeshComponent* SkeletalMeshComponent = SkeletalMeshPoseTrack->GetComponent(FGameplayTimingViewExtender::GetWorldToVisualize()))
@@ -448,7 +448,7 @@ void FAnimNodesTrack::GetCustomDebugObjects(const IAnimationBlueprintEditor& InA
 				if(const FObjectInfo* AnimInstanceObjectInfo = GameplayProvider->FindObjectInfo(GetGameplayTrack().GetObjectId()))
 				{
 					// find a counterpart skeletal mesh pose track
-					TSharedPtr<FSkeletalMeshPoseTrack> SkeletalMeshPoseTrack = SharedData.FindSkeletalMeshPoseTrack(AnimInstanceObjectInfo->OuterId);
+					TSharedPtr<FSkeletalMeshPoseTrack> SkeletalMeshPoseTrack = SharedData.FindSkeletalMeshPoseTrack(AnimInstanceObjectInfo->GetOuterUObjectId());
 					if(SkeletalMeshPoseTrack.IsValid())
 					{
 						SkeletalMeshPoseTrack->MarkPotentiallyDebugged();
diff --git a/Engine/Plugins/Animation/GameplayInsights/Source/GameplayInsights/Private/AnimationSharedData.cpp b/Engine/Plugins/Animation/GameplayInsights/Source/GameplayInsights/Private/AnimationSharedData.cpp
--- a/Engine/Plugins/Animation/GameplayInsights/Source/GameplayInsights/Private/AnimationSharedData.cpp
+++ b/Engine/Plugins/Animation/GameplayInsights/Source/GameplayInsights/Private/AnimationSharedData.cpp
@@ -23,11 +23,6 @@
 
 // TraceInsights
 #include "Insights/ITimingViewSession.h"
-#include "Insights/ViewModels/TimingEvent.h"
-
-#if WITH_ENGINE
-#include "Animation/AnimTypes.h"
-#endif 
 
 #if WITH_EDITOR
 #include "EditorViewportClient.h"
@@ -88,8 +83,8 @@ void FAnimationSharedData::Tick(UE::Insights::Timing::ITimingViewSession& InTimi
 			GameplayProvider->EnumerateObjects([this, &InTimingViewSession, &InAnalysisSession, &AnimationProvider, &GameplayProvider](const FObjectInfo& InObjectInfo)
 			{
 				TSharedPtr<FObjectEventsTrack> ObjectEventsTrack;
-
-				AnimationProvider->ReadSkeletalMeshPoseTimeline(InObjectInfo.Id, [this, &InObjectInfo, &ObjectEventsTrack, &InTimingViewSession, &InAnalysisSession](const IAnimationProvider::SkeletalMeshPoseTimeline& InTimeline, bool bInHasCurves)
+				const uint64 ObjectId = InObjectInfo.GetUObjectId();
+				AnimationProvider->ReadSkeletalMeshPoseTimeline(ObjectId, [this, ObjectId, &InObjectInfo, &ObjectEventsTrack, &InTimingViewSession, &InAnalysisSession](const IAnimationProvider::SkeletalMeshPoseTimeline& InTimeline, bool bInHasCurves)
 				{
 					if(!ObjectEventsTrack.IsValid())
 					{
@@ -101,10 +96,10 @@ void FAnimationSharedData::Tick(UE::Insights::Timing::ITimingViewSession& InTimi
 						return InTrack.Is<FSkeletalMeshPoseTrack>();
 					};
 
-					TSharedPtr<FSkeletalMeshPoseTrack> ExistingSkeletalMeshPoseTrack = StaticCastSharedPtr<FSkeletalMeshPoseTrack>(ObjectEventsTrack->GetGameplayTrack().FindChildTrack(InObjectInfo.Id, FindSkeletalMeshPoseTrack));
+					TSharedPtr<FSkeletalMeshPoseTrack> ExistingSkeletalMeshPoseTrack = StaticCastSharedPtr<FSkeletalMeshPoseTrack>(ObjectEventsTrack->GetGameplayTrack().FindChildTrack(ObjectId, FindSkeletalMeshPoseTrack));
 					if(!ExistingSkeletalMeshPoseTrack.IsValid())
 					{
-						TSharedPtr<FSkeletalMeshPoseTrack> SkeletalMeshPoseTrack = MakeShared<FSkeletalMeshPoseTrack>(*this, InObjectInfo.Id, InObjectInfo.Name);
+						TSharedPtr<FSkeletalMeshPoseTrack> SkeletalMeshPoseTrack = MakeShared<FSkeletalMeshPoseTrack>(*this, ObjectId, InObjectInfo.Name);
 						SkeletalMeshPoseTrack->SetVisibilityFlag(bSkeletalMeshPoseTracksEnabled);
 						SkeletalMeshPoseTracks.Add(SkeletalMeshPoseTrack.ToSharedRef());
 
@@ -123,10 +118,10 @@ void FAnimationSharedData::Tick(UE::Insights::Timing::ITimingViewSession& InTimi
 							return InTrack.Is<FSkeletalMeshCurvesTrack>();
 						};
 
-						TSharedPtr<FSkeletalMeshCurvesTrack> ExistingSkeletalMeshCurvesTrack = StaticCastSharedPtr<FSkeletalMeshCurvesTrack>(ObjectEventsTrack->GetGameplayTrack().FindChildTrack(InObjectInfo.Id, FindSkeletalMeshCurvesTrack));
+						TSharedPtr<FSkeletalMeshCurvesTrack> ExistingSkeletalMeshCurvesTrack = StaticCastSharedPtr<FSkeletalMeshCurvesTrack>(ObjectEventsTrack->GetGameplayTrack().FindChildTrack(ObjectId, FindSkeletalMeshCurvesTrack));
 						if(!ExistingSkeletalMeshCurvesTrack.IsValid())
 						{
-							TSharedPtr<FSkeletalMeshCurvesTrack> SkeletalMeshCurvesTrack = MakeShared<FSkeletalMeshCurvesTrack>(*this, InObjectInfo.Id, InObjectInfo.Name);
+							TSharedPtr<FSkeletalMeshCurvesTrack> SkeletalMeshCurvesTrack = MakeShared<FSkeletalMeshCurvesTrack>(*this, ObjectId, InObjectInfo.Name);
 							SkeletalMeshCurvesTrack->SetVisibilityFlag(bSkeletalMeshCurveTracksEnabled);
 							SkeletalMeshCurvesTracks.Add(SkeletalMeshCurvesTrack.ToSharedRef());
 
@@ -140,7 +135,7 @@ void FAnimationSharedData::Tick(UE::Insights::Timing::ITimingViewSession& InTimi
 					}
 				});
 
-				AnimationProvider->ReadTickRecordTimeline(InObjectInfo.Id, [this, &InObjectInfo, &ObjectEventsTrack, &InTimingViewSession, &GameplayProvider, &InAnalysisSession](const IAnimationProvider::TickRecordTimeline& InTimeline)
+				AnimationProvider->ReadTickRecordTimeline(ObjectId, [this, ObjectId, &InObjectInfo, &ObjectEventsTrack, &InTimingViewSession, &GameplayProvider, &InAnalysisSession](const IAnimationProvider::TickRecordTimeline& InTimeline)
 				{
 					if(!ObjectEventsTrack.IsValid())
 					{
@@ -152,10 +147,10 @@ void FAnimationSharedData::Tick(UE::Insights::Timing::ITimingViewSession& InTimi
 						return InTrack.Is<FAnimationTickRecordsTrack>();
 					};
 
-					TSharedPtr<FAnimationTickRecordsTrack> ExistingAnimationTickRecordsTrack = StaticCastSharedPtr<FAnimationTickRecordsTrack>(ObjectEventsTrack->GetGameplayTrack().FindChildTrack(InObjectInfo.Id, FindTickRecordTrackWithAssetId));
+					TSharedPtr<FAnimationTickRecordsTrack> ExistingAnimationTickRecordsTrack = StaticCastSharedPtr<FAnimationTickRecordsTrack>(ObjectEventsTrack->GetGameplayTrack().FindChildTrack(ObjectId, FindTickRecordTrackWithAssetId));
 					if(!ExistingAnimationTickRecordsTrack.IsValid())
 					{
-						TSharedPtr<FAnimationTickRecordsTrack> AnimationTickRecordsTrack = MakeShared<FAnimationTickRecordsTrack>(*this, InObjectInfo.Id, InObjectInfo.Name);
+						TSharedPtr<FAnimationTickRecordsTrack> AnimationTickRecordsTrack = MakeShared<FAnimationTickRecordsTrack>(*this, ObjectId, InObjectInfo.Name);
 						AnimationTickRecordsTrack->SetVisibilityFlag(bTickRecordTracksEnabled);
 						AnimationTickRecordsTracks.Add(AnimationTickRecordsTrack.ToSharedRef());
 
@@ -168,7 +163,7 @@ void FAnimationSharedData::Tick(UE::Insights::Timing::ITimingViewSession& InTimi
 					}
 				});
 
-				AnimationProvider->ReadAnimGraphTimeline(InObjectInfo.Id, [this, &InObjectInfo, &ObjectEventsTrack, &InTimingViewSession, &InAnalysisSession](const IAnimationProvider::AnimGraphTimeline& InTimeline)
+				AnimationProvider->ReadAnimGraphTimeline(ObjectId, [this, ObjectId, &InObjectInfo, &ObjectEventsTrack, &InTimingViewSession, &InAnalysisSession](const IAnimationProvider::AnimGraphTimeline& InTimeline)
 				{
 					if(!ObjectEventsTrack.IsValid())
 					{
@@ -180,10 +175,10 @@ void FAnimationSharedData::Tick(UE::Insights::Timing::ITimingViewSession& InTimi
 						return InTrack.Is<FAnimNodesTrack>();
 					};
 
-					TSharedPtr<FAnimNodesTrack> ExistingAnimNodesTrack = StaticCastSharedPtr<FAnimNodesTrack>(ObjectEventsTrack->GetGameplayTrack().FindChildTrack(InObjectInfo.Id, FindAnimNodesTrack));
+					TSharedPtr<FAnimNodesTrack> ExistingAnimNodesTrack = StaticCastSharedPtr<FAnimNodesTrack>(ObjectEventsTrack->GetGameplayTrack().FindChildTrack(ObjectId, FindAnimNodesTrack));
 					if(!ExistingAnimNodesTrack.IsValid())
 					{
-						TSharedPtr<FAnimNodesTrack> AnimNodesTrack = MakeShared<FAnimNodesTrack>(*this, InObjectInfo.Id, InObjectInfo.Name);
+						TSharedPtr<FAnimNodesTrack> AnimNodesTrack = MakeShared<FAnimNodesTrack>(*this, ObjectId, InObjectInfo.Name);
 						AnimNodesTrack->SetVisibilityFlag(bAnimNodeTracksEnabled);
 						AnimNodesTracks.Add(AnimNodesTrack.ToSharedRef());
 
@@ -196,7 +191,7 @@ void FAnimationSharedData::Tick(UE::Insights::Timing::ITimingViewSession& InTimi
 					}
 				});
 
-				auto ProcessNotifyTimeline = [this, &InObjectInfo, &ObjectEventsTrack, &InTimingViewSession, &InAnalysisSession]()
+				auto ProcessNotifyTimeline = [this, ObjectId, &InObjectInfo, &ObjectEventsTrack, &InTimingViewSession, &InAnalysisSession]()
 				{
 					if(!ObjectEventsTrack.IsValid())
 					{
@@ -208,10 +203,10 @@ void FAnimationSharedData::Tick(UE::Insights::Timing::ITimingViewSession& InTimi
 						return InTrack.Is<FAnimNotifiesTrack>();
 					};
 
-					TSharedPtr<FAnimNotifiesTrack> ExistingAnimNotifyTrack = StaticCastSharedPtr<FAnimNotifiesTrack>(ObjectEventsTrack->GetGameplayTrack().FindChildTrack(InObjectInfo.Id, FindAnimNotifyTrack));
+					TSharedPtr<FAnimNotifiesTrack> ExistingAnimNotifyTrack = StaticCastSharedPtr<FAnimNotifiesTrack>(ObjectEventsTrack->GetGameplayTrack().FindChildTrack(ObjectId, FindAnimNotifyTrack));
 					if(!ExistingAnimNotifyTrack.IsValid())
 					{
-						TSharedPtr<FAnimNotifiesTrack> AnimNotifyTrack = MakeShared<FAnimNotifiesTrack>(*this, InObjectInfo.Id, InObjectInfo.Name);
+						TSharedPtr<FAnimNotifiesTrack> AnimNotifyTrack = MakeShared<FAnimNotifiesTrack>(*this, ObjectId, InObjectInfo.Name);
 						AnimNotifyTrack->SetVisibilityFlag(bAnimNotifyTracksEnabled);
 						AnimNotifyTracks.Add(AnimNotifyTrack.ToSharedRef());
 
@@ -224,17 +219,17 @@ void FAnimationSharedData::Tick(UE::Insights::Timing::ITimingViewSession& InTimi
 					}
 				};
 
-				AnimationProvider->ReadNotifyTimeline(InObjectInfo.Id, [&ProcessNotifyTimeline](const IAnimationProvider::AnimNotifyTimeline& InTimeline)
+				AnimationProvider->ReadNotifyTimeline(ObjectId, [&ProcessNotifyTimeline](const IAnimationProvider::AnimNotifyTimeline& InTimeline)
 				{
 					ProcessNotifyTimeline();
 				});
 
-				AnimationProvider->EnumerateNotifyStateTimelines(InObjectInfo.Id, [&ProcessNotifyTimeline](uint32 InNotifyNameId, const IAnimationProvider::AnimNotifyTimeline& InTimeline)
+				AnimationProvider->EnumerateNotifyStateTimelines(ObjectId, [&ProcessNotifyTimeline](uint32 InNotifyNameId, const IAnimationProvider::AnimNotifyTimeline& InTimeline)
 				{
 					ProcessNotifyTimeline();
 				});
 
-				AnimationProvider->ReadMontageTimeline(InObjectInfo.Id, [this, &InObjectInfo, &ObjectEventsTrack, &InTimingViewSession, &InAnalysisSession](const IAnimationProvider::AnimMontageTimeline& InTimeline)
+				AnimationProvider->ReadMontageTimeline(ObjectId, [this, ObjectId, &InObjectInfo, &ObjectEventsTrack, &InTimingViewSession, &InAnalysisSession](const IAnimationProvider::AnimMontageTimeline& InTimeline)
 				{
 					if(!ObjectEventsTrack.IsValid())
 					{
@@ -246,10 +241,10 @@ void FAnimationSharedData::Tick(UE::Insights::Timing::ITimingViewSession& InTimi
 						return InTrack.Is<FMontageTrack>();
 					};
 
-					TSharedPtr<FMontageTrack> ExistingMontageTrack = StaticCastSharedPtr<FMontageTrack>(ObjectEventsTrack->GetGameplayTrack().FindChildTrack(InObjectInfo.Id, FindMontageTrack));
+					TSharedPtr<FMontageTrack> ExistingMontageTrack = StaticCastSharedPtr<FMontageTrack>(ObjectEventsTrack->GetGameplayTrack().FindChildTrack(ObjectId, FindMontageTrack));
 					if(!ExistingMontageTrack.IsValid())
 					{
-						TSharedPtr<FMontageTrack> MontageTrack = MakeShared<FMontageTrack>(*this, InObjectInfo.Id, InObjectInfo.Name);
+						TSharedPtr<FMontageTrack> MontageTrack = MakeShared<FMontageTrack>(*this, ObjectId, InObjectInfo.Name);
 						MontageTrack->SetVisibilityFlag(bMontageTracksEnabled);
 						MontageTracks.Add(MontageTrack.ToSharedRef());
 
diff --git a/Engine/Plugins/Animation/GameplayInsights/Source/GameplayInsights/Private/BlendWeightsTrack.cpp b/Engine/Plugins/Animation/GameplayInsights/Source/GameplayInsights/Private/BlendWeightsTrack.cpp
--- a/Engine/Plugins/Animation/GameplayInsights/Source/GameplayInsights/Private/BlendWeightsTrack.cpp
+++ b/Engine/Plugins/Animation/GameplayInsights/Source/GameplayInsights/Private/BlendWeightsTrack.cpp
@@ -443,12 +443,12 @@ void FBlendWeightsTrackCreator::GetTrackTypesInternal(TArray<FRewindDebuggerTrac
 	Types.Add({BlendWeightsName, LOCTEXT("Blend Weights", "Blend Weights")});
 }
 
-TSharedPtr<RewindDebugger::FRewindDebuggerTrack> FBlendWeightsTrackCreator::CreateTrackInternal(uint64 ObjectId) const
+TSharedPtr<RewindDebugger::FRewindDebuggerTrack> FBlendWeightsTrackCreator::CreateTrackInternal(const FObjectId& InObjectId) const
 {
- 	return MakeShared<RewindDebugger::FBlendWeightsTrack>(ObjectId);
+ 	return MakeShared<RewindDebugger::FBlendWeightsTrack>(InObjectId.GetMainId());
 }
 
-bool FBlendWeightsTrackCreator::HasDebugInfoInternal(uint64 ObjectId) const
+bool FBlendWeightsTrackCreator::HasDebugInfoInternal(const FObjectId& InObjectId) const
 {
 	TRACE_CPUPROFILER_EVENT_SCOPE(FBlendWeightsTrack::HasDebugInfoInternal);
 	const TraceServices::IAnalysisSession* AnalysisSession = IRewindDebugger::Instance()->GetAnalysisSession();
@@ -456,7 +456,7 @@ bool FBlendWeightsTrackCreator::HasDebugInfoInternal(uint64 ObjectId) const
 	bool bHasData = false;
 	if (const FAnimationProvider* AnimationProvider = AnalysisSession->ReadProvider<FAnimationProvider>(FAnimationProvider::ProviderName))
 	{
-		AnimationProvider->ReadTickRecordTimeline(ObjectId, [&bHasData](const FAnimationProvider::TickRecordTimeline& InGraphTimeline)
+		AnimationProvider->ReadTickRecordTimeline(InObjectId.GetMainId(), [&bHasData](const FAnimationProvider::TickRecordTimeline& InGraphTimeline)
 		{
 			bHasData = true;
 		});
diff --git a/Engine/Plugins/Animation/GameplayInsights/Source/GameplayInsights/Private/BlendWeightsTrack.h b/Engine/Plugins/Animation/GameplayInsights/Source/GameplayInsights/Private/BlendWeightsTrack.h
--- a/Engine/Plugins/Animation/GameplayInsights/Source/GameplayInsights/Private/BlendWeightsTrack.h
+++ b/Engine/Plugins/Animation/GameplayInsights/Source/GameplayInsights/Private/BlendWeightsTrack.h
@@ -87,8 +87,8 @@ public:
 	virtual FName GetTargetTypeNameInternal() const;
 	virtual FName GetNameInternal() const override;
 	virtual void GetTrackTypesInternal(TArray<FRewindDebuggerTrackType>& Types) const override;
-	virtual TSharedPtr<FRewindDebuggerTrack> CreateTrackInternal(uint64 ObjectId) const override;
-	virtual bool HasDebugInfoInternal(uint64 ObjectId) const override;
+	virtual TSharedPtr<FRewindDebuggerTrack> CreateTrackInternal(const FObjectId& InObjectId) const override;
+	virtual bool HasDebugInfoInternal(const FObjectId& InObjectId) const override;
 };
 
 }
\ No newline at end of file
diff --git a/Engine/Plugins/Animation/GameplayInsights/Source/GameplayInsights/Private/ExternalMorphTrack.cpp b/Engine/Plugins/Animation/GameplayInsights/Source/GameplayInsights/Private/ExternalMorphTrack.cpp
--- a/Engine/Plugins/Animation/GameplayInsights/Source/GameplayInsights/Private/ExternalMorphTrack.cpp
+++ b/Engine/Plugins/Animation/GameplayInsights/Source/GameplayInsights/Private/ExternalMorphTrack.cpp
@@ -287,18 +287,18 @@ namespace RewindDebugger
     	Types.Add({ExternalMorphSetName, LOCTEXT("External Morph Sets", "External Morph Sets")});
     }
 
-	TSharedPtr<RewindDebugger::FRewindDebuggerTrack> FExternalMorphSetGroupTrackCreator::CreateTrackInternal(uint64 ObjectId) const
+	TSharedPtr<RewindDebugger::FRewindDebuggerTrack> FExternalMorphSetGroupTrackCreator::CreateTrackInternal(const FObjectId& InObjectId) const
 	{
-		return MakeShared<RewindDebugger::FExternalMorphSetGroupTrack>(ObjectId);
+		return MakeShared<RewindDebugger::FExternalMorphSetGroupTrack>(InObjectId.GetMainId());
 	}
 
-	bool FExternalMorphSetGroupTrackCreator::HasDebugInfoInternal(uint64 ObjectId) const
+	bool FExternalMorphSetGroupTrackCreator::HasDebugInfoInternal(const FObjectId& InObjectId) const
 	{
 		const TraceServices::IAnalysisSession* AnalysisSession = IRewindDebugger::Instance()->GetAnalysisSession();
 		TraceServices::FAnalysisSessionReadScope SessionReadScope(*AnalysisSession);
 		if (const FAnimationProvider* AnimationProvider = AnalysisSession->ReadProvider<FAnimationProvider>(FAnimationProvider::ProviderName))
 		{
-			return AnimationProvider->HasExternalMorphSets(ObjectId);
+			return AnimationProvider->HasExternalMorphSets(InObjectId.GetMainId());
 		}
 		return false;
 	}
diff --git a/Engine/Plugins/Animation/GameplayInsights/Source/GameplayInsights/Private/ExternalMorphTrack.h b/Engine/Plugins/Animation/GameplayInsights/Source/GameplayInsights/Private/ExternalMorphTrack.h
--- a/Engine/Plugins/Animation/GameplayInsights/Source/GameplayInsights/Private/ExternalMorphTrack.h
+++ b/Engine/Plugins/Animation/GameplayInsights/Source/GameplayInsights/Private/ExternalMorphTrack.h
@@ -106,8 +106,8 @@ namespace RewindDebugger
 		virtual FName GetTargetTypeNameInternal() const override;
 		virtual FName GetNameInternal() const override;
 		virtual void GetTrackTypesInternal(TArray<FRewindDebuggerTrackType>& Types) const override;
-		virtual TSharedPtr<FRewindDebuggerTrack> CreateTrackInternal(uint64 ObjectId) const override;
-		virtual bool HasDebugInfoInternal(uint64 ObjectId) const override;
+		virtual TSharedPtr<FRewindDebuggerTrack> CreateTrackInternal(const FObjectId& InObjectId) const override;
+		virtual bool HasDebugInfoInternal(const FObjectId& InObjectId) const override;
 		// ~END IRewindDebuggerTrackCreator overrides.
 	};
 
diff --git a/Engine/Plugins/Animation/GameplayInsights/Source/GameplayInsights/Private/GameplayAnalyzer.cpp b/Engine/Plugins/Animation/GameplayInsights/Source/GameplayInsights/Private/GameplayAnalyzer.cpp
--- a/Engine/Plugins/Animation/GameplayInsights/Source/GameplayInsights/Private/GameplayAnalyzer.cpp
+++ b/Engine/Plugins/Animation/GameplayInsights/Source/GameplayInsights/Private/GameplayAnalyzer.cpp
@@ -35,6 +35,10 @@ void FGameplayAnalyzer::OnAnalysisBegin(const FOnAnalysisContext& Context)
 	Builder.RouteEvent(RouteId_PropertiesEnd, "Object", "PropertiesEnd");
 	Builder.RouteEvent(RouteId_PropertyValue2, "Object", "PropertyValue2");
 	Builder.RouteEvent(RouteId_PropertyValue3, "Object", "PropertyValue3");
+	Builder.RouteEvent(RouteId_ChildElement, "Object", "ObjectChildElement");
+	Builder.RouteEvent(RouteId_ChildElementLifetimeBegin, "Object", "ObjectChildElementLifetimeBegin");
+	Builder.RouteEvent(RouteId_ChildElementLifetimeEnd, "Object", "ObjectChildElementLifetimeEnd");
+	Builder.RouteEvent(RouteId_ChildElementEvent, "Object", "ObjectChildElementEvent");
 }
 
 bool FGameplayAnalyzer::OnEvent(uint16 RouteId, EStyle Style, const FOnEventContext& Context)
@@ -103,7 +107,7 @@ bool FGameplayAnalyzer::OnEvent(uint16 RouteId, EStyle Style, const FOnEventCont
 		uint64 Id = EventData.GetValue<uint64>("Id");
 		uint64 OuterId = EventData.GetValue<uint64>("OuterId");
 		uint64 ClassId = EventData.GetValue<uint64>("ClassId");
-		GameplayProvider.AppendObject(Id, OuterId, ClassId, *ObjectName, *ObjectPathName);
+		GameplayProvider.AppendObject(RewindDebugger::FObjectId{Id}, RewindDebugger::FObjectId{OuterId}, ClassId, *ObjectName, *ObjectPathName);
 		break;
 	}
 	case RouteId_ObjectTransform:
@@ -122,14 +126,14 @@ bool FGameplayAnalyzer::OnEvent(uint16 RouteId, EStyle Style, const FOnEventCont
 		uint64 Cycle = EventData.GetValue<uint64>("Cycle");
 		uint64 Id = EventData.GetValue<uint64>("Id");
 		FString Event = TraceServices::FTraceAnalyzerUtils::LegacyAttachmentString<TCHAR>("Event", Context);
-		GameplayProvider.AppendObjectEvent(Id, Context.EventTime.AsSeconds(Cycle), *Event);
+		GameplayProvider.AppendObjectEvent(RewindDebugger::FObjectId{Id}, Context.EventTime.AsSeconds(Cycle), *Event);
 		break;
 	}
 	case RouteId_ObjectLifetimeBegin:
 	{
 		uint64 Cycle = EventData.GetValue<uint64>("Cycle");
 		uint64 Id = EventData.GetValue<uint64>("Id");
-		GameplayProvider.AppendObjectLifetimeBegin(Id, Context.EventTime.AsSeconds(Cycle), 0.0);
+		GameplayProvider.AppendObjectLifetimeBegin(RewindDebugger::FObjectId{Id}, Context.EventTime.AsSeconds(Cycle), 0.0);
 		break;
 	}
 	case RouteId_ObjectLifetimeBegin2:
@@ -137,14 +141,14 @@ bool FGameplayAnalyzer::OnEvent(uint16 RouteId, EStyle Style, const FOnEventCont
 		uint64 Cycle = EventData.GetValue<uint64>("Cycle");
 		double RecordingTime = EventData.GetValue<double>("RecordingTime");
 		uint64 Id = EventData.GetValue<uint64>("Id");
-		GameplayProvider.AppendObjectLifetimeBegin(Id, Context.EventTime.AsSeconds(Cycle), RecordingTime);
+		GameplayProvider.AppendObjectLifetimeBegin(RewindDebugger::FObjectId{Id}, Context.EventTime.AsSeconds(Cycle), RecordingTime);
 		break;
 	}
 	case RouteId_ObjectLifetimeEnd:
 	{
 		uint64 Cycle = EventData.GetValue<uint64>("Cycle");
 		uint64 Id = EventData.GetValue<uint64>("Id");
-		GameplayProvider.AppendObjectLifetimeEnd(Id, Context.EventTime.AsSeconds(Cycle), 0.0);
+		GameplayProvider.AppendObjectLifetimeEnd(RewindDebugger::FObjectId{Id}, Context.EventTime.AsSeconds(Cycle), 0.0);
 		break;
 	}
 	case RouteId_ObjectLifetimeEnd2:
@@ -152,7 +156,7 @@ bool FGameplayAnalyzer::OnEvent(uint16 RouteId, EStyle Style, const FOnEventCont
 		uint64 Cycle = EventData.GetValue<uint64>("Cycle");
 		double RecordingTime = EventData.GetValue<double>("RecordingTime");
 		uint64 Id = EventData.GetValue<uint64>("Id");
-		GameplayProvider.AppendObjectLifetimeEnd(Id, Context.EventTime.AsSeconds(Cycle), RecordingTime);
+		GameplayProvider.AppendObjectLifetimeEnd(RewindDebugger::FObjectId{Id}, Context.EventTime.AsSeconds(Cycle), RecordingTime);
 		break;
 	}
 	case RouteId_PawnPossess:
@@ -216,18 +220,56 @@ bool FGameplayAnalyzer::OnEvent(uint16 RouteId, EStyle Style, const FOnEventCont
 		break;
 	}
 	case RouteId_PropertyValue3:
-		{
-			uint64 Cycle = EventData.GetValue<uint64>("Cycle");
-			uint64 ObjectId = EventData.GetValue<uint64>("ObjectId");
-			int32 ParentId = EventData.GetValue<int32>("ParentId");
-			uint32 TypeId = EventData.GetValue<uint32>("TypeId");
-			uint32 NameId = EventData.GetValue<uint32>("NameId");
-			uint32 ParentNameId = EventData.GetValue<uint32>("ParentNameId");
-			
-			FStringView Value; EventData.GetString("Value", Value);
-			GameplayProvider.AppendPropertyValue(ObjectId, Context.EventTime.AsSeconds(Cycle), Cycle, ParentId, TypeId, NameId, ParentNameId, Value);
-			break;
-		}
+	{
+		uint64 Cycle = EventData.GetValue<uint64>("Cycle");
+		uint64 ObjectId = EventData.GetValue<uint64>("ObjectId");
+		int32 ParentId = EventData.GetValue<int32>("ParentId");
+		uint32 TypeId = EventData.GetValue<uint32>("TypeId");
+		uint32 NameId = EventData.GetValue<uint32>("NameId");
+		uint32 ParentNameId = EventData.GetValue<uint32>("ParentNameId");
+
+		FStringView Value; EventData.GetString("Value", Value);
+		GameplayProvider.AppendPropertyValue(ObjectId, Context.EventTime.AsSeconds(Cycle), Cycle, ParentId, TypeId, NameId, ParentNameId, Value);
+		break;
+	}
+	case RouteId_ChildElement:
+	{
+		uint64 Id = EventData.GetValue<uint64>("Id");
+		uint64 ChildElementId = EventData.GetValue<uint64>("ChildElementId");
+		uint64 OuterId = EventData.GetValue<uint64>("OuterId");
+		uint64 OuterChildElementId = EventData.GetValue<uint64>("OuterChildElementId");
+		FString ElementName;
+		EventData.GetString("Name", ElementName);
+		GameplayProvider.AppendObject(RewindDebugger::FObjectId(Id, ChildElementId), RewindDebugger::FObjectId(OuterId, OuterChildElementId), /*ClassId*/0, *ElementName, /*PathName*/{});
+		break;
+	}
+	case RouteId_ChildElementLifetimeBegin:
+	{
+		uint64 Cycle = EventData.GetValue<uint64>("Cycle");
+		double RecordingTime = EventData.GetValue<double>("RecordingTime");
+		uint64 Id = EventData.GetValue<uint64>("Id");
+		uint64 ChildElementId = EventData.GetValue<uint64>("ChildElementId");
+		GameplayProvider.AppendObjectLifetimeBegin(RewindDebugger::FObjectId(Id, ChildElementId), Context.EventTime.AsSeconds(Cycle), RecordingTime);
+		break;
+	}
+	case RouteId_ChildElementLifetimeEnd:
+	{
+		uint64 Cycle = EventData.GetValue<uint64>("Cycle");
+		double RecordingTime = EventData.GetValue<double>("RecordingTime");
+		uint64 Id = EventData.GetValue<uint64>("Id");
+		uint64 ChildElementId = EventData.GetValue<uint64>("ChildElementId");
+		GameplayProvider.AppendObjectLifetimeEnd(RewindDebugger::FObjectId(Id, ChildElementId), Context.EventTime.AsSeconds(Cycle), RecordingTime);
+		break;
+	}
+	case RouteId_ChildElementEvent:
+	{
+		uint64 Cycle = EventData.GetValue<uint64>("Cycle");
+		uint64 Id = EventData.GetValue<uint64>("Id");
+		uint64 ChildElementId = EventData.GetValue<uint64>("ChildElementId");
+		FString Event = TraceServices::FTraceAnalyzerUtils::LegacyAttachmentString<TCHAR>("Event", Context);
+		GameplayProvider.AppendObjectEvent(RewindDebugger::FObjectId(Id, ChildElementId), Context.EventTime.AsSeconds(Cycle), *Event);
+		break;
+	}
 	}
 
 	return true;
diff --git a/Engine/Plugins/Animation/GameplayInsights/Source/GameplayInsights/Private/GameplayAnalyzer.h b/Engine/Plugins/Animation/GameplayInsights/Source/GameplayInsights/Private/GameplayAnalyzer.h
--- a/Engine/Plugins/Animation/GameplayInsights/Source/GameplayInsights/Private/GameplayAnalyzer.h
+++ b/Engine/Plugins/Animation/GameplayInsights/Source/GameplayInsights/Private/GameplayAnalyzer.h
@@ -38,6 +38,10 @@ private:
 		RouteId_PropertyValue2,
 		RouteId_PropertyValue3,
 		RouteId_RecordingInfo,
+		RouteId_ChildElement,
+		RouteId_ChildElementLifetimeBegin,
+		RouteId_ChildElementLifetimeEnd,
+		RouteId_ChildElementEvent,
 	};
 
 	TraceServices::IAnalysisSession& Session;
diff --git a/Engine/Plugins/Animation/GameplayInsights/Source/GameplayInsights/Private/GameplayProvider.cpp b/Engine/Plugins/Animation/GameplayInsights/Source/GameplayInsights/Private/GameplayProvider.cpp
--- a/Engine/Plugins/Animation/GameplayInsights/Source/GameplayInsights/Private/GameplayProvider.cpp
+++ b/Engine/Plugins/Animation/GameplayInsights/Source/GameplayInsights/Private/GameplayProvider.cpp
@@ -117,7 +117,7 @@ FGameplayProvider::FGameplayProvider(TraceServices::IAnalysisSession& InSession)
 {
 }
 
-bool FGameplayProvider::ReadObjectEventsTimeline(uint64 InObjectId, TFunctionRef<void(const ObjectEventsTimeline&)> Callback) const
+bool FGameplayProvider::ReadObjectEventsTimeline(const RewindDebugger::FObjectId& InObjectId, TFunctionRef<void(const ObjectEventsTimeline&)> Callback) const
 {
 	Session.ReadAccessCheck();
 
@@ -134,7 +134,7 @@ bool FGameplayProvider::ReadObjectEventsTimeline(uint64 InObjectId, TFunctionRef
 	return false;
 }
 
-bool FGameplayProvider::ReadObjectEvent(uint64 InObjectId, uint64 InMessageId, TFunctionRef<void(const FObjectEventMessage&)> Callback) const
+bool FGameplayProvider::ReadObjectEvent(const RewindDebugger::FObjectId& InObjectId, uint64 InMessageId, TFunctionRef<void(const FObjectEventMessage&)> Callback) const
 {
 	Session.ReadAccessCheck();
 
@@ -257,12 +257,12 @@ const FObjectPropertyValue* FGameplayProvider::FindPropertyValueFromStorageIndex
 	return nullptr;
 }
 
-void FGameplayProvider::EnumerateSubobjects(uint64 ObjectId, TFunctionRef<void(uint64 SubobjectId)> Callback) const
+void FGameplayProvider::EnumerateSubobjects(const RewindDebugger::FObjectId& InObjectId, TFunctionRef<void(const RewindDebugger::FObjectId& SubObjectId)> Callback) const
 {
-	TArray<uint64, TInlineAllocator<32>> SubobjectIds;
-	ObjectHierarchy.MultiFind(ObjectId, SubobjectIds);
+	TArray<RewindDebugger::FObjectId, TInlineAllocator<32>> SubObjectIds;
+	ObjectHierarchy.MultiFind(InObjectId, SubObjectIds);
 
-	for (auto SubObjectId : SubobjectIds)
+	for (const RewindDebugger::FObjectId& SubObjectId : SubObjectIds)
 	{
 		Callback(SubObjectId);
 	}
@@ -335,6 +335,11 @@ bool FGameplayProvider::IsSubClassOf(uint64 InSubClassId, uint64 InParentClassId
 }
 
 const FObjectInfo* FGameplayProvider::FindObjectInfo(uint64 InObjectId) const
+{
+	return FindObjectInfo(RewindDebugger::FObjectId{InObjectId});
+}
+
+const FObjectInfo* FGameplayProvider::FindObjectInfo(const RewindDebugger::FObjectId& InObjectId) const
 {
 	Session.ReadAccessCheck();
 
@@ -370,10 +375,10 @@ const FWorldInfo* FGameplayProvider::FindWorldInfoFromObject(const uint64 InObje
 		{
 			if (ObjectInfo->ClassId == WorldClass->Id)
 			{
-				return FindWorldInfo(ObjectInfo->Id);
+				return FindWorldInfo(ObjectInfo->GetUObjectId());
 			}
 
-			ObjectInfo = FindObjectInfo(ObjectInfo->OuterId);
+			ObjectInfo = FindObjectInfo(ObjectInfo->GetOuterId());
 		}
 	}
 
@@ -419,6 +424,11 @@ const FClassInfo& FGameplayProvider::GetClassInfoFromObject(const uint64 InObjec
 }
 
 const FObjectInfo& FGameplayProvider::GetObjectInfo(uint64 InObjectId) const
+{
+	return GetObjectInfo(RewindDebugger::FObjectId{InObjectId});
+}
+
+const FObjectInfo& FGameplayProvider::GetObjectInfo(const RewindDebugger::FObjectId& InObjectId) const
 {
 	if (const FObjectInfo* ObjectInfo = FindObjectInfo(InObjectId))
 	{
@@ -426,7 +436,13 @@ const FObjectInfo& FGameplayProvider::GetObjectInfo(uint64 InObjectId) const
 	}
 
 	static FText UnknownText(LOCTEXT("Unknown", "Unknown"));
-	static FObjectInfo DefaultObjectInfo = { 0, 0, 0, *UnknownText.ToString(), *UnknownText.ToString() };
+	static FObjectInfo DefaultObjectInfo = {
+		RewindDebugger::FObjectId()
+		, RewindDebugger::FObjectId()
+		, /*ClassId*/0
+		, *UnknownText.ToString()
+		, *UnknownText.ToString()
+	};
 	return DefaultObjectInfo;
 }
 
@@ -464,7 +480,7 @@ void FGameplayProvider::AppendClass(uint64 InClassId, uint64 InSuperId, const TC
 	}
 }
 
-void FGameplayProvider::AppendObject(uint64 InObjectId, uint64 InOuterId, uint64 InClassId, const TCHAR* InObjectName, const TCHAR* InObjectPathName)
+void FGameplayProvider::AppendObject(const RewindDebugger::FObjectId& InObjectId, const RewindDebugger::FObjectId& InOuterId, uint64 InClassId, const TCHAR* InObjectName, const TCHAR* InObjectPathName)
 {
 	Session.WriteAccessCheck();
 
@@ -475,17 +491,9 @@ void FGameplayProvider::AppendObject(uint64 InObjectId, uint64 InOuterId, uint64
 		const TCHAR* NewObjectName = Session.StoreString(InObjectName);
 		const TCHAR* NewObjectPathName = Session.StoreString(InObjectPathName);
 
-		FObjectInfo NewObjectInfo;
-		NewObjectInfo.Id = InObjectId;
-		NewObjectInfo.OuterId = InOuterId;
-		NewObjectInfo.ClassId = InClassId;
-		NewObjectInfo.Name = NewObjectName;
-		NewObjectInfo.PathName = NewObjectPathName;
+		ObjectIdToIndexMap.Add(InObjectId, ObjectInfos.Add({ InObjectId, InOuterId, InClassId, NewObjectName, NewObjectPathName }));
 
-		int32 NewObjectInfoIndex = ObjectInfos.Add(NewObjectInfo);
-		ObjectIdToIndexMap.Add(InObjectId, NewObjectInfoIndex);
-
-		ObjectHierarchy.AddUnique(NewObjectInfo.OuterId, NewObjectInfo.Id);
+		ObjectHierarchy.AddUnique(InOuterId, InObjectId);
 	}
 }
 
@@ -503,12 +511,12 @@ void FGameplayProvider::AppendObjectTransform(uint64 InObjectId, double InProfil
 	TransformTimeline.AppendEvent(InProfileTime, Message);
 }
 
-void FGameplayProvider::AppendObjectLifetimeBegin(uint64 InObjectId, double InProfileTime, double InRecordingTime)
+void FGameplayProvider::AppendObjectLifetimeBegin(const RewindDebugger::FObjectId& InObjectId, double InProfileTime, double InRecordingTime)
 {
 	Session.WriteAccessCheck();
 	bHasAnyData = true;
 
-	if (InObjectId)
+	if (InObjectId.IsSet())
 	{
 		FObjectExistsMessage Message;
 		Message.ObjectId = InObjectId;
@@ -517,7 +525,7 @@ void FGameplayProvider::AppendObjectLifetimeBegin(uint64 InObjectId, double InPr
 	}
 }
 
-void FGameplayProvider::AppendObjectLifetimeEnd(uint64 InObjectId, double InProfileTime, double InRecordingTime)
+void FGameplayProvider::AppendObjectLifetimeEnd(const RewindDebugger::FObjectId& InObjectId, double InProfileTime, double InRecordingTime)
 {
 	Session.WriteAccessCheck();
 	bHasAnyData = true;
@@ -537,13 +545,17 @@ void FGameplayProvider::AppendObjectLifetimeEnd(uint64 InObjectId, double InProf
 		// do not remove from ActiveObjectRecordingLifetimes - lifetimes can be queried by Object Id in GetObjectRecordingLifetime
 	}
 
-	if (int32* ObjectInfoIndex = ObjectIdToIndexMap.Find(InObjectId))
+	// For now, we don't have a need to notify EndPlay on child element so let's avoid it
+	if (!InObjectId.IsChildElement())
 	{
-		OnObjectEndPlayDelegate.Broadcast(InObjectId, InProfileTime, ObjectInfos[*ObjectInfoIndex]);
+		if (const int32* ObjectInfoIndex = ObjectIdToIndexMap.Find(InObjectId))
+		{
+			OnObjectEndPlayDelegate.Broadcast(InObjectId.GetMainId(), InProfileTime, ObjectInfos[*ObjectInfoIndex]);
+		}
 	}
 }
 
-void FGameplayProvider::AppendObjectEvent(uint64 InObjectId, double InTime, const TCHAR* InEventName)
+void FGameplayProvider::AppendObjectEvent(const RewindDebugger::FObjectId& InObjectId, double InTime, const TCHAR* InEventName)
 {
 	Session.WriteAccessCheck();
 
@@ -569,14 +581,16 @@ void FGameplayProvider::AppendObjectEvent(uint64 InObjectId, double InTime, cons
 	}
 
 	FObjectEventMessage Message;
-	Message.Id = InObjectId;
+	Message.ObjectId = InObjectId;
 	Message.Name = Session.StoreString(InEventName);
 
-	if (Message.Name == EndPlayEvent)
+	// For now, we don't have a need to notify EndPlay on child element so let's avoid it
+	if (Message.Name == EndPlayEvent
+		&& !InObjectId.IsChildElement())
 	{
-		if (int32* ObjectInfoIndex = ObjectIdToIndexMap.Find(InObjectId))
+		if (const int32* ObjectInfoIndex = ObjectIdToIndexMap.Find(InObjectId))
 		{
-			OnObjectEndPlayDelegate.Broadcast(InObjectId, InTime, ObjectInfos[*ObjectInfoIndex]);
+			OnObjectEndPlayDelegate.Broadcast(InObjectId.GetMainId(), InTime, ObjectInfos[*ObjectInfoIndex]);
 		}
 	}
 
@@ -637,11 +651,11 @@ bool FGameplayProvider::GetObjectTransform(uint64 InObjectId, double InStartTime
 	return bSuccess;
 }
 
-TRange<double> FGameplayProvider::GetObjectTraceLifetime(uint64 ObjectId) const
+TRange<double> FGameplayProvider::GetObjectTraceLifetime(const RewindDebugger::FObjectId& InObjectId) const
 {
 	Session.ReadAccessCheck();
 
-	if (const uint64* FoundIndex = ActiveObjectLifetimes.Find(ObjectId))
+	if (const uint64* FoundIndex = ActiveObjectLifetimes.Find(InObjectId))
 	{
 		return TRange<double>(ObjectLifetimes.GetEventStartTime(*FoundIndex), ObjectLifetimes.GetEventEndTime(*FoundIndex));
 	}
@@ -651,12 +665,12 @@ TRange<double> FGameplayProvider::GetObjectTraceLifetime(uint64 ObjectId) const
 	}
 }
 
-TRange<double> FGameplayProvider::GetObjectRecordingLifetime(uint64 ObjectId) const
+TRange<double> FGameplayProvider::GetObjectRecordingLifetime(const RewindDebugger::FObjectId& InObjectId) const
 {
 	TRACE_CPUPROFILER_EVENT_SCOPE(FGameplayProvider::GetObjectRecordingLifetime);
 	Session.ReadAccessCheck();
 
-	if (const uint64* FoundIndex = ActiveObjectRecordingLifetimes.Find(ObjectId))
+	if (const uint64 *FoundIndex = ActiveObjectRecordingLifetimes.Find(InObjectId))
 	{
 		return TRange<double>(ObjectRecordingLifetimes.GetEventStartTime(*FoundIndex), ObjectRecordingLifetimes.GetEventEndTime(*FoundIndex));
 	}
diff --git a/Engine/Plugins/Animation/GameplayInsights/Source/GameplayInsights/Private/GameplayProvider.h b/Engine/Plugins/Animation/GameplayInsights/Source/GameplayInsights/Private/GameplayProvider.h
--- a/Engine/Plugins/Animation/GameplayInsights/Source/GameplayInsights/Private/GameplayProvider.h
+++ b/Engine/Plugins/Animation/GameplayInsights/Source/GameplayInsights/Private/GameplayProvider.h
@@ -10,7 +10,15 @@
 #include "Model/PointTimeline.h"
 #include "Templates/SharedPointer.h"
 
-namespace TraceServices { class IAnalysisSession; }
+namespace RewindDebugger
+{
+struct FObjectId;
+}
+
+namespace TraceServices
+{
+class IAnalysisSession;
+}
 
 class FGameplayProvider : public IGameplayProvider
 {
@@ -20,28 +28,33 @@ public:
 	FGameplayProvider(TraceServices::IAnalysisSession& InSession);
 
 	/** IGameplayProvider interface */
-	virtual bool ReadObjectEventsTimeline(uint64 InObjectId, TFunctionRef<void(const ObjectEventsTimeline&)> Callback) const override;
-	virtual bool ReadObjectEvent(uint64 InObjectId, uint64 InMessageId, TFunctionRef<void(const FObjectEventMessage&)> Callback) const override;
+	virtual bool ReadObjectEventsTimeline(const RewindDebugger::FObjectId& InObjectId, TFunctionRef<void(const ObjectEventsTimeline&)> Callback) const override;
+	virtual bool ReadObjectEvent(const RewindDebugger::FObjectId& InObjectId, uint64 InMessageId, TFunctionRef<void(const FObjectEventMessage&)> Callback) const override;
 	virtual bool ReadObjectPropertiesTimeline(uint64 InObjectId, TFunctionRef<void(const ObjectPropertiesTimeline&)> Callback) const override;
 	virtual bool ReadObjectPropertiesStorage(uint64 InObjectId, const FObjectPropertiesMessage& InMessage, TFunctionRef<void(const TConstArrayView<FObjectPropertyValue>&)> Callback) const override;
 	virtual void EnumerateObjectPropertyValues(uint64 InObjectId, const FObjectPropertiesMessage& InMessage, TFunctionRef<void(const FObjectPropertyValue&)> Callback) const override;
 	virtual void EnumerateObjects(TFunctionRef<void(const FObjectInfo&)> Callback) const override;
 	virtual void EnumerateObjects(double StartTime, double EndTime, TFunctionRef<void(const FObjectInfo&)> Callback) const override;
 	virtual void EnumerateWorlds(TFunctionRef<void(const FWorldInfo&)> Callback) const override;
-	virtual void EnumerateSubobjects(uint64 ObjectId, TFunctionRef<void(uint64 SubobjectId)> Callback) const override;
+	virtual void EnumerateSubobjects(const RewindDebugger::FObjectId& InObjectId, TFunctionRef<void(const RewindDebugger::FObjectId& SubObjectId)> Callback) const override;
 	virtual const FObjectPropertyValue* FindPropertyValueFromStorageIndex(uint64 InObjectId, int64 InStorageIndex) const override;
 	virtual const FClassInfo* FindClassInfo(uint64 InClassId) const override;
 	virtual const UClass* FindClass(uint64 InClassId, bool bSearchBlueprints = true) const override;
 	virtual const FClassInfo* FindClassInfo(const TCHAR* InClassPath) const override;
 	virtual bool IsSubClassOf(uint64 InSubClassId, uint64 InParentClassId) const override;
-	virtual const FObjectInfo* FindObjectInfo(uint64 InObjectId) const override;
+	virtual const FObjectInfo* FindObjectInfo(const uint64 InObjectId) const override;
+	virtual const FObjectInfo* FindObjectInfo(const RewindDebugger::FObjectId& InObjectId) const override;
 	virtual const FWorldInfo* FindWorldInfo(uint64 InObjectId) const override;
 	virtual const FWorldInfo* FindWorldInfoFromObject(uint64 InObjectId) const override;
 	virtual bool IsWorld(uint64 InObjectId) const override;
 	virtual const FClassInfo& GetClassInfo(uint64 InClassId) const override;
 	virtual const FClassInfo& GetClassInfoFromObject(uint64 InObjectId) const override;
 	virtual const FObjectInfo& GetObjectInfo(uint64 InObjectId) const override;
-	virtual FOnObjectEndPlay& OnObjectEndPlay() override { return OnObjectEndPlayDelegate; }
+	virtual const FObjectInfo& GetObjectInfo(const RewindDebugger::FObjectId& InObjectId) const override;
+	virtual FOnObjectEndPlay& OnObjectEndPlay() override
+	{
+		return OnObjectEndPlayDelegate;
+	}
 	virtual const TCHAR* GetPropertyName(uint32 InPropertyStringId) const override;
 	virtual const RecordingInfoTimeline* GetRecordingInfo(uint32 RecordingId) const override; 
 	virtual void ReadViewTimeline(TFunctionRef<void(const ViewTimeline&)> Callback) const override;
@@ -52,19 +65,19 @@ public:
 	void AppendClass(uint64 InClassId, uint64 InSuperId, const TCHAR* InClassName, const TCHAR* InClassPathName);
 
 	/** Add an object message */
-	void AppendObject(uint64 InObjectId, uint64 InOuterId, uint64 InClassId, const TCHAR* InObjectName, const TCHAR* InObjectPathName);
+	void AppendObject(const RewindDebugger::FObjectId& InObjectId, const RewindDebugger::FObjectId& InOuterId, uint64 InClassId, const TCHAR* InObjectName, const TCHAR* InObjectPathName);
 
 	/** Add an object transform */
 	void AppendObjectTransform(uint64 InObjectId, double InProfileTime, double InRecordingTime, const FTransform& Transform);
 
 	/** Add an object create message */
-	void AppendObjectLifetimeBegin(uint64 InObjectId, double InProfileTime, double InRecordingTime);
+	void AppendObjectLifetimeBegin(const RewindDebugger::FObjectId& InObjectId, double InProfileTime, double InRecordingTime);
 
 	/** Add an object destroy message */
-	void AppendObjectLifetimeEnd(uint64 InObjectId, double InProfileTime, double InRecordingTime);	
+	void AppendObjectLifetimeEnd(const RewindDebugger::FObjectId& InObjectId, double InProfileTime, double InRecordingTime);	
 
 	/** Add an object event message */
-	void AppendObjectEvent(uint64 InObjectId, double InTime, const TCHAR* InEvent);
+	void AppendObjectEvent(const RewindDebugger::FObjectId& InObjectId, double InTime, const TCHAR* InEvent);
 
 	/** Add a Controller Attach message */
 	void AppendPawnPossess(uint64 InControllerId, uint64 InPawnId, double InTime);
@@ -103,10 +116,10 @@ public:
 	virtual bool GetObjectTransform(uint64 InObjectId, double InStartTime, double InEndTime, FTransform& OutTransform) const override;
 
 	/** find the Trace time range for which an object existed */
-	virtual TRange<double> GetObjectTraceLifetime(uint64 ObjectId) const override;
+	virtual TRange<double> GetObjectTraceLifetime(const RewindDebugger::FObjectId& InObjectId) const override;
 
 	/** find the Recording time range for which an object existed */
-	virtual TRange<double> GetObjectRecordingLifetime(uint64 ObjectId) const override;
+	virtual TRange<double> GetObjectRecordingLifetime(const RewindDebugger::FObjectId& InObjectId) const override;
 
 private:
 	TraceServices::IAnalysisSession& Session;
@@ -124,13 +137,13 @@ private:
 	TMap<uint64, int32> ClassIdToIndexMap;
 
 	/** Objects that are in use. Map from Id to ObjectInfo index */
-	TMap<uint64, int32> ObjectIdToIndexMap;
+	TMap<RewindDebugger::FObjectId, int32> ObjectIdToIndexMap;
 
 	/** Worlds that are in use. Map from Id to WorldInfo index */
 	TMap<uint64, int32> WorldIdToIndexMap;
 
 	/** Map from object Ids to timeline index */
-	TMap<uint64, uint32> ObjectIdToEventTimelines;
+	TMap<RewindDebugger::FObjectId, uint32> ObjectIdToEventTimelines;
 	TMap<uint64, uint32> ObjectIdToTransformTimeline;
 	TMap<uint64, uint32> ObjectIdToPropertiesStorage;
 
@@ -149,7 +162,7 @@ private:
 
 	struct FObjectExistsMessage
 	{
-		uint64 ObjectId = 0;
+		RewindDebugger::FObjectId ObjectId;
 	};
 
 	struct FObjectPropertiesStorage
@@ -195,10 +208,10 @@ private:
 	TMap<uint64, uint64> ActivePawnPossession;
 
 	// map from object id to index in ObjectLifetimes, for lookup when objects are destroyed
-	TMap<uint64, uint64> ActiveObjectLifetimes;
+	TMap<RewindDebugger::FObjectId, uint64> ActiveObjectLifetimes;
 
 	// map from object id to index in ObjectRecordingLifetimes, for lookup
-	TMap<uint64, uint64> ActiveObjectRecordingLifetimes;
+	TMap<RewindDebugger::FObjectId, uint64> ActiveObjectRecordingLifetimes;
 
 	/** Whether we have any data */
 	bool bHasAnyData = false;
@@ -210,5 +223,5 @@ private:
 	bool bHasObjectProperties = false;
 
 	/** MultiMap from parent object to sub objects, to avoid slow reverse lookups */
-	TMultiMap<uint64, uint64> ObjectHierarchy;
+	TMultiMap<RewindDebugger::FObjectId, RewindDebugger::FObjectId> ObjectHierarchy;
 };
diff --git a/Engine/Plugins/Animation/GameplayInsights/Source/GameplayInsights/Private/GameplaySharedData.cpp b/Engine/Plugins/Animation/GameplayInsights/Source/GameplayInsights/Private/GameplaySharedData.cpp
--- a/Engine/Plugins/Animation/GameplayInsights/Source/GameplayInsights/Private/GameplaySharedData.cpp
+++ b/Engine/Plugins/Animation/GameplayInsights/Source/GameplayInsights/Private/GameplaySharedData.cpp
@@ -24,6 +24,7 @@
 
 FGameplaySharedData::FGameplaySharedData()
 	: AnalysisSession(nullptr)
+	, TimingViewSession(nullptr)
 	, bObjectTracksDirty(false)
 	, bObjectTracksEnabled(false)
 	, bObjectPropertyTracksEnabled(true)
@@ -49,22 +50,23 @@ TSharedRef<FObjectEventsTrack> FGameplaySharedData::GetObjectEventsTrackForId(UE
 	const FGameplayProvider* GameplayProvider = InAnalysisSession.ReadProvider<FGameplayProvider>(FGameplayProvider::ProviderName);
 	check(GameplayProvider);
 
-	TSharedPtr<FObjectEventsTrack> LeafObjectEventsTrack = ObjectTracks.FindRef(InObjectInfo.Id);
+	const uint64 ObjectId = InObjectInfo.GetUObjectId();
+	TSharedPtr<FObjectEventsTrack> LeafObjectEventsTrack = ObjectTracks.FindRef(ObjectId);
 	if(!LeafObjectEventsTrack.IsValid())
 	{
-		LeafObjectEventsTrack = MakeShared<FObjectEventsTrack>(*this, InObjectInfo.Id, InObjectInfo.Name);
+		LeafObjectEventsTrack = MakeShared<FObjectEventsTrack>(*this, ObjectId, InObjectInfo.Name);
 		LeafObjectEventsTrack->SetVisibilityFlag(bObjectTracksEnabled);
-		ObjectTracks.Add(InObjectInfo.Id, LeafObjectEventsTrack);
+		ObjectTracks.Add(ObjectId, LeafObjectEventsTrack);
 
 		InTimingViewSession.AddScrollableTrack(LeafObjectEventsTrack);
 		InvalidateObjectTracksOrder();
 	}
 
 	// Fill the outer chain
-	if(InObjectInfo.OuterId != 0)
+	if(InObjectInfo.GetOuterId().IsSet())
 	{
 		TSharedPtr<FObjectEventsTrack> ObjectEventsTrack = LeafObjectEventsTrack;
-		uint64 OuterId = InObjectInfo.OuterId;
+		uint64 OuterId = InObjectInfo.GetOuterUObjectId();
 		while(const FObjectInfo* OuterInfo = GameplayProvider->FindObjectInfo(OuterId))
 		{
 			TSharedPtr<FObjectEventsTrack> OuterEventsTrack = ObjectTracks.FindRef(OuterId);
@@ -85,7 +87,7 @@ TSharedRef<FObjectEventsTrack> FGameplaySharedData::GetObjectEventsTrackForId(UE
 			}
 
 			ObjectEventsTrack = OuterEventsTrack;
-			OuterId = OuterInfo->OuterId;
+			OuterId = OuterInfo->GetOuterUObjectId();
 		}
 	}
 
@@ -161,7 +163,7 @@ void FGameplaySharedData::Tick(UE::Insights::Timing::ITimingViewSession& InTimin
 				{
 					TSharedPtr<FObjectEventsTrack> ObjectEventsTrack;
 
-					GameplayProvider->ReadObjectEventsTimeline(InObjectInfo.Id, [this, &InTimingViewSession, &InAnalysisSession, &InObjectInfo, &GameplayProvider, &ObjectEventsTrack](const IGameplayProvider::ObjectEventsTimeline& InTimeline)
+					GameplayProvider->ReadObjectEventsTimeline(InObjectInfo.GetId(), [this, &InTimingViewSession, &InAnalysisSession, &InObjectInfo, &GameplayProvider, &ObjectEventsTrack](const IGameplayProvider::ObjectEventsTimeline& InTimeline)
 					{
 						if(bObjectTracksEnabled)
 						{
@@ -169,7 +171,7 @@ void FGameplaySharedData::Tick(UE::Insights::Timing::ITimingViewSession& InTimin
 						}
 					});
 
-					GameplayProvider->ReadObjectPropertiesTimeline(InObjectInfo.Id, [this, &InObjectInfo, &ObjectEventsTrack, &InTimingViewSession, &InAnalysisSession, &GameplayProvider](const IGameplayProvider::ObjectPropertiesTimeline& InTimeline)
+					GameplayProvider->ReadObjectPropertiesTimeline(InObjectInfo.GetUObjectId(), [this, &InObjectInfo, &ObjectEventsTrack, &InTimingViewSession, &InAnalysisSession, &GameplayProvider](const IGameplayProvider::ObjectPropertiesTimeline& InTimeline)
 					{
 						if(!ObjectEventsTrack.IsValid())
 						{
@@ -181,10 +183,10 @@ void FGameplaySharedData::Tick(UE::Insights::Timing::ITimingViewSession& InTimin
 							return InTrack.Is<FObjectPropertiesTrack>();
 						};
 
-						TSharedPtr<FObjectPropertiesTrack> ExistingObjectPropertiesTrack = StaticCastSharedPtr<FObjectPropertiesTrack>(ObjectEventsTrack->GetGameplayTrack().FindChildTrack(InObjectInfo.Id, FindObjectProperties));
+						TSharedPtr<FObjectPropertiesTrack> ExistingObjectPropertiesTrack = StaticCastSharedPtr<FObjectPropertiesTrack>(ObjectEventsTrack->GetGameplayTrack().FindChildTrack(InObjectInfo.GetUObjectId(), FindObjectProperties));
 						if(!ExistingObjectPropertiesTrack.IsValid())
 						{
-							TSharedPtr<FObjectPropertiesTrack> ObjectPropertiesTrack = MakeShared<FObjectPropertiesTrack>(*this, InObjectInfo.Id, InObjectInfo.Name);
+							TSharedPtr<FObjectPropertiesTrack> ObjectPropertiesTrack = MakeShared<FObjectPropertiesTrack>(*this, InObjectInfo.GetUObjectId(), InObjectInfo.Name);
 							ObjectPropertiesTrack->SetVisibilityFlag(bObjectPropertyTracksEnabled);
 							ObjectPropertyTracks.Add(ObjectPropertiesTrack.ToSharedRef());
 
diff --git a/Engine/Plugins/Animation/GameplayInsights/Source/GameplayInsights/Private/GameplayTrack.cpp b/Engine/Plugins/Animation/GameplayInsights/Source/GameplayInsights/Private/GameplayTrack.cpp
--- a/Engine/Plugins/Animation/GameplayInsights/Source/GameplayInsights/Private/GameplayTrack.cpp
+++ b/Engine/Plugins/Animation/GameplayInsights/Source/GameplayInsights/Private/GameplayTrack.cpp
@@ -26,6 +26,11 @@ INSIGHTS_IMPLEMENT_RTTI(FGameplayTimingEventsTrack);
 
 #define LOCTEXT_NAMESPACE "GameplayTrack"
 
+RewindDebugger::FObjectId FGameplayTrack::GetAssociatedObjectId() const
+{
+	return RewindDebugger::FObjectId{ ObjectId };
+}
+
 void FGameplayTrack::AddChildTrack(FGameplayTrack& InChildTrack)
 {
 	check(InChildTrack.Parent == nullptr);
diff --git a/Engine/Plugins/Animation/GameplayInsights/Source/GameplayInsights/Private/GameplayTrack.h b/Engine/Plugins/Animation/GameplayInsights/Source/GameplayInsights/Private/GameplayTrack.h
--- a/Engine/Plugins/Animation/GameplayInsights/Source/GameplayInsights/Private/GameplayTrack.h
+++ b/Engine/Plugins/Animation/GameplayInsights/Source/GameplayInsights/Private/GameplayTrack.h
@@ -12,8 +12,16 @@ namespace GameplayTrackConstants
 struct FVariantTreeNode;
 class FGameplaySharedData;
 struct FWorldInfo;
-namespace TraceServices { class IAnalysisSession; }
-namespace TraceServices { struct FFrame; }
+namespace TraceServices
+{
+	class IAnalysisSession;
+	struct FFrame;
+}
+
+namespace RewindDebugger
+{
+	struct FObjectId;
+}
 
 // Provides parent/child hierarchy structure and the owning object Id
 // Designed as a compositional member of outer timing tracks
@@ -31,6 +39,8 @@ public:
 	// Get the object ID for this track
 	uint64 GetObjectId() const { return ObjectId; }
 
+	RewindDebugger::FObjectId GetAssociatedObjectId() const;
+
 	/** Get the parent track */
 	FGameplayTrack* GetParentTrack() { return Parent; }
 
diff --git a/Engine/Plugins/Animation/GameplayInsights/Source/GameplayInsights/Private/InertializationsTrack.cpp b/Engine/Plugins/Animation/GameplayInsights/Source/GameplayInsights/Private/InertializationsTrack.cpp
--- a/Engine/Plugins/Animation/GameplayInsights/Source/GameplayInsights/Private/InertializationsTrack.cpp
+++ b/Engine/Plugins/Animation/GameplayInsights/Source/GameplayInsights/Private/InertializationsTrack.cpp
@@ -266,12 +266,12 @@ void FInertializationsTrackCreator::GetTrackTypesInternal(TArray<FRewindDebugger
 	Types.Add({InertializationsName, LOCTEXT("Inertializations", "Inertializations")});
 }
 
-TSharedPtr<RewindDebugger::FRewindDebuggerTrack> FInertializationsTrackCreator::CreateTrackInternal(uint64 ObjectId) const
+TSharedPtr<RewindDebugger::FRewindDebuggerTrack> FInertializationsTrackCreator::CreateTrackInternal(const FObjectId& InObjectId) const
 {
- 	return MakeShared<RewindDebugger::FInertializationsTrack>(ObjectId);
+ 	return MakeShared<RewindDebugger::FInertializationsTrack>(InObjectId.GetMainId());
 }
 
-bool FInertializationsTrackCreator::HasDebugInfoInternal(uint64 ObjectId) const
+bool FInertializationsTrackCreator::HasDebugInfoInternal(const FObjectId& InObjectId) const
 {
 	TRACE_CPUPROFILER_EVENT_SCOPE(FInertializationsTrack::HasDebugInfoInternal);
 	const TraceServices::IAnalysisSession* AnalysisSession = IRewindDebugger::Instance()->GetAnalysisSession();
@@ -279,7 +279,7 @@ bool FInertializationsTrackCreator::HasDebugInfoInternal(uint64 ObjectId) const
 	
 	if (const FAnimationProvider* AnimationProvider = AnalysisSession->ReadProvider<FAnimationProvider>(FAnimationProvider::ProviderName))
 	{
-		if (AnimationProvider->ReadInertializationTimeline(ObjectId, [](const IAnimationProvider::InertializationTimeline& InTimeline){}))
+		if (AnimationProvider->ReadInertializationTimeline(InObjectId.GetMainId(), [](const IAnimationProvider::InertializationTimeline& InTimeline){}))
 		{
 			return true;
 		}
diff --git a/Engine/Plugins/Animation/GameplayInsights/Source/GameplayInsights/Private/InertializationsTrack.h b/Engine/Plugins/Animation/GameplayInsights/Source/GameplayInsights/Private/InertializationsTrack.h
--- a/Engine/Plugins/Animation/GameplayInsights/Source/GameplayInsights/Private/InertializationsTrack.h
+++ b/Engine/Plugins/Animation/GameplayInsights/Source/GameplayInsights/Private/InertializationsTrack.h
@@ -68,8 +68,8 @@ public:
 	virtual FName GetTargetTypeNameInternal() const;
 	virtual FName GetNameInternal() const override;
 	virtual void GetTrackTypesInternal(TArray<FRewindDebuggerTrackType>& Types) const override;
-	virtual TSharedPtr<FRewindDebuggerTrack> CreateTrackInternal(uint64 ObjectId) const override;
-	virtual bool HasDebugInfoInternal(uint64 ObjectId) const override;
+	virtual TSharedPtr<FRewindDebuggerTrack> CreateTrackInternal(const FObjectId& InObjectId) const override;
+	virtual bool HasDebugInfoInternal(const FObjectId& InObjectId) const override;
 };
 
 }
\ No newline at end of file
diff --git a/Engine/Plugins/Animation/GameplayInsights/Source/GameplayInsights/Private/MontagesTrack.cpp b/Engine/Plugins/Animation/GameplayInsights/Source/GameplayInsights/Private/MontagesTrack.cpp
--- a/Engine/Plugins/Animation/GameplayInsights/Source/GameplayInsights/Private/MontagesTrack.cpp
+++ b/Engine/Plugins/Animation/GameplayInsights/Source/GameplayInsights/Private/MontagesTrack.cpp
@@ -405,12 +405,12 @@ FName FMontagesTrack::GetNameInternal() const
 	return MontagesName;
 }
 
-TSharedPtr<RewindDebugger::FRewindDebuggerTrack> FMontagesTrackCreator::CreateTrackInternal(uint64 ObjectId) const
+TSharedPtr<RewindDebugger::FRewindDebuggerTrack> FMontagesTrackCreator::CreateTrackInternal(const FObjectId& InObjectId) const
 {
- 	return MakeShared<RewindDebugger::FMontagesTrack>(ObjectId);
+ 	return MakeShared<RewindDebugger::FMontagesTrack>(InObjectId.GetMainId());
 }
 
-bool FMontagesTrackCreator::HasDebugInfoInternal(uint64 ObjectId) const
+bool FMontagesTrackCreator::HasDebugInfoInternal(const FObjectId& InObjectId) const
 {
 	TRACE_CPUPROFILER_EVENT_SCOPE(FMontagesTrack::HasDebugInfoInternal);
 	const TraceServices::IAnalysisSession* AnalysisSession = IRewindDebugger::Instance()->GetAnalysisSession();
@@ -418,7 +418,7 @@ bool FMontagesTrackCreator::HasDebugInfoInternal(uint64 ObjectId) const
 	bool bHasData = false;
 	if (const FAnimationProvider* AnimationProvider = AnalysisSession->ReadProvider<FAnimationProvider>(FAnimationProvider::ProviderName))
 	{
-		AnimationProvider->ReadMontageTimeline(ObjectId, [&bHasData](const FAnimationProvider::AnimMontageTimeline& InTimeline)
+		AnimationProvider->ReadMontageTimeline(InObjectId.GetMainId(), [&bHasData](const FAnimationProvider::AnimMontageTimeline& InTimeline)
 		{
 			bHasData = true;
 		});
diff --git a/Engine/Plugins/Animation/GameplayInsights/Source/GameplayInsights/Private/MontagesTrack.h b/Engine/Plugins/Animation/GameplayInsights/Source/GameplayInsights/Private/MontagesTrack.h
--- a/Engine/Plugins/Animation/GameplayInsights/Source/GameplayInsights/Private/MontagesTrack.h
+++ b/Engine/Plugins/Animation/GameplayInsights/Source/GameplayInsights/Private/MontagesTrack.h
@@ -80,8 +80,8 @@ public:
 	virtual FName GetTargetTypeNameInternal() const;
 	virtual FName GetNameInternal() const override;
 	virtual void GetTrackTypesInternal(TArray<FRewindDebuggerTrackType>& Types) const override;
-	virtual TSharedPtr<FRewindDebuggerTrack> CreateTrackInternal(uint64 ObjectId) const override;
-	virtual bool HasDebugInfoInternal(uint64 ObjectId) const override;
+	virtual TSharedPtr<FRewindDebuggerTrack> CreateTrackInternal(const FObjectId& InObjectId) const override;
+	virtual bool HasDebugInfoInternal(const FObjectId& InObjectId) const override;
 };
 
 }
\ No newline at end of file
diff --git a/Engine/Plugins/Animation/GameplayInsights/Source/GameplayInsights/Private/NotifiesTrack.cpp b/Engine/Plugins/Animation/GameplayInsights/Source/GameplayInsights/Private/NotifiesTrack.cpp
--- a/Engine/Plugins/Animation/GameplayInsights/Source/GameplayInsights/Private/NotifiesTrack.cpp
+++ b/Engine/Plugins/Animation/GameplayInsights/Source/GameplayInsights/Private/NotifiesTrack.cpp
@@ -169,9 +169,9 @@ void FNotifiesTrackCreator::GetTrackTypesInternal(TArray<FRewindDebuggerTrackTyp
 }
 
 
-TSharedPtr<RewindDebugger::FRewindDebuggerTrack> FNotifiesTrackCreator::CreateTrackInternal(uint64 ObjectId) const
+TSharedPtr<RewindDebugger::FRewindDebuggerTrack> FNotifiesTrackCreator::CreateTrackInternal(const FObjectId& InObjectId) const
 {
-	return MakeShared<RewindDebugger::FNotifiesTrack>(ObjectId);
+	return MakeShared<RewindDebugger::FNotifiesTrack>(InObjectId.GetMainId());
 }
 
 		
@@ -264,7 +264,7 @@ void FNotifiesTrack::IterateSubTracksInternal(TFunction<void(TSharedPtr<FRewindD
 	}
 };
 
-bool FNotifiesTrackCreator::HasDebugInfoInternal(uint64 ObjectId) const
+bool FNotifiesTrackCreator::HasDebugInfoInternal(const FObjectId& InObjectId) const
 {
 	TRACE_CPUPROFILER_EVENT_SCOPE(FNotifiesTrack::HasDebugInfoInternal);
 	const TraceServices::IAnalysisSession* AnalysisSession = IRewindDebugger::Instance()->GetAnalysisSession();
@@ -273,7 +273,7 @@ bool FNotifiesTrackCreator::HasDebugInfoInternal(uint64 ObjectId) const
 	bool bHasData = false;
 	if (const FAnimationProvider* AnimationProvider = AnalysisSession->ReadProvider<FAnimationProvider>(FAnimationProvider::ProviderName))
 	{
-		AnimationProvider->ReadNotifyTimeline(ObjectId, [&bHasData](const FAnimationProvider::AnimNotifyTimeline& InTimeline)
+		AnimationProvider->ReadNotifyTimeline(InObjectId.GetMainId(), [&bHasData](const FAnimationProvider::AnimNotifyTimeline& InTimeline)
 		{
 			bHasData = true;
 		});
diff --git a/Engine/Plugins/Animation/GameplayInsights/Source/GameplayInsights/Private/NotifiesTrack.h b/Engine/Plugins/Animation/GameplayInsights/Source/GameplayInsights/Private/NotifiesTrack.h
--- a/Engine/Plugins/Animation/GameplayInsights/Source/GameplayInsights/Private/NotifiesTrack.h
+++ b/Engine/Plugins/Animation/GameplayInsights/Source/GameplayInsights/Private/NotifiesTrack.h
@@ -102,8 +102,8 @@ private:
 	virtual FName GetTargetTypeNameInternal() const;
 	virtual FName GetNameInternal() const override;
 	virtual void GetTrackTypesInternal(TArray<FRewindDebuggerTrackType>& Types) const override;
-	virtual TSharedPtr<FRewindDebuggerTrack> CreateTrackInternal(uint64 ObjectId) const override;
-	virtual bool HasDebugInfoInternal(uint64 ObjectId) const override;
+	virtual TSharedPtr<FRewindDebuggerTrack> CreateTrackInternal(const FObjectId& InObjectId) const override;
+	virtual bool HasDebugInfoInternal(const FObjectId& InObjectId) const override;
 };
 
 }
\ No newline at end of file
diff --git a/Engine/Plugins/Animation/GameplayInsights/Source/GameplayInsights/Private/ObjectEventsTrack.cpp b/Engine/Plugins/Animation/GameplayInsights/Source/GameplayInsights/Private/ObjectEventsTrack.cpp
--- a/Engine/Plugins/Animation/GameplayInsights/Source/GameplayInsights/Private/ObjectEventsTrack.cpp
+++ b/Engine/Plugins/Animation/GameplayInsights/Source/GameplayInsights/Private/ObjectEventsTrack.cpp
@@ -4,7 +4,6 @@
 #include "GameplayProvider.h"
 #include "Insights/ViewModels/TimingTrackViewport.h"
 #include "Insights/ViewModels/TimingEvent.h"
-#include "Algo/Sort.h"
 #include "Insights/ViewModels/TooltipDrawState.h"
 #include "GameplaySharedData.h"
 #include "Insights/ViewModels/TimingEventSearch.h"
@@ -37,7 +36,7 @@ void FObjectEventsTrack::BuildDrawState(ITimingEventsTrackDrawStateBuilder& Buil
 		TraceServices::FAnalysisSessionReadScope SessionReadScope(SharedData.GetAnalysisSession());
 
 		// object events
-		GameplayProvider->ReadObjectEventsTimeline(GetGameplayTrack().GetObjectId(), [&Context, &Builder](const FGameplayProvider::ObjectEventsTimeline& InTimeline)
+		GameplayProvider->ReadObjectEventsTimeline(GetGameplayTrack().GetAssociatedObjectId(), [&Context, &Builder](const FGameplayProvider::ObjectEventsTimeline& InTimeline)
 		{
 			InTimeline.EnumerateEvents(Context.GetViewport().GetStartTime(), Context.GetViewport().GetEndTime(), [&Builder](double InStartTime, double InEndTime, uint32 InDepth, const FObjectEventMessage& InMessage)
 			{
@@ -97,7 +96,7 @@ void FObjectEventsTrack::FindObjectEvent(const FTimingEventSearchParameters& InP
 			{
 				TraceServices::FAnalysisSessionReadScope SessionReadScope(SharedData.GetAnalysisSession());
 
-				GameplayProvider->ReadObjectEventsTimeline(GetGameplayTrack().GetObjectId(), [&InContext](const FGameplayProvider::ObjectEventsTimeline& InTimeline)
+				GameplayProvider->ReadObjectEventsTimeline(GetGameplayTrack().GetAssociatedObjectId(), [&InContext](const FGameplayProvider::ObjectEventsTimeline& InTimeline)
 				{
 					InTimeline.EnumerateEvents(InContext.GetParameters().StartTime, InContext.GetParameters().EndTime, [&InContext](double InEventStartTime, double InEventEndTime, uint32 InDepth, const FObjectEventMessage& InMessage)
 					{
@@ -159,7 +158,7 @@ void FObjectEventsTrack::GetVariantsAtFrame(const TraceServices::FFrame& InFrame
 		TSharedRef<FVariantTreeNode> Header = OutVariants.Add_GetRef(FVariantTreeNode::MakeHeader(FText::FromString(GetName()), 0));
 
 		// object events
-		GameplayProvider->ReadObjectEventsTimeline(GetGameplayTrack().GetObjectId(), [&InFrame, &Header](const FGameplayProvider::ObjectEventsTimeline& InTimeline)
+		GameplayProvider->ReadObjectEventsTimeline(GetGameplayTrack().GetAssociatedObjectId(), [&InFrame, &Header](const FGameplayProvider::ObjectEventsTimeline& InTimeline)
 		{
 			InTimeline.EnumerateEvents(InFrame.StartTime, InFrame.EndTime, [&Header](double InStartTime, double InEndTime, uint32 InDepth, const FObjectEventMessage& InMessage)
 			{
diff --git a/Engine/Plugins/Animation/GameplayInsights/Source/GameplayInsights/Private/PoseWatchTrack.cpp b/Engine/Plugins/Animation/GameplayInsights/Source/GameplayInsights/Private/PoseWatchTrack.cpp
--- a/Engine/Plugins/Animation/GameplayInsights/Source/GameplayInsights/Private/PoseWatchTrack.cpp
+++ b/Engine/Plugins/Animation/GameplayInsights/Source/GameplayInsights/Private/PoseWatchTrack.cpp
@@ -258,9 +258,9 @@ void FPoseWatchesTrackCreator::GetTrackTypesInternal(TArray<FRewindDebuggerTrack
 	Types.Add({PoseWatchesName, LOCTEXT("Pose Watches", "Pose Watches")});
 }
 
-TSharedPtr<RewindDebugger::FRewindDebuggerTrack> FPoseWatchesTrackCreator::CreateTrackInternal(uint64 ObjectId) const
+TSharedPtr<RewindDebugger::FRewindDebuggerTrack> FPoseWatchesTrackCreator::CreateTrackInternal(const FObjectId& InObjectId) const
 {
-	return MakeShared<RewindDebugger::FPoseWatchesTrack>(ObjectId);
+	return MakeShared<RewindDebugger::FPoseWatchesTrack>(InObjectId.GetMainId());
 }
 
 		
@@ -349,7 +349,7 @@ void FPoseWatchesTrack::IterateSubTracksInternal(TFunction<void(TSharedPtr<FRewi
 	}
 };
 
-bool FPoseWatchesTrackCreator::HasDebugInfoInternal(uint64 ObjectId) const
+bool FPoseWatchesTrackCreator::HasDebugInfoInternal(const FObjectId& InObjectId) const
 {
 	TRACE_CPUPROFILER_EVENT_SCOPE(FPoseWatchesTrack::HasDebugInfoInternal);
 	const TraceServices::IAnalysisSession* AnalysisSession = IRewindDebugger::Instance()->GetAnalysisSession();
@@ -358,7 +358,7 @@ bool FPoseWatchesTrackCreator::HasDebugInfoInternal(uint64 ObjectId) const
 	bool bHasData = false;
 	if (const FAnimationProvider* AnimationProvider = AnalysisSession->ReadProvider<FAnimationProvider>(FAnimationProvider::ProviderName))
 	{
-		AnimationProvider->ReadPoseWatchTimeline(ObjectId, [&bHasData](const FAnimationProvider::PoseWatchTimeline& InTimeline)
+		AnimationProvider->ReadPoseWatchTimeline(InObjectId.GetMainId(), [&bHasData](const FAnimationProvider::PoseWatchTimeline& InTimeline)
 		{
 			bHasData = true;
 		});
diff --git a/Engine/Plugins/Animation/GameplayInsights/Source/GameplayInsights/Private/PoseWatchTrack.h b/Engine/Plugins/Animation/GameplayInsights/Source/GameplayInsights/Private/PoseWatchTrack.h
--- a/Engine/Plugins/Animation/GameplayInsights/Source/GameplayInsights/Private/PoseWatchTrack.h
+++ b/Engine/Plugins/Animation/GameplayInsights/Source/GameplayInsights/Private/PoseWatchTrack.h
@@ -90,8 +90,8 @@ private:
 	virtual FName GetTargetTypeNameInternal() const;
 	virtual FName GetNameInternal() const override;
 	virtual void GetTrackTypesInternal(TArray<FRewindDebuggerTrackType>& Types) const override;
-	virtual TSharedPtr<FRewindDebuggerTrack> CreateTrackInternal(uint64 ObjectId) const override;
-	virtual bool HasDebugInfoInternal(uint64 ObjectId) const override;
+	virtual TSharedPtr<FRewindDebuggerTrack> CreateTrackInternal(const FObjectId& InObjectId) const override;
+	virtual bool HasDebugInfoInternal(const FObjectId& InObjectId) const override;
 };
 
 }
diff --git a/Engine/Plugins/Animation/GameplayInsights/Source/GameplayInsights/Private/PropertiesTrack.cpp b/Engine/Plugins/Animation/GameplayInsights/Source/GameplayInsights/Private/PropertiesTrack.cpp
--- a/Engine/Plugins/Animation/GameplayInsights/Source/GameplayInsights/Private/PropertiesTrack.cpp
+++ b/Engine/Plugins/Animation/GameplayInsights/Source/GameplayInsights/Private/PropertiesTrack.cpp
@@ -302,16 +302,16 @@ namespace RewindDebugger
     }
 
 
-	TSharedPtr<FRewindDebuggerTrack> FPropertiesTrackCreator::CreateTrackInternal(uint64 ObjectId) const
+	TSharedPtr<FRewindDebuggerTrack> FPropertiesTrackCreator::CreateTrackInternal(const FObjectId& InObjectId) const
 	{
-		TSharedRef<FPropertiesTrack> NewTrack = MakeShared<FPropertiesTrack>(ObjectId);
+		TSharedRef<FPropertiesTrack> NewTrack = MakeShared<FPropertiesTrack>(InObjectId.GetMainId());
 
 		NewTrack->Initialize();
 		
 		return NewTrack;
 	}
 
-	bool FPropertiesTrackCreator::HasDebugInfoInternal(uint64 ObjectId) const
+	bool FPropertiesTrackCreator::HasDebugInfoInternal(const FObjectId& InObjectId) const
 	{
 		TRACE_CPUPROFILER_EVENT_SCOPE(FPropertiesTrack::HasDebugInfoInternal);
 		bool bHasData = false;
@@ -322,7 +322,7 @@ namespace RewindDebugger
 		
 			if (const FGameplayProvider* GameplayProvider = AnalysisSession->ReadProvider<FGameplayProvider>(FGameplayProvider::ProviderName))
 			{
-				GameplayProvider->ReadObjectPropertiesTimeline(ObjectId, [&bHasData](const IGameplayProvider::ObjectPropertiesTimeline & ObjPropsTimeline)
+				GameplayProvider->ReadObjectPropertiesTimeline(InObjectId.GetMainId(), [&bHasData](const IGameplayProvider::ObjectPropertiesTimeline & ObjPropsTimeline)
 				{
 					bHasData = true;
 				});
diff --git a/Engine/Plugins/Animation/GameplayInsights/Source/GameplayInsights/Private/PropertiesTrack.h b/Engine/Plugins/Animation/GameplayInsights/Source/GameplayInsights/Private/PropertiesTrack.h
--- a/Engine/Plugins/Animation/GameplayInsights/Source/GameplayInsights/Private/PropertiesTrack.h
+++ b/Engine/Plugins/Animation/GameplayInsights/Source/GameplayInsights/Private/PropertiesTrack.h
@@ -76,8 +76,8 @@ namespace RewindDebugger
 		virtual FName GetTargetTypeNameInternal() const override;
 		virtual FName GetNameInternal() const override;
 		virtual void GetTrackTypesInternal(TArray<FRewindDebuggerTrackType>& Types) const override;
-		virtual TSharedPtr<FRewindDebuggerTrack> CreateTrackInternal(uint64 ObjectId) const override;
-		virtual bool HasDebugInfoInternal(uint64 ObjectId) const override;
+		virtual TSharedPtr<FRewindDebuggerTrack> CreateTrackInternal(const FObjectId& InObjectId) const override;
+		virtual bool HasDebugInfoInternal(const FObjectId& InObjectId) const override;
 		/** End IRewindDebuggerTrackCreator interface */
 	};
 }
\ No newline at end of file
diff --git a/Engine/Plugins/Animation/GameplayInsights/Source/GameplayInsights/Private/SAnimGraphSchematicView.cpp b/Engine/Plugins/Animation/GameplayInsights/Source/GameplayInsights/Private/SAnimGraphSchematicView.cpp
--- a/Engine/Plugins/Animation/GameplayInsights/Source/GameplayInsights/Private/SAnimGraphSchematicView.cpp
+++ b/Engine/Plugins/Animation/GameplayInsights/Source/GameplayInsights/Private/SAnimGraphSchematicView.cpp
@@ -962,9 +962,9 @@ void FAnimGraphSchematicTrackCreator::GetTrackTypesInternal(TArray<RewindDebugge
 	Types.Add({AnimGraphSchematicName, LOCTEXT("Anim Graph", "Anim Graph")});
 }
 
-TSharedPtr<RewindDebugger::FRewindDebuggerTrack> FAnimGraphSchematicTrackCreator::CreateTrackInternal(uint64 ObjectId) const
+TSharedPtr<RewindDebugger::FRewindDebuggerTrack> FAnimGraphSchematicTrackCreator::CreateTrackInternal(const RewindDebugger::FObjectId& InObjectId) const
 {
-	return MakeShared<FAnimGraphSchematicTrack>(ObjectId);
+	return MakeShared<FAnimGraphSchematicTrack>(InObjectId.GetMainId());
 }
 
 FText FAnimGraphSchematicTrack::GetDisplayNameInternal() const
@@ -1017,7 +1017,7 @@ bool FAnimGraphSchematicTrack::UpdateInternal()
 	return false;
 }
 
-bool FAnimGraphSchematicTrackCreator::HasDebugInfoInternal(uint64 ObjectId) const
+bool FAnimGraphSchematicTrackCreator::HasDebugInfoInternal(const RewindDebugger::FObjectId& InObjectId) const
 {
 	const TraceServices::IAnalysisSession* AnalysisSession = IRewindDebugger::Instance()->GetAnalysisSession();
 	
@@ -1025,7 +1025,7 @@ bool FAnimGraphSchematicTrackCreator::HasDebugInfoInternal(uint64 ObjectId) cons
 	bool bHasData = false;
 	if (const FAnimationProvider* AnimationProvider = AnalysisSession->ReadProvider<FAnimationProvider>(FAnimationProvider::ProviderName))
 	{
-		AnimationProvider->ReadAnimGraphTimeline(ObjectId, [&bHasData](const FAnimationProvider::AnimGraphTimeline& InGraphTimeline)
+		AnimationProvider->ReadAnimGraphTimeline(InObjectId.GetMainId(), [&bHasData](const FAnimationProvider::AnimGraphTimeline& InGraphTimeline)
 		{
 			bHasData = true;
 		});
diff --git a/Engine/Plugins/Animation/GameplayInsights/Source/GameplayInsights/Private/SAnimGraphSchematicView.h b/Engine/Plugins/Animation/GameplayInsights/Source/GameplayInsights/Private/SAnimGraphSchematicView.h
--- a/Engine/Plugins/Animation/GameplayInsights/Source/GameplayInsights/Private/SAnimGraphSchematicView.h
+++ b/Engine/Plugins/Animation/GameplayInsights/Source/GameplayInsights/Private/SAnimGraphSchematicView.h
@@ -145,6 +145,6 @@ public:
 	virtual FName GetTargetTypeNameInternal() const override;
 	virtual FName GetNameInternal() const override;
 	virtual void GetTrackTypesInternal(TArray<RewindDebugger::FRewindDebuggerTrackType>& Types) const override;
-	virtual TSharedPtr<RewindDebugger::FRewindDebuggerTrack> CreateTrackInternal(uint64 ObjectId) const override;
-	virtual bool HasDebugInfoInternal(uint64 ObjectId) const override;
+	virtual TSharedPtr<RewindDebugger::FRewindDebuggerTrack> CreateTrackInternal(const RewindDebugger::FObjectId& InObjectId) const override;
+	virtual bool HasDebugInfoInternal(const RewindDebugger::FObjectId& InObjectId) const override;
 };
\ No newline at end of file
diff --git a/Engine/Plugins/Animation/GameplayInsights/Source/GameplayInsights/Private/SPropertyTrackView.cpp b/Engine/Plugins/Animation/GameplayInsights/Source/GameplayInsights/Private/SPropertyTrackView.cpp
--- a/Engine/Plugins/Animation/GameplayInsights/Source/GameplayInsights/Private/SPropertyTrackView.cpp
+++ b/Engine/Plugins/Animation/GameplayInsights/Source/GameplayInsights/Private/SPropertyTrackView.cpp
@@ -18,7 +18,7 @@ void SPropertyTrackView::GetVariantsAtFrame(const TraceServices::FFrame& InFrame
 		{
 			PropertyTrack->ReadObjectPropertyValueAtFrame(InFrame, *GameplayProvider, [GameplayProvider, &PropertyTrack, &Header](const FObjectPropertyValue& InValue, uint32 InValueIndex, const FObjectPropertiesMessage& InMessage)
 			{
-				GameplayProvider->ReadObjectPropertiesStorage(PropertyTrack->GetObjectId(), InMessage, [GameplayProvider, InValueIndex, &PropertyTrack, &Header](const TConstArrayView<FObjectPropertyValue> & InStorage)
+				GameplayProvider->ReadObjectPropertiesStorage(PropertyTrack->GetUObjectId(), InMessage, [GameplayProvider, InValueIndex, &PropertyTrack, &Header](const TConstArrayView<FObjectPropertyValue> & InStorage)
 				{
 					Header->AddChild(FObjectPropertyHelpers::GetVariantNodeFromProperty(InValueIndex, *GameplayProvider, InStorage));
 				});
diff --git a/Engine/Plugins/Animation/GameplayInsights/Source/GameplayInsights/Private/SkeletalMeshPoseTrack.cpp b/Engine/Plugins/Animation/GameplayInsights/Source/GameplayInsights/Private/SkeletalMeshPoseTrack.cpp
--- a/Engine/Plugins/Animation/GameplayInsights/Source/GameplayInsights/Private/SkeletalMeshPoseTrack.cpp
+++ b/Engine/Plugins/Animation/GameplayInsights/Source/GameplayInsights/Private/SkeletalMeshPoseTrack.cpp
@@ -217,11 +217,11 @@ void FSkeletalMeshPoseTrack::BuildContextMenu(FMenuBuilder& MenuBuilder)
 	{
 		TraceServices::FAnalysisSessionReadScope SessionReadScope(SharedData.GetAnalysisSession());
 	
-		const FObjectInfo* ComponentObjectInfo = GameplayProvider->FindObjectInfo(GetGameplayTrack().GetObjectId());
+		const FObjectInfo* ComponentObjectInfo = GameplayProvider->FindObjectInfo(GetGameplayTrack().GetAssociatedObjectId());
 		if(ComponentObjectInfo != nullptr)
 		{
 			// @FIXME: Outer does always equal owning actor, although does in nearly all cases with skeletal mesh components
-			const FObjectInfo* ActorObjectInfo = GameplayProvider->FindObjectInfo(ComponentObjectInfo->OuterId);
+			const FObjectInfo* ActorObjectInfo = GameplayProvider->FindObjectInfo(ComponentObjectInfo->GetOuterId());
 			if(ActorObjectInfo != nullptr)
 			{
 				MenuBuilder.BeginSection(TEXT("DrawingSection"), FText::Format(LOCTEXT("DrawingActor", "Drawing ({0})"), FText::FromString(ActorObjectInfo->Name)));
@@ -238,8 +238,8 @@ void FSkeletalMeshPoseTrack::BuildContextMenu(FMenuBuilder& MenuBuilder)
 								bool bSetDrawPose = true;
 								SharedData.EnumerateSkeletalMeshPoseTracks([GameplayProvider, ActorObjectInfo, &bSetDrawPose](const TSharedRef<FSkeletalMeshPoseTrack>& InTrack)
 								{
-									const FObjectInfo* OtherComponentObjectInfo = GameplayProvider->FindObjectInfo(InTrack->GetGameplayTrack().GetObjectId());
-									if(OtherComponentObjectInfo->OuterId == ActorObjectInfo->Id)
+									const FObjectInfo* OtherComponentObjectInfo = GameplayProvider->FindObjectInfo(InTrack->GetGameplayTrack().GetAssociatedObjectId());
+									if(OtherComponentObjectInfo->GetOuterId() == ActorObjectInfo->GetId())
 									{
 										bSetDrawPose &= InTrack->bDrawPose;
 									}
@@ -247,8 +247,8 @@ void FSkeletalMeshPoseTrack::BuildContextMenu(FMenuBuilder& MenuBuilder)
 
 								SharedData.EnumerateSkeletalMeshPoseTracks([GameplayProvider, ActorObjectInfo, &bSetDrawPose](const TSharedRef<FSkeletalMeshPoseTrack>& InTrack)
 								{
-									const FObjectInfo* OtherComponentObjectInfo = GameplayProvider->FindObjectInfo(InTrack->GetGameplayTrack().GetObjectId());
-									if(OtherComponentObjectInfo->OuterId == ActorObjectInfo->Id)
+									const FObjectInfo* OtherComponentObjectInfo = GameplayProvider->FindObjectInfo(InTrack->GetGameplayTrack().GetAssociatedObjectId());
+									if(OtherComponentObjectInfo->GetOuterId() == ActorObjectInfo->GetId())
 									{
 										InTrack->bDrawPose = !bSetDrawPose;
 										InTrack->UpdateComponentVisibility();
@@ -265,8 +265,8 @@ void FSkeletalMeshPoseTrack::BuildContextMenu(FMenuBuilder& MenuBuilder)
 								bool bDrawPoseSet = true;
 								SharedData.EnumerateSkeletalMeshPoseTracks([GameplayProvider, ActorObjectInfo, &bDrawPoseSet](const TSharedRef<FSkeletalMeshPoseTrack>& InTrack)
 								{
-									const FObjectInfo* OtherComponentObjectInfo = GameplayProvider->FindObjectInfo(InTrack->GetGameplayTrack().GetObjectId());
-									if(OtherComponentObjectInfo->OuterId == ActorObjectInfo->Id)
+									const FObjectInfo* OtherComponentObjectInfo = GameplayProvider->FindObjectInfo(InTrack->GetGameplayTrack().GetAssociatedObjectId());
+									if(OtherComponentObjectInfo->GetOuterId() == ActorObjectInfo->GetId())
 									{
 										bDrawPoseSet &= InTrack->bDrawPose;
 									}
@@ -290,8 +290,8 @@ void FSkeletalMeshPoseTrack::BuildContextMenu(FMenuBuilder& MenuBuilder)
 								bool bSetDrawSkeleton = true;
 								SharedData.EnumerateSkeletalMeshPoseTracks([GameplayProvider, ActorObjectInfo, &bSetDrawSkeleton](const TSharedRef<FSkeletalMeshPoseTrack>& InTrack)
 								{
-									const FObjectInfo* OtherComponentObjectInfo = GameplayProvider->FindObjectInfo(InTrack->GetGameplayTrack().GetObjectId());
-									if(OtherComponentObjectInfo->OuterId == ActorObjectInfo->Id)
+									const FObjectInfo* OtherComponentObjectInfo = GameplayProvider->FindObjectInfo(InTrack->GetGameplayTrack().GetAssociatedObjectId());
+									if(OtherComponentObjectInfo->GetOuterId() == ActorObjectInfo->GetId())
 									{
 										bSetDrawSkeleton &= InTrack->bDrawSkeleton;
 									}
@@ -299,8 +299,8 @@ void FSkeletalMeshPoseTrack::BuildContextMenu(FMenuBuilder& MenuBuilder)
 
 								SharedData.EnumerateSkeletalMeshPoseTracks([GameplayProvider, ActorObjectInfo, &bSetDrawSkeleton](const TSharedRef<FSkeletalMeshPoseTrack>& InTrack)
 								{
-									const FObjectInfo* OtherComponentObjectInfo = GameplayProvider->FindObjectInfo(InTrack->GetGameplayTrack().GetObjectId());
-									if(OtherComponentObjectInfo->OuterId == ActorObjectInfo->Id)
+									const FObjectInfo* OtherComponentObjectInfo = GameplayProvider->FindObjectInfo(InTrack->GetGameplayTrack().GetAssociatedObjectId());
+									if(OtherComponentObjectInfo->GetOuterId() == ActorObjectInfo->GetId())
 									{
 										InTrack->bDrawSkeleton = !bSetDrawSkeleton;
 										InTrack->UpdateComponentVisibility();
@@ -317,8 +317,8 @@ void FSkeletalMeshPoseTrack::BuildContextMenu(FMenuBuilder& MenuBuilder)
 								bool bDrawSkeletonSet = true;
 								SharedData.EnumerateSkeletalMeshPoseTracks([GameplayProvider, ActorObjectInfo, &bDrawSkeletonSet](const TSharedRef<FSkeletalMeshPoseTrack>& InTrack)
 								{
-									const FObjectInfo* OtherComponentObjectInfo = GameplayProvider->FindObjectInfo(InTrack->GetGameplayTrack().GetObjectId());
-									if(OtherComponentObjectInfo->OuterId == ActorObjectInfo->Id)
+									const FObjectInfo* OtherComponentObjectInfo = GameplayProvider->FindObjectInfo(InTrack->GetGameplayTrack().GetAssociatedObjectId());
+									if(OtherComponentObjectInfo->GetOuterId() == ActorObjectInfo->GetId())
 									{
 										bDrawSkeletonSet &= InTrack->bDrawSkeleton;
 									}
diff --git a/Engine/Plugins/Animation/GameplayInsights/Source/GameplayInsights/Public/IGameplayProvider.h b/Engine/Plugins/Animation/GameplayInsights/Source/GameplayInsights/Public/IGameplayProvider.h
--- a/Engine/Plugins/Animation/GameplayInsights/Source/GameplayInsights/Public/IGameplayProvider.h
+++ b/Engine/Plugins/Animation/GameplayInsights/Source/GameplayInsights/Public/IGameplayProvider.h
@@ -7,6 +7,8 @@
 #endif // UE_ENABLE_INCLUDE_ORDER_DEPRECATED_IN_5_7
 #include "Math/Vector.h"
 #include "Math/Rotator.h"
+#include "Math/Range.h" // for deprecation only
+#include "RewindDebuggerTypes.h"
 #include "TraceServices/Model/AnalysisSession.h"
 #include "TraceServices/Containers/Timelines.h"
 
@@ -20,8 +22,59 @@ struct FClassInfo
 
 struct FObjectInfo
 {
+PRAGMA_DISABLE_DEPRECATION_WARNINGS
+	FObjectInfo(const FObjectInfo& Other) = default;
+	FObjectInfo(FObjectInfo&& Other) = default;
+	FObjectInfo& operator=(const FObjectInfo& Other) = default;
+	FObjectInfo& operator=(FObjectInfo&& Other) = default;
+PRAGMA_ENABLE_DEPRECATION_WARNINGS
+
+	FObjectInfo(const RewindDebugger::FObjectId& ObjectId, const RewindDebugger::FObjectId& OuterObjectId, const uint64 ClassId, const TCHAR* Name
+		, const TCHAR* PathName)
+		: ObjectId(ObjectId)
+		, OuterObjectId(OuterObjectId)
+		, ClassId(ClassId)
+		, Name(Name)
+		, PathName(PathName)
+	{
+	}
+
+	/** @return Id of an object recorded in the trace */
+	const RewindDebugger::FObjectId& GetId() const
+	{
+		return ObjectId;
+	}
+
+	/** @return Id of the recorded object's parent */
+	const RewindDebugger::FObjectId& GetOuterId() const
+	{
+		return OuterObjectId;
+	}
+
+	/** @return Part of the Id representing the UObject */
+	uint64 GetUObjectId() const
+	{
+		return ObjectId.GetMainId();
+	}
+
+	/** @return Part of the Id of the parent object representing the UObject */
+	uint64 GetOuterUObjectId() const
+	{
+		return OuterObjectId.GetMainId();
+	}
+
+private:
+	RewindDebugger::FObjectId ObjectId;
+	RewindDebugger::FObjectId OuterObjectId;
+
+public:
+#if WITH_EDITORONLY_DATA
+	UE_DEPRECATED(5.7, "Use ObjectId instead")
 	uint64 Id = 0;
+	UE_DEPRECATED(5.7, "Use OuterObjectId instead")
 	uint64 OuterId = 0;
+#endif
+
 	uint64 ClassId = 0;
 	const TCHAR* Name = nullptr;
 	const TCHAR* PathName = nullptr;
@@ -47,8 +100,21 @@ struct FObjectPropertyValue
 
 struct FObjectEventMessage
 {
-	uint64 Id = 0;
+	RewindDebugger::FObjectId ObjectId;
 	const TCHAR* Name = nullptr;
+
+PRAGMA_DISABLE_DEPRECATION_WARNINGS
+	FObjectEventMessage() = default;
+	FObjectEventMessage(const FObjectEventMessage& Other) = default;
+	FObjectEventMessage(FObjectEventMessage&& Other) = default;
+	FObjectEventMessage& operator=(const FObjectEventMessage& Other) = default;
+	FObjectEventMessage& operator=(FObjectEventMessage&& Other) = default;
+PRAGMA_ENABLE_DEPRECATION_WARNINGS
+
+#if WITH_EDITORONLY_DATA
+	UE_DEPRECATED(5.7, "Use Identifier instead")
+	uint64 Id = 0;
+#endif
 };
 
 struct FRecordingInfoMessage
@@ -115,35 +181,76 @@ public:
 	typedef TraceServices::ITimeline<FRecordingInfoMessage> RecordingInfoTimeline;
 	typedef TraceServices::ITimeline<FViewMessage> ViewTimeline;
 
-	virtual bool ReadObjectEventsTimeline(uint64 InObjectId, TFunctionRef<void(const ObjectEventsTimeline&)> Callback) const = 0;
-	virtual bool ReadObjectEvent(uint64 InObjectId, uint64 InMessageId, TFunctionRef<void(const FObjectEventMessage&)> Callback) const = 0;
+	virtual bool ReadObjectEventsTimeline(const RewindDebugger::FObjectId& InObjectId, TFunctionRef<void(const ObjectEventsTimeline&)> Callback) const = 0;
+	virtual bool ReadObjectEvent(const RewindDebugger::FObjectId& InObjectId, uint64 InMessageId, TFunctionRef<void(const FObjectEventMessage&)> Callback) const = 0;
 	virtual bool ReadObjectPropertiesTimeline(uint64 InObjectId, TFunctionRef<void(const ObjectPropertiesTimeline&)> Callback) const = 0;
 	virtual bool ReadObjectPropertiesStorage(uint64 InObjectId, const FObjectPropertiesMessage& InMessage, TFunctionRef<void(const TConstArrayView<FObjectPropertyValue> &)> Callback) const = 0;
 	virtual void EnumerateObjectPropertyValues(uint64 InObjectId, const FObjectPropertiesMessage& InMessage, TFunctionRef<void(const FObjectPropertyValue&)> Callback) const = 0;
 	virtual void EnumerateObjects(TFunctionRef<void(const FObjectInfo&)> Callback) const = 0;
 	virtual void EnumerateObjects(double StartTime, double EndTime, TFunctionRef<void(const FObjectInfo&)> Callback) const = 0;
 	virtual void EnumerateWorlds(TFunctionRef<void(const FWorldInfo&)> Callback) const = 0;
-	virtual void EnumerateSubobjects(uint64 ObjectId, TFunctionRef<void(uint64 SubobjectId)> Callback) const = 0;
+	virtual void EnumerateSubobjects(const RewindDebugger::FObjectId& InObjectId, TFunctionRef<void(const RewindDebugger::FObjectId& SubObjectId)> Callback) const = 0;
 	virtual const FObjectPropertyValue * FindPropertyValueFromStorageIndex(uint64 InObjectId, int64 InStorageIndex) const = 0;
 	virtual const FClassInfo* FindClassInfo(uint64 InClassId) const = 0;
 	virtual const UClass* FindClass(uint64 InClassId, bool bSearchBlueprints) const = 0;
 	virtual const FClassInfo* FindClassInfo(const TCHAR* InClassPath) const = 0;
 	virtual bool IsSubClassOf(uint64 InSubClassId, uint64 InParentClassId) const = 0;
-	virtual const FObjectInfo* FindObjectInfo(uint64 InObjectId) const = 0;
+	virtual const FObjectInfo* FindObjectInfo(const uint64 InObjectId) const = 0;
+	virtual const FObjectInfo* FindObjectInfo(const RewindDebugger::FObjectId& InObjectId) const = 0;
 	virtual const FWorldInfo* FindWorldInfo(uint64 InObjectId) const = 0;
 	virtual const FWorldInfo* FindWorldInfoFromObject(uint64 InObjectId) const = 0;
 	virtual bool IsWorld(uint64 InObjectId) const = 0;
 	virtual const FClassInfo& GetClassInfo(uint64 InClassId) const = 0;
 	virtual const FClassInfo& GetClassInfoFromObject(uint64 InObjectId) const = 0;
-	virtual const FObjectInfo& GetObjectInfo(uint64 InObjectId) const = 0;
+	virtual const FObjectInfo& GetObjectInfo(const uint64 InObjectId) const = 0;
+	virtual const FObjectInfo& GetObjectInfo(const RewindDebugger::FObjectId& InObjectId) const = 0;
 	virtual FOnObjectEndPlay& OnObjectEndPlay() = 0;
 	virtual const TCHAR* GetPropertyName(uint32 InPropertyStringId) const = 0;
 	virtual const RecordingInfoTimeline* GetRecordingInfo(uint32 RecordingId) const = 0; 
 	virtual void ReadViewTimeline(TFunctionRef<void(const ViewTimeline&)> Callback) const = 0;
 	virtual uint64 FindPossessingController(uint64 Pawn, double Time) const = 0;
-	virtual TRange<double> GetObjectTraceLifetime(uint64 ObjectId) const = 0;
-	virtual TRange<double> GetObjectRecordingLifetime(uint64 ObjectId) const = 0;
+	virtual TRange<double> GetObjectTraceLifetime(const RewindDebugger::FObjectId& InObjectId) const = 0;
+	virtual TRange<double> GetObjectRecordingLifetime(const RewindDebugger::FObjectId& InObjectId) const = 0;
 	virtual double GetRecordingDuration() const = 0;
 	virtual FSlateIcon FindIconForClass(uint64 ClassId) const = 0;
 	virtual bool GetObjectTransform(uint64 InObjectId, double InStartTime, double InEndTime, FTransform& OutTransform) const = 0;
-};
\ No newline at end of file
+
+	UE_DEPRECATED(5.7, "Use the version taking FObjectId instead")
+	virtual bool ReadObjectEventsTimeline(uint64 InObjectId, TFunctionRef<void(const ObjectEventsTimeline&)> InCallback) const final
+	{
+		return ReadObjectEventsTimeline(RewindDebugger::FObjectId(InObjectId), InCallback);
+	}
+
+	UE_DEPRECATED(5.7, "Use the version taking FObjectId instead")
+	virtual bool ReadObjectEvent(uint64 InObjectId, uint64 InMessageId, TFunctionRef<void(const FObjectEventMessage&)> InCallback) const final
+	{
+		return ReadObjectEvent(RewindDebugger::FObjectId(InObjectId), InMessageId, InCallback);
+	}
+
+	UE_DEPRECATED(5.7, "Use the version taking FObjectId instead")
+	virtual void EnumerateSubobjects(uint64 InObjectId, TFunctionRef<void(uint64 SubobjectId)> InCallback) const final
+	{
+	}
+
+	UE_DEPRECATED(5.7, "Use the version taking FObjectId instead")
+	virtual TRange<double> GetObjectTraceLifetime(uint64 InObjectId) const final
+	{
+		return GetObjectTraceLifetime(RewindDebugger::FObjectId(InObjectId));
+	}
+
+	UE_DEPRECATED(5.7, "Use the version taking FObjectId instead")
+	virtual TRange<double> GetObjectRecordingLifetime(uint64 InObjectId) const final
+	{
+		return GetObjectRecordingLifetime(RewindDebugger::FObjectId(InObjectId));
+	}
+};
+
+inline const FObjectInfo* IGameplayProvider::FindObjectInfo(const uint64 InObjectId) const
+{
+	return FindObjectInfo(RewindDebugger::FObjectId(InObjectId));
+}
+
+inline const FObjectInfo& IGameplayProvider::GetObjectInfo(const uint64 InObjectId) const
+{
+	return GetObjectInfo(RewindDebugger::FObjectId(InObjectId));
+}
diff --git a/Engine/Plugins/Animation/GameplayInsights/Source/RewindDebugger/Private/AnimInstanceHelpers.cpp b/Engine/Plugins/Animation/GameplayInsights/Source/RewindDebugger/Private/AnimInstanceHelpers.cpp
--- a/Engine/Plugins/Animation/GameplayInsights/Source/RewindDebugger/Private/AnimInstanceHelpers.cpp
+++ b/Engine/Plugins/Animation/GameplayInsights/Source/RewindDebugger/Private/AnimInstanceHelpers.cpp
@@ -73,7 +73,7 @@ bool FAnimInstanceDoubleClickHandler::HandleDoubleClick(IRewindDebugger* RewindD
 	TSharedPtr<FDebugObjectInfo> SelectedObject = RewindDebugger->GetSelectedObject();
 	if (SelectedObject.IsValid())
 	{
-		return OpenAnimBlueprintAndAttachDebugger(RewindDebugger->GetAnalysisSession(), SelectedObject->ObjectId);
+		return OpenAnimBlueprintAndAttachDebugger(RewindDebugger->GetAnalysisSession(), SelectedObject->GetUObjectId());
 	}
 	return false;
 }
@@ -98,7 +98,7 @@ void FAnimInstanceMenu::Register()
 						LOCTEXT("Open AnimBP", "Open/Debug AnimGraph"),
 							LOCTEXT("Open AnimBP ToolTip", "Open this Animation Blueprint and attach the debugger to this instance"),
 							FSlateIcon(),
-							FExecuteAction::CreateLambda([ObjectId = Context->SelectedObject->ObjectId]()
+							FExecuteAction::CreateLambda([ObjectId = Context->SelectedObject->GetUObjectId()]()
 							{
 								const IUnrealInsightsModule& UnrealInsightsModule = FModuleManager::LoadModuleChecked<IUnrealInsightsModule>("TraceInsights");
 								const TSharedPtr<const TraceServices::IAnalysisSession> Session = UnrealInsightsModule.GetAnalysisSession();
diff --git a/Engine/Plugins/Animation/GameplayInsights/Source/RewindDebugger/Private/BlueprintDoubleClickHandler.cpp b/Engine/Plugins/Animation/GameplayInsights/Source/RewindDebugger/Private/BlueprintDoubleClickHandler.cpp
--- a/Engine/Plugins/Animation/GameplayInsights/Source/RewindDebugger/Private/BlueprintDoubleClickHandler.cpp
+++ b/Engine/Plugins/Animation/GameplayInsights/Source/RewindDebugger/Private/BlueprintDoubleClickHandler.cpp
@@ -50,7 +50,7 @@ bool FBlueprintDoubleClickHandler::HandleDoubleClick(IRewindDebugger* RewindDebu
 	const TSharedPtr<FDebugObjectInfo> SelectedObject = RewindDebugger->GetSelectedObject();
 	if (SelectedObject.IsValid())
 	{
-		return OpenBlueprintAndAttachDebugger(RewindDebugger->GetAnalysisSession(), SelectedObject->ObjectId);
+		return OpenBlueprintAndAttachDebugger(RewindDebugger->GetAnalysisSession(), SelectedObject->GetUObjectId());
 	}
 #endif
 	
diff --git a/Engine/Plugins/Animation/GameplayInsights/Source/RewindDebugger/Private/PropertyTraceMenu.cpp b/Engine/Plugins/Animation/GameplayInsights/Source/RewindDebugger/Private/PropertyTraceMenu.cpp
--- a/Engine/Plugins/Animation/GameplayInsights/Source/RewindDebugger/Private/PropertyTraceMenu.cpp
+++ b/Engine/Plugins/Animation/GameplayInsights/Source/RewindDebugger/Private/PropertyTraceMenu.cpp
@@ -25,26 +25,26 @@ void FPropertyTraceMenu::Register()
 				LOCTEXT("Trace Object Properties", "Trace Object Properties"),
 				LOCTEXT("Trace Object Properties Tooltip", "Record this object's properties so they will show in the Rewind Debugger when scrubbing."),
 				FSlateIcon(), 
-				FUIAction(FExecuteAction::CreateLambda([SelectedObject=Context->SelectedObject]()
+				FUIAction(FExecuteAction::CreateLambda([SelectedObject = Context->SelectedObject]()
 					{
 						IGameplayInsightsModule& GameplayInsightsModule = FModuleManager::GetModuleChecked<IGameplayInsightsModule>("GameplayInsights");
 						if (SelectedObject.IsValid())
 						{
-							if (UObject* Object = FObjectTrace::GetObjectFromId(SelectedObject->ObjectId))
+							if (UObject* Object = FObjectTrace::GetObjectFromId(SelectedObject->GetUObjectId()))
 							{
 								GameplayInsightsModule.EnableObjectPropertyTrace(Object, !GameplayInsightsModule.IsObjectPropertyTraceEnabled(Object));
 							}
 						}
 					}),
 					FCanExecuteAction(),
-					FGetActionCheckState::CreateLambda([SelectedObject=Context->SelectedObject]()
+					FGetActionCheckState::CreateLambda([SelectedObject = Context->SelectedObject]()
 					{
 						IGameplayInsightsModule& GameplayInsightsModule = FModuleManager::GetModuleChecked<IGameplayInsightsModule>("GameplayInsights");
 						bool bEnabled = false;
 	 
 						if (SelectedObject.IsValid())
 						{
-							if (UObject* Object = FObjectTrace::GetObjectFromId(SelectedObject->ObjectId))
+							if (UObject* Object = FObjectTrace::GetObjectFromId(SelectedObject->GetUObjectId()))
 							{
 								bEnabled = GameplayInsightsModule.IsObjectPropertyTraceEnabled(Object);
 							}
diff --git a/Engine/Plugins/Animation/GameplayInsights/Source/RewindDebugger/Private/RewindDebugger.cpp b/Engine/Plugins/Animation/GameplayInsights/Source/RewindDebugger/Private/RewindDebugger.cpp
--- a/Engine/Plugins/Animation/GameplayInsights/Source/RewindDebugger/Private/RewindDebugger.cpp
+++ b/Engine/Plugins/Animation/GameplayInsights/Source/RewindDebugger/Private/RewindDebugger.cpp
@@ -114,9 +114,9 @@ FRewindDebugger::FRewindDebugger()
 			GetTargetObjectIds(CandidateIds);
 			// make sure all the SubObjects of the target actor have been traced
 #if OBJECT_TRACE_ENABLED
-			for (uint64 TargetObjectId : CandidateIds)
+			for (const RewindDebugger::FObjectId& TargetObjectId : CandidateIds)
 			{
-				if (const UObject* TargetObject = FObjectTrace::GetObjectFromId(TargetObjectId))
+				if (const UObject* TargetObject = FObjectTrace::GetObjectFromId(TargetObjectId.GetMainId()))
 				{
 					TraceSubobjects(TargetObject);
 				}
@@ -274,7 +274,7 @@ uint64 FRewindDebugger::GetTargetActorId() const
 				{
 					if (SelectedObjectName.Get() == InObjectInfo.Name)
 					{
-						TargetActorId = InObjectInfo.Id;
+						TargetActorId = InObjectInfo.GetUObjectId();
 					}
 				});
 		}
@@ -283,7 +283,7 @@ uint64 FRewindDebugger::GetTargetActorId() const
 	return TargetActorId;
 }
 
-void FRewindDebugger::GetTargetObjectIds(TArray<uint64>& OutTargetObjectIds) const
+void FRewindDebugger::GetTargetObjectIds(TArray<RewindDebugger::FObjectId>& OutTargetObjectIds) const
 {
 	OutTargetObjectIds.Empty(2);
 
@@ -300,9 +300,10 @@ void FRewindDebugger::GetTargetObjectIds(TArray<uint64>& OutTargetObjectIds) con
 		{
 			GameplayProvider->EnumerateObjects(CurrentTraceRange.GetLowerBoundValue(), CurrentTraceRange.GetUpperBoundValue(), [this, &OutTargetObjectIds](const FObjectInfo& InObjectInfo)
 				{
-					if (SelectedObjectName.Get() == InObjectInfo.Name)
+					// Skip children
+					if (!InObjectInfo.GetId().IsChildElement() && SelectedObjectName.Get() == InObjectInfo.Name)
 					{
-						OutTargetObjectIds.Add(InObjectInfo.Id);
+						OutTargetObjectIds.Add(InObjectInfo.GetId());
 					}
 				});
 		}
@@ -312,9 +313,9 @@ void FRewindDebugger::GetTargetObjectIds(TArray<uint64>& OutTargetObjectIds) con
 #if OBJECT_TRACE_ENABLED
 	if (IsRecording())
 	{
-		for (uint64 CandidateId : CandidateIds)
+		for (const RewindDebugger::FObjectId& CandidateId : CandidateIds)
 		{
-			if (const UObject* TargetObject = FObjectTrace::GetObjectFromId(CandidateId))
+			if (const UObject* TargetObject = FObjectTrace::GetObjectFromId(CandidateId.GetMainId()))
 			{
 				TraceSubobjects(TargetObject);
 			}
@@ -353,51 +354,48 @@ void FRewindDebugger::RefreshDebugTracks()
 			TrackListChangedDelegate.ExecuteIfBound();
 		}
 	}
-	else
+	else if (DebugTracks.Num() || CandidateIds.Num())
 	{
-		if (GetAnalysisSession())
-		{
-			bool bChanged = false;
+		bool bChanged = false;
 
-			// remove any existing tracks that don't match the current list of object ids
-			for (int TrackIndex = DebugTracks.Num() - 1; TrackIndex >= 0; TrackIndex--)
+		// remove any existing tracks that don't match the current list of object ids
+		for (int TrackIndex = DebugTracks.Num() - 1; TrackIndex >= 0; TrackIndex--)
+		{
+			if (!CandidateIds.Contains(DebugTracks[TrackIndex]->GetAssociatedObjectId()))
 			{
-				if (!CandidateIds.Contains(DebugTracks[TrackIndex]->GetObjectId()))
-				{
-					DebugTracks.RemoveAt(TrackIndex);
-				}
+				DebugTracks.RemoveAt(TrackIndex);
 			}
+		}
 
-			// add new tracks for current list of object identifiers if they don't already exist
-			for (uint64 CandidateIdentifier : CandidateIds)
-			{
-				const TSharedPtr<RewindDebugger::FRewindDebuggerTrack>* FoundTrack = DebugTracks.FindByPredicate(
-					[CandidateIdentifier](const TSharedPtr<RewindDebugger::FRewindDebuggerTrack>& Track)
-					{
-						return Track->GetObjectId() == CandidateIdentifier;
-					});
-
-				if (!FoundTrack)
+		// add new tracks for current list of object identifiers if they don't already exist
+		for (const RewindDebugger::FObjectId& CandidateIdentifier : CandidateIds)
+		{
+			const TSharedPtr<RewindDebugger::FRewindDebuggerTrack>* FoundTrack = DebugTracks.FindByPredicate(
+				[CandidateIdentifier](const TSharedPtr<RewindDebugger::FRewindDebuggerTrack>& Track)
 				{
-					DebugTracks.Add(MakeShared<RewindDebugger::FRewindDebuggerObjectTrack>(CandidateIdentifier, SelectedObjectName.Get(), true));
-					bChanged = true;
-				}
-			}
+					return Track->GetAssociatedObjectId() == CandidateIdentifier;
+				});
 
-			// update all tracks
-			for (const TSharedPtr<RewindDebugger::FRewindDebuggerTrack>& DebugTrack : DebugTracks)
+			if (!FoundTrack)
 			{
-				if (DebugTrack->Update())
-				{
-					bChanged = true;
-				}
+				DebugTracks.Add(MakeShared<RewindDebugger::FRewindDebuggerObjectTrack>(CandidateIdentifier, SelectedObjectName.Get(), true));
+				bChanged = true;
 			}
+		}
 
-			if (bChanged)
+		// update all tracks
+		for (const TSharedPtr<RewindDebugger::FRewindDebuggerTrack>& DebugTrack : DebugTracks)
+		{
+			if (DebugTrack->Update())
 			{
-				TrackListChangedDelegate.ExecuteIfBound();
+				bChanged = true;
 			}
 		}
+
+		if (bChanged)
+		{
+			TrackListChangedDelegate.ExecuteIfBound();
+		}
 	}
 }
 
@@ -985,12 +983,12 @@ const FObjectInfo* FRewindDebugger::FindOwningActorInfo(const IGameplayProvider*
 			return &ObjectInfo;
 		}
 
-		if (ObjectInfo.OuterId == 0)
+		if (!ObjectInfo.GetOuterId().IsSet())
 		{
 			return nullptr;
 		}
 
-		InObjectId = ObjectInfo.OuterId;
+		InObjectId = ObjectInfo.GetOuterUObjectId();
 	}
 }
 
@@ -1091,7 +1089,7 @@ void FRewindDebugger::Tick(float DeltaTime)
 						if (FrameProvider.GetFrameFromTime(ETraceFrameType::TraceFrameType_Game, CurrentTraceTime, Frame))
 						{
 							bool bNewActor = false;
-							if (!CandidateIds.Contains(TargetActorIdForMesh))
+							if (!CandidateIds.Contains(RewindDebugger::FObjectId{TargetActorIdForMesh}))
 							{
 								AnimationProvider->EnumerateSkeletalMeshPoseTimelines([this, &bNewActor, GameplayProvider](uint64 ObjectId, const IAnimationProvider::SkeletalMeshPoseTimeline& TimelineData)
 									{
@@ -1103,11 +1101,10 @@ void FRewindDebugger::Tick(float DeltaTime)
 											{
 												if (const FObjectInfo* ActorInfo = FindOwningActorInfo(GameplayProvider, ObjectId))
 												{
-													const uint64 ActorId = ActorInfo->Id;
-													if (CandidateIds.Contains(ActorId))
+													if (CandidateIds.Contains(ActorInfo->GetId()))
 													{
 														bNewActor = true;
-														TargetActorIdForMesh = ActorId;
+														TargetActorIdForMesh = ActorInfo->GetUObjectId();
 														TargetActorMeshId = ObjectId;
 													}
 												}
@@ -1632,7 +1629,7 @@ TSharedPtr<SWidget> FRewindDebugger::BuildTrackContextMenu() const
 			TraceServices::FAnalysisSessionReadScope SessionReadScope(*Session);
 
 			const IGameplayProvider* GameplayProvider = Session->ReadProvider<IGameplayProvider>("GameplayProvider");
-			const FObjectInfo& ObjectInfo = GameplayProvider->GetObjectInfo(SelectedTrack->GetObjectId());
+			const FObjectInfo& ObjectInfo = GameplayProvider->GetObjectInfo(SelectedTrack->GetAssociatedObjectId());
 			uint64 ClassId = ObjectInfo.ClassId;
 			while (ClassId != 0)
 			{
@@ -1656,7 +1653,7 @@ TSharedPtr<FDebugObjectInfo> FRewindDebugger::GetSelectedObject() const
 			SelectedObject = MakeShared<FDebugObjectInfo>();
 		}
 
-		SelectedObject->ObjectId = SelectedTrack->GetObjectId();
+		SelectedObject->Id = SelectedTrack->GetAssociatedObjectId();
 		SelectedObject->ObjectName = SelectedTrack->GetDisplayName().ToString();
 		return SelectedObject;
 	}
@@ -1675,7 +1672,7 @@ void FRewindDebugger::RefreshDebuggedObjects(TArray<TSharedPtr<RewindDebugger::F
 	OutObjects.SetNum(0, EAllowShrinking::No);
 	for (const TSharedPtr<RewindDebugger::FRewindDebuggerTrack>& Track : InTracks)
 	{
-		const int Index = OutObjects.Add(MakeShared<FDebugObjectInfo>(Track->GetObjectId(), Track->GetDisplayName().ToString()));
+		const int Index = OutObjects.Add(MakeShared<FDebugObjectInfo>(Track->GetAssociatedObjectId(), Track->GetDisplayName().ToString()));
 		TArray<TSharedPtr<RewindDebugger::FRewindDebuggerTrack>> TrackChildren;
 		Track->IterateSubTracks([&TrackChildren](const TSharedPtr<RewindDebugger::FRewindDebuggerTrack>& Child) { TrackChildren.Add(Child); });
 		RefreshDebuggedObjects(TrackChildren, OutObjects[Index]->Children);
@@ -1692,7 +1689,7 @@ bool FRewindDebugger::IsObjectCurrentlyDebugged(uint64 ObjectId) const
 {
 	for (const TSharedPtr<RewindDebugger::FRewindDebuggerTrack>& Track : DebugTracks)
 	{
-		if (Track->GetObjectId() == ObjectId)
+		if (Track->GetUObjectId() == ObjectId)
 		{
 			return true;
 		}
@@ -1700,7 +1697,7 @@ bool FRewindDebugger::IsObjectCurrentlyDebugged(uint64 ObjectId) const
 		bool Found = false;
 		Track->IterateSubTracks([ObjectId, &Found](const TSharedPtr<RewindDebugger::FRewindDebuggerTrack>& Child)
 			{
-				if (Child->GetObjectId() == ObjectId)
+				if (Child->GetUObjectId() == ObjectId)
 				{
 					Found = true;
 					// @Todo STDBG: want to stop iteration here
diff --git a/Engine/Plugins/Animation/GameplayInsights/Source/RewindDebugger/Private/RewindDebugger.h b/Engine/Plugins/Animation/GameplayInsights/Source/RewindDebugger/Private/RewindDebugger.h
--- a/Engine/Plugins/Animation/GameplayInsights/Source/RewindDebugger/Private/RewindDebugger.h
+++ b/Engine/Plugins/Animation/GameplayInsights/Source/RewindDebugger/Private/RewindDebugger.h
@@ -12,13 +12,15 @@
 #include "IRewindDebuggerTrackCreator.h"
 #include "RewindDebuggerTrack.h"
 
-namespace TraceServices { class IAnalysisSession; }
+namespace TraceServices
+{
+	class IAnalysisSession;
+}
 class FMenuBuilder;
 class IGameplayProvider;
 class SDockTab;
 class SWidget;
 class USkeletalMeshComponent;
-struct FGameplayObjectIdentifier;
 struct FObjectInfo;
 
 /**
@@ -81,7 +83,7 @@ public:
 
 	void OnConnection();
 
-	void GetTargetObjectIds(TArray<uint64>& OutTargetObjectIds) const;
+	void GetTargetObjectIds(TArray<RewindDebugger::FObjectId>& OutTargetObjectIds) const;
 
 	TArray<TSharedPtr<RewindDebugger::FRewindDebuggerTrack>>& GetTracks()
 	{
@@ -270,7 +272,7 @@ private:
 	TArray<TSharedPtr<RewindDebugger::FRewindDebuggerTrack>> DebugTracks;
 	TSharedPtr<RewindDebugger::FRewindDebuggerTrack> SelectedTrack;
 
-	TArray<uint64> CandidateIds;
+	TArray<RewindDebugger::FObjectId> CandidateIds;
 
 	mutable class IUnrealInsightsModule* UnrealInsightsModule = nullptr;
 	FTSTicker::FDelegateHandle TickerHandle;
diff --git a/Engine/Plugins/Animation/GameplayInsights/Source/RewindDebugger/Private/RewindDebuggerAnimation.cpp b/Engine/Plugins/Animation/GameplayInsights/Source/RewindDebugger/Private/RewindDebuggerAnimation.cpp
--- a/Engine/Plugins/Animation/GameplayInsights/Source/RewindDebugger/Private/RewindDebuggerAnimation.cpp
+++ b/Engine/Plugins/Animation/GameplayInsights/Source/RewindDebugger/Private/RewindDebuggerAnimation.cpp
@@ -241,7 +241,7 @@ UAnimInstance* FRewindDebuggerAnimation::SpawnAnimInstance(uint64 ObjectId, cons
 			TSoftObjectPtr<UAnimBlueprintGeneratedClass> ClassPtr =	TSoftObjectPtr<UAnimBlueprintGeneratedClass>(FSoftObjectPath(FStringView(ClassInfo->PathName)));
 			if(ClassPtr.LoadSynchronous())
 			{
-				if (const FSpawnedMeshComponentInfo* MeshInfo = SpawnedMeshComponents.Find(ObjectInfo->OuterId))
+				if (const FSpawnedMeshComponentInfo* MeshInfo = SpawnedMeshComponents.Find(ObjectInfo->GetOuterUObjectId()))
 				{
 					AnimInstanceInfo->AnimInstance = NewObject<UAnimInstance>(MeshInfo->Component, ClassPtr.Get());
 					return AnimInstanceInfo->AnimInstance.Get();
diff --git a/Engine/Plugins/Animation/GameplayInsights/Source/RewindDebugger/Private/RewindDebuggerObjectTrack.cpp b/Engine/Plugins/Animation/GameplayInsights/Source/RewindDebugger/Private/RewindDebuggerObjectTrack.cpp
--- a/Engine/Plugins/Animation/GameplayInsights/Source/RewindDebugger/Private/RewindDebuggerObjectTrack.cpp
+++ b/Engine/Plugins/Animation/GameplayInsights/Source/RewindDebugger/Private/RewindDebuggerObjectTrack.cpp
@@ -18,28 +18,25 @@ namespace RewindDebugger
 {
 
 // check if an object is or is a subclass of a type by name, based on Insights traced type info
-static bool IsTargetType(uint64 ObjectId, FName TargetTypeName, const TraceServices::IAnalysisSession& Session)
+static bool GetTypeHierarchyNames(const uint64 InObjectId, const TraceServices::IAnalysisSession& InSession, TArray<const TCHAR*>& OutTypeNames)
 {
-	TraceServices::FAnalysisSessionReadScope SessionReadScope(Session);
+	TraceServices::FAnalysisSessionReadScope SessionReadScope(InSession);
 
-	const IGameplayProvider* GameplayProvider = Session.ReadProvider<IGameplayProvider>("GameplayProvider");
-	const FObjectInfo& ObjectInfo = GameplayProvider->GetObjectInfo(ObjectId);
+	const IGameplayProvider* GameplayProvider = InSession.ReadProvider<IGameplayProvider>("GameplayProvider");
+	const FObjectInfo& ObjectInfo = GameplayProvider->GetObjectInfo(InObjectId);
 	uint64 ClassId = ObjectInfo.ClassId;
 
-	while (ClassId != 0)
+	while (ClassId != FObjectId::InvalidId)
 	{
 		const FClassInfo& ClassInfo = GameplayProvider->GetClassInfo(ClassId);
-		if (ClassInfo.Name == TargetTypeName)
-		{
-			return true;
-		}
+		OutTypeNames.Add(ClassInfo.Name);
 		ClassId = ClassInfo.SuperId;
 	}
 
-	return false;
+	return !OutTypeNames.IsEmpty();
 }
 
-FRewindDebuggerObjectTrack::FRewindDebuggerObjectTrack(uint64 InObjectId, const FString& InObjectName, bool bInAddController)
+FRewindDebuggerObjectTrack::FRewindDebuggerObjectTrack(const FObjectId& InObjectId, const FString& InObjectName, bool bInAddController)
 	: ObjectName(InObjectName)
 	, ObjectId(InObjectId)
 	, bAddController(bInAddController)
@@ -47,17 +44,24 @@ FRewindDebuggerObjectTrack::FRewindDebuggerObjectTrack(uint64 InObjectId, const
 	, bIconSearched(false)
 {
 	ExistenceRange = MakeShared<SEventTimelineView::FTimelineEventData>();
+	
+	const TraceServices::IAnalysisSession* Session = IRewindDebugger::Instance()->GetAnalysisSession();
 
-	IRewindDebugger* RewindDebugger = IRewindDebugger::Instance();
-	const TraceServices::IAnalysisSession* Session = RewindDebugger->GetAnalysisSession();
-	TraceServices::FAnalysisSessionReadScope SessionReadScope(*Session);
-	FRewindDebuggerTrackCreators::EnumerateCreators([this, &Session](const IRewindDebuggerTrackCreator* Creator)
-		{
-			if (IsTargetType(ObjectId, Creator->GetTargetTypeName(), *Session))
+	if (TArray<const TCHAR*> TypeNames; GetTypeHierarchyNames(InObjectId.GetMainId(), *Session, TypeNames))
+	{
+		FRewindDebuggerTrackCreators::EnumerateCreators([&TrackChildren = TrackChildren, &TypeNames, IsChildElementTrack = InObjectId.IsChildElement()](const IRewindDebuggerTrackCreator* Creator)
 			{
-				TrackChildren.Push({ Creator, nullptr });
-			}
-		});
+				if (IsChildElementTrack && !Creator->IsChildElementSupported())
+				{
+					return;
+				}
+
+				if (TypeNames.FindByPredicate([TargetTypeName = Creator->GetTargetTypeName()](const TCHAR* TypeName) { return TargetTypeName == TypeName; }))
+				{
+					TrackChildren.Push({ Creator, /*Track*/nullptr });
+				}
+			});
+	}
 
 	// sort by creators by priority + name
 	TrackChildren.Sort([](const FTrackCreatorAndTrack& A, const FTrackCreatorAndTrack& B)
@@ -145,11 +149,7 @@ FText FRewindDebuggerObjectTrack::GetDisplayNameInternal() const
 {
 	if (!bDisplayNameValid)
 	{
-		if (ObjectId == 0)
-		{
-			DisplayName = FText::FromString(ObjectName);
-		}
-		else
+		if (ObjectId.IsSet())
 		{
 			IRewindDebugger* RewindDebugger = IRewindDebugger::Instance();
 			const TraceServices::IAnalysisSession* Session = RewindDebugger->GetAnalysisSession();
@@ -161,7 +161,7 @@ FText FRewindDebuggerObjectTrack::GetDisplayNameInternal() const
 
 			DisplayName = FText::FromString(ObjectInfo.Name);
 
-			if (const FWorldInfo* WorldInfo = GameplayProvider->FindWorldInfoFromObject(ObjectId))
+			if (const FWorldInfo* WorldInfo = GameplayProvider->FindWorldInfoFromObject(ObjectId.GetMainId()))
 			{
 				if (WorldInfo->NetMode == FWorldInfo::ENetMode::DedicatedServer)
 				{
@@ -169,6 +169,10 @@ FText FRewindDebuggerObjectTrack::GetDisplayNameInternal() const
 				}
 			}
 		}
+		else
+		{
+			DisplayName = FText::FromString(ObjectName);
+		}
 
 		bDisplayNameValid = true;
 	}
@@ -211,7 +215,7 @@ bool FRewindDebuggerObjectTrack::UpdateInternal()
 		bChanged = true;
 	}
 
-	TArray<uint64, TInlineAllocator<32>> FoundObjects;
+	TArray<FObjectId, TInlineAllocator<32>> FoundObjects;
 
 	FoundObjects.Add(ObjectId); // prevent debug views from being removed
 
@@ -242,32 +246,35 @@ bool FRewindDebuggerObjectTrack::UpdateInternal()
 	// Fallback code path to add views with no track implementation
 	{
 		TRACE_CPUPROFILER_EVENT_SCOPE(FRewindDebuggerObjectTrack::UpdateInternal_AddViews);
-		FRewindDebuggerViewCreators::EnumerateCreators([this, &Session, &bChanged, RewindDebugger](const IRewindDebuggerViewCreator* Creator)
-			{
-				int32 FoundIndex = Children.FindLastByPredicate([&bChanged, this, Creator, RewindDebugger](const TSharedPtr<FRewindDebuggerTrack>& Track) { return Track->GetName() == Creator->GetName(); });
 
-				bool bHasDebugInfo = Creator->HasDebugInfo(ObjectId)
-					&& IsTargetType(ObjectId, Creator->GetTargetTypeName(), *Session);
-
-				if (FoundIndex >= 0)
+		if (TArray<const TCHAR*> TypeNames; GetTypeHierarchyNames(ObjectId.GetMainId(), *Session, TypeNames))
+		{
+			FRewindDebuggerViewCreators::EnumerateCreators([ObjectId = ObjectId.GetMainId(), &Children = Children, &bChanged, &TypeNames](const IRewindDebuggerViewCreator* Creator)
 				{
-					if (!bHasDebugInfo)
+					const int32 FoundIndex = Children.FindLastByPredicate([Creator](const TSharedPtr<FRewindDebuggerTrack>& Track) { return Track->GetName() == Creator->GetName(); });
+
+					const bool bHasDebugInfo = Creator->HasDebugInfo(ObjectId)
+						&& TypeNames.FindByPredicate([TargetTypeName = Creator->GetTargetTypeName()](const TCHAR* TypeName) { return TargetTypeName == TypeName; });
+
+					if (FoundIndex >= 0)
 					{
-						bChanged = true;
-						Children.RemoveAt(FoundIndex);
+						if (!bHasDebugInfo)
+						{
+							bChanged = true;
+							Children.RemoveAt(FoundIndex);
+						}
 					}
-				}
-				else
-				{
-					if (bHasDebugInfo)
+					else
 					{
-						bChanged = true;
-						TSharedPtr<FRewindDebuggerTrack> Track = MakeShared<FRewindDebuggerFallbackTrack>(ObjectId, Creator);
-						Children.Add(Track);
+						if (bHasDebugInfo)
+						{
+							bChanged = true;
+							const TSharedPtr<FRewindDebuggerTrack> Track = MakeShared<FRewindDebuggerFallbackTrack>(ObjectId, Creator);
+							Children.Add(Track);
+						}
 					}
-				}
-			}
-		);
+				});
+		}
 	}
 
 	// add child objects
@@ -276,7 +283,7 @@ bool FRewindDebuggerObjectTrack::UpdateInternal()
 
 		TRange<double> ViewRange = RewindDebugger->GetCurrentViewRange();
 
-		GameplayProvider->EnumerateSubobjects(ObjectId, [this, &FoundObjects, &bChanged, &ViewRange, GameplayProvider](uint64 SubObjectId)
+		GameplayProvider->EnumerateSubobjects(ObjectId, [this, &FoundObjects, &bChanged, &ViewRange, GameplayProvider](const FObjectId& SubObjectId)
 			{
 				const TRange<double> Lifetime = GameplayProvider->GetObjectRecordingLifetime(SubObjectId);
 				const TRange<double> Overlap = TRange<double>::Intersection(Lifetime, ViewRange);
@@ -285,7 +292,7 @@ bool FRewindDebuggerObjectTrack::UpdateInternal()
 				{
 					const int32 FoundIndex = Children.FindLastByPredicate([SubObjectId](const TSharedPtr<FRewindDebuggerTrack>& Track)
 						{
-							return Track->GetObjectId() == SubObjectId;
+							return Track->GetAssociatedObjectId() == SubObjectId;
 						});
 
 					if (FoundIndex == INDEX_NONE)
@@ -295,7 +302,7 @@ bool FRewindDebuggerObjectTrack::UpdateInternal()
 						bChanged = true;
 					}
 
-					FoundObjects.Add(SubObjectId);
+					FoundObjects.Add(FObjectId{SubObjectId});
 				}
 			});
 	}
@@ -305,28 +312,28 @@ bool FRewindDebuggerObjectTrack::UpdateInternal()
 	{
 		TRACE_CPUPROFILER_EVENT_SCOPE(FRewindDebuggerObjectTrack::FindController);
 		// Should probably update this to use a time range and return all possessing controllers from the visible time range.  For now just returns the one at the current time.
-		if (uint64 ControllerId = GameplayProvider->FindPossessingController(ObjectId, RewindDebugger->CurrentTraceTime()))
+		if (uint64 ControllerId = GameplayProvider->FindPossessingController(ObjectId.GetMainId(), RewindDebugger->CurrentTraceTime()))
 		{
 			const FObjectInfo& ObjectInfo = GameplayProvider->GetObjectInfo(ControllerId);
 			const int32 FoundIndex = Children.FindLastByPredicate([ObjectInfo](const TSharedPtr<FRewindDebuggerTrack>& Track)
 				{
-					return Track->GetObjectId() == ObjectInfo.Id;
+					return Track->GetAssociatedObjectId() == ObjectInfo.GetId();
 				});
 
 			if (FoundIndex < 0)
 			{
 				bChanged = true;
-				Children.Add(MakeShared<FRewindDebuggerObjectTrack>(ObjectInfo.Id, ObjectInfo.Name));
+				Children.Add(MakeShared<FRewindDebuggerObjectTrack>(ObjectInfo.GetId(), ObjectInfo.Name));
 			}
 
-			FoundObjects.Add(ControllerId);
+			FoundObjects.Add(FObjectId{ControllerId});
 		}
 	}
 
 	// remove any components previously in the list that were not found in this time range.
 	for (int Index = Children.Num() - 1; Index >= 0; Index--)
 	{
-		if (!FoundObjects.Contains(Children[Index]->GetObjectId()))
+		if (!FoundObjects.Contains(Children[Index]->GetAssociatedObjectId()))
 		{
 			bChanged = true;
 			Children.RemoveAt(Index);
@@ -336,6 +343,7 @@ bool FRewindDebuggerObjectTrack::UpdateInternal()
 	if (bChanged)
 	{
 		// sort child object tracks by name
+		// note that stable track ordering requires non-empty display names
 		Children.Sort([](const TSharedPtr<FRewindDebuggerTrack>& A, const TSharedPtr<FRewindDebuggerTrack>& B)
 			{
 				return A->GetDisplayName().ToString() < B->GetDisplayName().ToString();
@@ -369,6 +377,6 @@ bool FRewindDebuggerObjectTrack::UpdateInternal()
 	return bChanged;
 }
 
-}
+} // namespace RewindDebugger
 
 #undef LOCTEXT_NAMESPACE
\ No newline at end of file
diff --git a/Engine/Plugins/Animation/GameplayInsights/Source/RewindDebugger/Private/RewindDebuggerObjectTrack.h b/Engine/Plugins/Animation/GameplayInsights/Source/RewindDebugger/Private/RewindDebuggerObjectTrack.h
--- a/Engine/Plugins/Animation/GameplayInsights/Source/RewindDebugger/Private/RewindDebuggerObjectTrack.h
+++ b/Engine/Plugins/Animation/GameplayInsights/Source/RewindDebugger/Private/RewindDebuggerObjectTrack.h
@@ -20,7 +20,7 @@ class FRewindDebuggerObjectTrack : public FRewindDebuggerTrack
 {
 public:
 
-	FRewindDebuggerObjectTrack(uint64 InObjectId, const FString& InObjectName, bool bInAddController = false);
+	FRewindDebuggerObjectTrack(const FObjectId& InObjectId, const FString& InObjectName, bool bInAddController = false);
 
 	TSharedPtr<SEventTimelineView::FTimelineEventData> GetExistenceRange() const
 	{
@@ -46,7 +46,12 @@ private:
 
 	virtual uint64 GetObjectIdInternal() const override
 	{
-		return ObjectId;
+		return ObjectId.GetMainId();
+	}
+
+	virtual uint64 GetChildElementIdInternal() const override
+	{
+		return ObjectId.GetChildElementId();
 	}
 
 	virtual bool HasDebugDataInternal() const override
@@ -60,7 +65,7 @@ private:
 	FString ObjectName;
 	FSlateIcon Icon;
 	TSharedPtr<SEventTimelineView::FTimelineEventData> ExistenceRange;
-	uint64 ObjectId;
+	FObjectId ObjectId;
 	TArray<FTrackCreatorAndTrack> TrackChildren;
 	TArray<TSharedPtr<FRewindDebuggerTrack>> Children;
 
diff --git a/Engine/Plugins/Animation/GameplayInsights/Source/RewindDebugger/Private/SRewindDebugger.cpp b/Engine/Plugins/Animation/GameplayInsights/Source/RewindDebugger/Private/SRewindDebugger.cpp
--- a/Engine/Plugins/Animation/GameplayInsights/Source/RewindDebugger/Private/SRewindDebugger.cpp
+++ b/Engine/Plugins/Animation/GameplayInsights/Source/RewindDebugger/Private/SRewindDebugger.cpp
@@ -188,6 +188,7 @@ TSharedRef<SWidget> SRewindDebugger::MakeObjectSelectionMenu()
 	}
 	MenuBuilder.EndSection();
 
+	// Menu entry for each object with recorded data in the scene
 	MenuBuilder.BeginSection("From Recording Section", LOCTEXT("FromRecording", "From Recording:"));
 
 	if (const IRewindDebugger* RewindDebugger = IRewindDebugger::Instance())
@@ -197,11 +198,21 @@ TSharedRef<SWidget> SRewindDebugger::MakeObjectSelectionMenu()
 			TraceServices::FAnalysisSessionReadScope SessionReadScope(*Session);
 			const IGameplayProvider* GameplayProvider = Session->ReadProvider<IGameplayProvider>("GameplayProvider");
 
-			const FClassInfo* ActorClassInfo = GameplayProvider->FindClassInfo(*AActor::StaticClass()->GetPathName());
+			const FClassInfo* LevelClassInfo = GameplayProvider->FindClassInfo(*ULevel::StaticClass()->GetPathName());
+			const FClassInfo* WorldClassInfo = GameplayProvider->FindClassInfo(*UWorld::StaticClass()->GetPathName());
 
-			GameplayProvider->EnumerateObjects([this, &MenuBuilder, GameplayProvider, ActorClassInfo](const FObjectInfo& ObjectInfo)
+			// Accept any child of the World or Level
+			GameplayProvider->EnumerateObjects([this, &MenuBuilder, GameplayProvider, LevelClassInfo, WorldClassInfo](const FObjectInfo& ObjectInfo)
 				{
-					if (GameplayProvider->IsSubClassOf(ObjectInfo.ClassId, ActorClassInfo->Id))
+					const RewindDebugger::FObjectId& OuterIdentifier = ObjectInfo.GetOuterId();
+					if (!OuterIdentifier.IsSet())
+					{
+						return;
+					}
+	
+					const FObjectInfo& OuterInfo = GameplayProvider->GetObjectInfo(OuterIdentifier);
+					if (OuterInfo.ClassId == LevelClassInfo->Id
+						|| OuterInfo.ClassId == WorldClassInfo->Id)
 					{
 						const FString Name = ObjectInfo.Name;
 						const FText SelectedLabel = FText::FromString(Name);
@@ -384,24 +395,28 @@ void SRewindDebugger::Construct(const FArguments& InArgs, TSharedRef<FUICommandL
 														+ SHorizontalBox::Slot().Padding(3)
 														[
 															SNew(STextBlock)
-																.Text_Lambda([this]() {
-																if (Tracks == nullptr || Tracks->Num() == 0)
-																{
-																	return LOCTEXT("Select Actor", "Debug Target Actor");
-																}
+																.Text_Lambda([this]()
+																	{
+																		if (Tracks == nullptr || Tracks->Num() == 0)
+																		{
+																			return LOCTEXT("Select Actor", "Debug Target Actor");
+																		}
 
-																FText ReadableName = (*Tracks)[0]->GetDisplayName();
+																		FText ReadableName = (*Tracks)[0]->GetDisplayName();
 #if OBJECT_TRACE_ENABLED
-																if (UObject* Object = FObjectTrace::GetObjectFromId((*Tracks)[0]->GetObjectId()))
-																{
-																	if (const AActor* Actor = Cast<AActor>(Object))
-																	{
-																		ReadableName = FText::FromString(Actor->GetActorLabel());
-																	}
-																}
+																		const RewindDebugger::FObjectId ObjectId = (*Tracks)[0]->GetAssociatedObjectId();
+																		if (!ObjectId.IsChildElement())
+																		{
+																			if (UObject* Object = FObjectTrace::GetObjectFromId(ObjectId.GetMainId()))
+																			{
+																				if (const AActor* Actor = Cast<AActor>(Object))
+																				{
+																					ReadableName = FText::FromString(Actor->GetActorLabel());
+																				}
+																			}
+																		}
 #endif // OBJECT_TRACE_ENABLED
-
-																return ReadableName;
+																		return ReadableName;
 																	})
 														]
 												]
diff --git a/Engine/Plugins/Animation/GameplayInsights/Source/RewindDebuggerVLog/Private/RewindDebuggerVLog.cpp b/Engine/Plugins/Animation/GameplayInsights/Source/RewindDebuggerVLog/Private/RewindDebuggerVLog.cpp
--- a/Engine/Plugins/Animation/GameplayInsights/Source/RewindDebuggerVLog/Private/RewindDebuggerVLog.cpp
+++ b/Engine/Plugins/Animation/GameplayInsights/Source/RewindDebuggerVLog/Private/RewindDebuggerVLog.cpp
@@ -213,7 +213,7 @@ void FRewindDebuggerVLog::AddLogEntries(const TArray<TSharedPtr<FDebugObjectInfo
 {
 	for (const TSharedPtr<FDebugObjectInfo>& ComponentInfo : Components)
 	{
-		const uint64 ObjectId = ComponentInfo->ObjectId;
+		const uint64 ObjectId = ComponentInfo->GetUObjectId();
 		if (!ObjectsVisited.Contains(ObjectId))
 		{
 			ObjectsVisited.Add(ObjectId);
diff --git a/Engine/Plugins/Animation/GameplayInsights/Source/RewindDebuggerVLog/Private/VisualLogTrack.cpp b/Engine/Plugins/Animation/GameplayInsights/Source/RewindDebuggerVLog/Private/VisualLogTrack.cpp
--- a/Engine/Plugins/Animation/GameplayInsights/Source/RewindDebuggerVLog/Private/VisualLogTrack.cpp
+++ b/Engine/Plugins/Animation/GameplayInsights/Source/RewindDebuggerVLog/Private/VisualLogTrack.cpp
@@ -214,9 +214,9 @@ void FVisualLogTrackCreator::GetTrackTypesInternal(TArray<FRewindDebuggerTrackTy
 	Types.Add({ VisualLogName, LOCTEXT("Visual Logging", "Visual Logging") });
 }
 
-TSharedPtr<RewindDebugger::FRewindDebuggerTrack> FVisualLogTrackCreator::CreateTrackInternal(uint64 ObjectId) const
+TSharedPtr<RewindDebugger::FRewindDebuggerTrack> FVisualLogTrackCreator::CreateTrackInternal(const FObjectId& InObjectId) const
 {
-	return MakeShared<RewindDebugger::FVisualLogTrack>(ObjectId);
+	return MakeShared<RewindDebugger::FVisualLogTrack>(InObjectId.GetMainId());
 }
 
 
@@ -299,7 +299,7 @@ void FVisualLogTrack::IterateSubTracksInternal(TFunction<void(TSharedPtr<FRewind
 	}
 };
 
-bool FVisualLogTrackCreator::HasDebugInfoInternal(uint64 ObjectId) const
+bool FVisualLogTrackCreator::HasDebugInfoInternal(const FObjectId& InObjectId) const
 {
 	TRACE_CPUPROFILER_EVENT_SCOPE(FVisualLogTrack::HasDebugInfoInternal);
 	const TraceServices::IAnalysisSession* AnalysisSession = IRewindDebugger::Instance()->GetAnalysisSession();
@@ -308,7 +308,7 @@ bool FVisualLogTrackCreator::HasDebugInfoInternal(uint64 ObjectId) const
 	bool bHasData = false;
 	if (const FVisualLoggerProvider* VLogProvider = AnalysisSession->ReadProvider<FVisualLoggerProvider>(FVisualLoggerProvider::ProviderName))
 	{
-		VLogProvider->ReadVisualLogEntryTimeline(ObjectId, [&bHasData](const FVisualLoggerProvider::VisualLogEntryTimeline& InTimeline)
+		VLogProvider->ReadVisualLogEntryTimeline(InObjectId.GetMainId(), [&bHasData](const FVisualLoggerProvider::VisualLogEntryTimeline& InTimeline)
 			{
 				bHasData = true;
 			});
diff --git a/Engine/Plugins/Animation/GameplayInsights/Source/RewindDebuggerVLog/Private/VisualLogTrack.h b/Engine/Plugins/Animation/GameplayInsights/Source/RewindDebuggerVLog/Private/VisualLogTrack.h
--- a/Engine/Plugins/Animation/GameplayInsights/Source/RewindDebuggerVLog/Private/VisualLogTrack.h
+++ b/Engine/Plugins/Animation/GameplayInsights/Source/RewindDebuggerVLog/Private/VisualLogTrack.h
@@ -82,8 +82,8 @@ private:
 	virtual FName GetTargetTypeNameInternal() const override;
 	virtual FName GetNameInternal() const override;
 	virtual void GetTrackTypesInternal(TArray<FRewindDebuggerTrackType>& Types) const override;
-	virtual TSharedPtr<FRewindDebuggerTrack> CreateTrackInternal(uint64 ObjectId) const override;
-	virtual bool HasDebugInfoInternal(uint64 ObjectId) const override;
+	virtual TSharedPtr<FRewindDebuggerTrack> CreateTrackInternal(const FObjectId& InObjectId) const override;
+	virtual bool HasDebugInfoInternal(const FObjectId& InObjectId) const override;
 };
 
 }
\ No newline at end of file
diff --git a/Engine/Plugins/Animation/PoseSearch/Source/Editor/Private/PoseSearchDebugger.cpp b/Engine/Plugins/Animation/PoseSearch/Source/Editor/Private/PoseSearchDebugger.cpp
--- a/Engine/Plugins/Animation/PoseSearch/Source/Editor/Private/PoseSearchDebugger.cpp
+++ b/Engine/Plugins/Animation/PoseSearch/Source/Editor/Private/PoseSearchDebugger.cpp
@@ -567,12 +567,12 @@ void FDebuggerTrackCreator::GetTrackTypesInternal(TArray<RewindDebugger::FRewind
 	Types.Add({ GetNameInternal(), LOCTEXT("Pose Search", "Pose Search") });
 }
 
-TSharedPtr<RewindDebugger::FRewindDebuggerTrack> FDebuggerTrackCreator::CreateTrackInternal(uint64 ObjectId) const
+TSharedPtr<RewindDebugger::FRewindDebuggerTrack> FDebuggerTrackCreator::CreateTrackInternal(const RewindDebugger::FObjectId& InObjectId) const
 {
-	return MakeShared<FDebuggerTrack>(ObjectId);
+	return MakeShared<FDebuggerTrack>(InObjectId.GetMainId());
 }
 
-bool FDebuggerTrackCreator::HasDebugInfoInternal(uint64 ObjectId) const
+bool FDebuggerTrackCreator::HasDebugInfoInternal(const RewindDebugger::FObjectId& InObjectId) const
 {
 	TRACE_CPUPROFILER_EVENT_SCOPE(PoseSearchDebugger::HasDebugInfoInternal);
 	// Get provider and validate
@@ -589,7 +589,7 @@ bool FDebuggerTrackCreator::HasDebugInfoInternal(uint64 ObjectId) const
 
 	bool bHasData = false;
 
-	PoseSearchProvider->EnumerateMotionMatchingStateTimelines(ObjectId, [&bHasData](const FTraceProvider::FMotionMatchingStateTimeline& InTimeline)
+	PoseSearchProvider->EnumerateMotionMatchingStateTimelines(InObjectId.GetMainId(), [&bHasData](const FTraceProvider::FMotionMatchingStateTimeline& InTimeline)
 	{
 		bHasData = true;
 	});
diff --git a/Engine/Plugins/Animation/PoseSearch/Source/Editor/Private/PoseSearchDebugger.h b/Engine/Plugins/Animation/PoseSearch/Source/Editor/Private/PoseSearchDebugger.h
--- a/Engine/Plugins/Animation/PoseSearch/Source/Editor/Private/PoseSearchDebugger.h
+++ b/Engine/Plugins/Animation/PoseSearch/Source/Editor/Private/PoseSearchDebugger.h
@@ -115,8 +115,8 @@ private:
 	virtual FName GetTargetTypeNameInternal() const override { return "AnimInstance"; }
 	virtual FName GetNameInternal() const override { return "PoseSearchDebugger"; }
 	virtual void GetTrackTypesInternal(TArray<RewindDebugger::FRewindDebuggerTrackType>& Types) const override;
-	virtual TSharedPtr<RewindDebugger::FRewindDebuggerTrack> CreateTrackInternal(uint64 ObjectId) const override;
-	virtual bool HasDebugInfoInternal(uint64 ObjectId) const override;
+	virtual TSharedPtr<RewindDebugger::FRewindDebuggerTrack> CreateTrackInternal(const RewindDebugger::FObjectId& InObjectId) const override;
+	virtual bool HasDebugInfoInternal(const RewindDebugger::FObjectId& InObjectId) const override;
 	virtual int32 GetSortOrderPriorityInternal() const override { return 10; };
 };
 
diff --git a/Engine/Plugins/Animation/PoseSearch/Source/Editor/Private/PoseSearchDebuggerView.cpp b/Engine/Plugins/Animation/PoseSearch/Source/Editor/Private/PoseSearchDebuggerView.cpp
--- a/Engine/Plugins/Animation/PoseSearch/Source/Editor/Private/PoseSearchDebuggerView.cpp
+++ b/Engine/Plugins/Animation/PoseSearch/Source/Editor/Private/PoseSearchDebuggerView.cpp
@@ -72,7 +72,7 @@ FText GenerateSearchName(const FTraceMotionMatchingStateMessage& MotionMatchingS
 					StringBuilder.Append(": ");
 
 					const FObjectInfo& SkeletalMeshComponentObjectInfo = GameplayProvider->GetObjectInfo(MotionMatchingState.SkeletalMeshComponentIds[RoleIndex]);
-					const FObjectInfo& SkeletalMeshComponentOuterObjectInfo = GameplayProvider->GetObjectInfo(SkeletalMeshComponentObjectInfo.OuterId);
+					const FObjectInfo& SkeletalMeshComponentOuterObjectInfo = GameplayProvider->GetObjectInfo(SkeletalMeshComponentObjectInfo.GetOuterId());
 
 					StringBuilder.Append(SkeletalMeshComponentOuterObjectInfo.Name);
 
diff --git a/Engine/Plugins/Cameras/GameplayCameras/Source/GameplayCamerasEditor/Private/Trace/CameraSystemRewindDebuggerTrack.cpp b/Engine/Plugins/Cameras/GameplayCameras/Source/GameplayCamerasEditor/Private/Trace/CameraSystemRewindDebuggerTrack.cpp
--- a/Engine/Plugins/Cameras/GameplayCameras/Source/GameplayCamerasEditor/Private/Trace/CameraSystemRewindDebuggerTrack.cpp
+++ b/Engine/Plugins/Cameras/GameplayCameras/Source/GameplayCamerasEditor/Private/Trace/CameraSystemRewindDebuggerTrack.cpp
@@ -156,7 +156,7 @@ void FCameraSystemRewindDebuggerTrackCreator::GetTrackTypesInternal(TArray<Rewin
 	Types.Add({GetNameInternal(), LOCTEXT("DisplayName", "Gameplay Camera System")});
 }
 
-TSharedPtr<RewindDebugger::FRewindDebuggerTrack> FCameraSystemRewindDebuggerTrackCreator::CreateTrackInternal(uint64 ObjectId) const
+TSharedPtr<RewindDebugger::FRewindDebuggerTrack> FCameraSystemRewindDebuggerTrackCreator::CreateTrackInternal(const RewindDebugger::FObjectId& InObjectId) const
 {
 	return MakeShared<FCameraSystemRewindDebuggerTrack>();
 }
diff --git a/Engine/Plugins/Cameras/GameplayCameras/Source/GameplayCamerasEditor/Private/Trace/CameraSystemRewindDebuggerTrack.h b/Engine/Plugins/Cameras/GameplayCameras/Source/GameplayCamerasEditor/Private/Trace/CameraSystemRewindDebuggerTrack.h
--- a/Engine/Plugins/Cameras/GameplayCameras/Source/GameplayCamerasEditor/Private/Trace/CameraSystemRewindDebuggerTrack.h
+++ b/Engine/Plugins/Cameras/GameplayCameras/Source/GameplayCamerasEditor/Private/Trace/CameraSystemRewindDebuggerTrack.h
@@ -62,7 +62,7 @@ protected:
 	virtual FName GetTargetTypeNameInternal() const override;
 	virtual FName GetNameInternal() const override;
 	virtual void GetTrackTypesInternal(TArray<RewindDebugger::FRewindDebuggerTrackType>& Types) const override;
-	virtual TSharedPtr<RewindDebugger::FRewindDebuggerTrack> CreateTrackInternal(uint64 ObjectId) const override;
+	virtual TSharedPtr<RewindDebugger::FRewindDebuggerTrack> CreateTrackInternal(const RewindDebugger::FObjectId& InObjectId) const override;
 };
 
 }  // namespace UE::Cameras
diff --git a/Engine/Plugins/Chooser/Source/ChooserEditor/Private/ChooserTrack.cpp b/Engine/Plugins/Chooser/Source/ChooserEditor/Private/ChooserTrack.cpp
--- a/Engine/Plugins/Chooser/Source/ChooserEditor/Private/ChooserTrack.cpp
+++ b/Engine/Plugins/Chooser/Source/ChooserEditor/Private/ChooserTrack.cpp
@@ -178,9 +178,9 @@ void FChoosersTrackCreator::GetTrackTypesInternal(TArray<RewindDebugger::FRewind
 	Types.Add({ChoosersName, LOCTEXT("Chooser", "Choosers")});
 }
 
-TSharedPtr<RewindDebugger::FRewindDebuggerTrack> FChoosersTrackCreator::CreateTrackInternal(uint64 ObjectId) const
+TSharedPtr<RewindDebugger::FRewindDebuggerTrack> FChoosersTrackCreator::CreateTrackInternal(const RewindDebugger::FObjectId& InObjectId) const
 {
-	return MakeShared<FChoosersTrack>(ObjectId);
+	return MakeShared<FChoosersTrack>(InObjectId.GetMainId());
 }
 
 		
@@ -253,7 +253,7 @@ void FChoosersTrack::IterateSubTracksInternal(TFunction<void(TSharedPtr<FRewindD
 	}
 };
 
-bool FChoosersTrackCreator::HasDebugInfoInternal(uint64 ObjectId) const
+bool FChoosersTrackCreator::HasDebugInfoInternal(const RewindDebugger::FObjectId& InObjectId) const
 {
 	TRACE_CPUPROFILER_EVENT_SCOPE(FChoosersTrack::HasDebugInfoInternal);
 	const TraceServices::IAnalysisSession* AnalysisSession = IRewindDebugger::Instance()->GetAnalysisSession();
@@ -262,7 +262,7 @@ bool FChoosersTrackCreator::HasDebugInfoInternal(uint64 ObjectId) const
 	bool bHasData = false;
 	if (const FChooserProvider* ChooserProvider = AnalysisSession->ReadProvider<FChooserProvider>(FChooserProvider::ProviderName))
 	{
-		ChooserProvider->ReadChooserEvaluationTimeline(ObjectId, [&bHasData](const FChooserProvider::ChooserEvaluationTimeline& InTimeline)
+		ChooserProvider->ReadChooserEvaluationTimeline(InObjectId.GetMainId(), [&bHasData](const FChooserProvider::ChooserEvaluationTimeline& InTimeline)
 	 	{
 	 		bHasData = true;
 	 	});
diff --git a/Engine/Plugins/Chooser/Source/ChooserEditor/Private/ChooserTrack.h b/Engine/Plugins/Chooser/Source/ChooserEditor/Private/ChooserTrack.h
--- a/Engine/Plugins/Chooser/Source/ChooserEditor/Private/ChooserTrack.h
+++ b/Engine/Plugins/Chooser/Source/ChooserEditor/Private/ChooserTrack.h
@@ -70,8 +70,8 @@ private:
 	virtual FName GetTargetTypeNameInternal() const;
 	virtual FName GetNameInternal() const override;
 	virtual void GetTrackTypesInternal(TArray<RewindDebugger::FRewindDebuggerTrackType>& Types) const override;
-	virtual TSharedPtr<RewindDebugger::FRewindDebuggerTrack> CreateTrackInternal(uint64 ObjectId) const override;
-	virtual bool HasDebugInfoInternal(uint64 ObjectId) const override;
+	virtual TSharedPtr<RewindDebugger::FRewindDebuggerTrack> CreateTrackInternal(const RewindDebugger::FObjectId& InObjectId) const override;
+	virtual bool HasDebugInfoInternal(const RewindDebugger::FObjectId& InObjectId) const override;
 };
 
 }
\ No newline at end of file
diff --git a/Engine/Plugins/Developer/FieldNotificationTrace/Source/FieldNotificationTraceEditor/Private/FieldNotificationTrack.cpp b/Engine/Plugins/Developer/FieldNotificationTrace/Source/FieldNotificationTraceEditor/Private/FieldNotificationTrack.cpp
--- a/Engine/Plugins/Developer/FieldNotificationTrace/Source/FieldNotificationTraceEditor/Private/FieldNotificationTrack.cpp
+++ b/Engine/Plugins/Developer/FieldNotificationTrace/Source/FieldNotificationTraceEditor/Private/FieldNotificationTrack.cpp
@@ -247,12 +247,12 @@ void FTracksCreator::GetTrackTypesInternal(TArray<RewindDebugger::FRewindDebugge
 	Types.Add({ FieldNotificationName, LOCTEXT("FieldNotification", "Field Notification") });
 }
 
-TSharedPtr<RewindDebugger::FRewindDebuggerTrack> FTracksCreator::CreateTrackInternal(uint64 ObjectId) const
+TSharedPtr<RewindDebugger::FRewindDebuggerTrack> FTracksCreator::CreateTrackInternal(const RewindDebugger::FObjectId& InObjectId) const
 {
-	return MakeShared<FObjectTrack>(ObjectId);
+	return MakeShared<FObjectTrack>(InObjectId.GetMainId());
 }
 
-bool FTracksCreator::HasDebugInfoInternal(uint64 ObjectId) const
+bool FTracksCreator::HasDebugInfoInternal(const RewindDebugger::FObjectId& InObjectId) const
 {
 	TRACE_CPUPROFILER_EVENT_SCOPE(FTracks::HasDebugInfoInternal);
 	const TraceServices::IAnalysisSession* AnalysisSession = IRewindDebugger::Instance()->GetAnalysisSession();
@@ -261,7 +261,7 @@ bool FTracksCreator::HasDebugInfoInternal(uint64 ObjectId) const
 	bool bHasData = false;
 	if (const FTraceProvider* Provider = AnalysisSession->ReadProvider<FTraceProvider>(FTraceProvider::ProviderName))
 	{
-		bHasData = Provider->HasData(ObjectId);
+		bHasData = Provider->HasData(InObjectId.GetMainId());
 	}
 	return bHasData;
 }
diff --git a/Engine/Plugins/Developer/FieldNotificationTrace/Source/FieldNotificationTraceEditor/Private/FieldNotificationTrack.h b/Engine/Plugins/Developer/FieldNotificationTrace/Source/FieldNotificationTraceEditor/Private/FieldNotificationTrack.h
--- a/Engine/Plugins/Developer/FieldNotificationTrace/Source/FieldNotificationTraceEditor/Private/FieldNotificationTrack.h
+++ b/Engine/Plugins/Developer/FieldNotificationTrace/Source/FieldNotificationTraceEditor/Private/FieldNotificationTrack.h
@@ -70,8 +70,8 @@ private:
 	virtual FName GetTargetTypeNameInternal() const;
 	virtual FName GetNameInternal() const override;
 	virtual void GetTrackTypesInternal(TArray<RewindDebugger::FRewindDebuggerTrackType>& Types) const override;
-	virtual TSharedPtr<RewindDebugger::FRewindDebuggerTrack> CreateTrackInternal(uint64 ObjectId) const override;
-	virtual bool HasDebugInfoInternal(uint64 ObjectId) const override;
+	virtual TSharedPtr<RewindDebugger::FRewindDebuggerTrack> CreateTrackInternal(const RewindDebugger::FObjectId& InObjectId) const override;
+	virtual bool HasDebugInfoInternal(const RewindDebugger::FObjectId& InObjectId) const override;
 };
 
 } // namespace
diff --git a/Engine/Plugins/Experimental/AnimNext/Source/AnimNextEditor/Private/RewindDebugger/AnimNextModuleTrack.cpp b/Engine/Plugins/Experimental/AnimNext/Source/AnimNextEditor/Private/RewindDebugger/AnimNextModuleTrack.cpp
--- a/Engine/Plugins/Experimental/AnimNext/Source/AnimNextEditor/Private/RewindDebugger/AnimNextModuleTrack.cpp
+++ b/Engine/Plugins/Experimental/AnimNext/Source/AnimNextEditor/Private/RewindDebugger/AnimNextModuleTrack.cpp
@@ -58,9 +58,9 @@ void FAnimNextModuleTrackCreator::GetTrackTypesInternal(TArray<RewindDebugger::F
 	Types.Add({AnimNextModulesName, LOCTEXT("AnimNextModule", "AnimNextModules")});
 }
 
-TSharedPtr<RewindDebugger::FRewindDebuggerTrack> FAnimNextModuleTrackCreator::CreateTrackInternal(uint64 ObjectId) const
+TSharedPtr<RewindDebugger::FRewindDebuggerTrack> FAnimNextModuleTrackCreator::CreateTrackInternal(const RewindDebugger::FObjectId& ObjectId) const
 {
-	return MakeShared<FAnimNextModuleTrack>(ObjectId);
+	return MakeShared<FAnimNextModuleTrack>(ObjectId.GetMainId());
 }
 
 		
@@ -263,7 +263,7 @@ bool FAnimNextModuleTrack::HandleDoubleClickInternal()
 }
 	
 	
-bool FAnimNextModuleTrackCreator::HasDebugInfoInternal(uint64 ObjectId) const
+bool FAnimNextModuleTrackCreator::HasDebugInfoInternal(const RewindDebugger::FObjectId& ObjectId) const
 {
 	TRACE_CPUPROFILER_EVENT_SCOPE(FAnimNextModuleTrack::HasDebugInfoInternal);
 
@@ -275,7 +275,7 @@ bool FAnimNextModuleTrackCreator::HasDebugInfoInternal(uint64 ObjectId) const
 		const FAnimNextProvider* AnimNextProvider = IRewindDebugger::Instance()->GetAnalysisSession()->ReadProvider<FAnimNextProvider>("AnimNextProvider");
     
 		uint64 ModuleId;
-		return AnimNextProvider->GetModuleId(ObjectId, ModuleId);
+		return AnimNextProvider->GetModuleId(ObjectId.GetChildElementId(), ModuleId);
 	}
 
 	return false;
diff --git a/Engine/Plugins/Experimental/AnimNext/Source/AnimNextEditor/Private/RewindDebugger/AnimNextModuleTrack.h b/Engine/Plugins/Experimental/AnimNext/Source/AnimNextEditor/Private/RewindDebugger/AnimNextModuleTrack.h
--- a/Engine/Plugins/Experimental/AnimNext/Source/AnimNextEditor/Private/RewindDebugger/AnimNextModuleTrack.h
+++ b/Engine/Plugins/Experimental/AnimNext/Source/AnimNextEditor/Private/RewindDebugger/AnimNextModuleTrack.h
@@ -71,8 +71,8 @@ private:
 	virtual FName GetTargetTypeNameInternal() const;
 	virtual FName GetNameInternal() const override { return "AnimNextModule"; }
 	virtual void GetTrackTypesInternal(TArray<RewindDebugger::FRewindDebuggerTrackType>& Types) const override;
-	virtual TSharedPtr<RewindDebugger::FRewindDebuggerTrack> CreateTrackInternal(uint64 ObjectId) const override;
-	virtual bool HasDebugInfoInternal(uint64 ObjectId) const override;
+	virtual TSharedPtr<RewindDebugger::FRewindDebuggerTrack> CreateTrackInternal(const RewindDebugger::FObjectId& InObjectId) const override;
+	virtual bool HasDebugInfoInternal(const RewindDebugger::FObjectId& InObjectId) const override;
 };
 
 }
diff --git a/Engine/Plugins/Experimental/AnimNextAnimGraph/Source/AnimNextAnimGraphEditor/Private/RewindDebugger/EvaluationProgramTrack.cpp b/Engine/Plugins/Experimental/AnimNextAnimGraph/Source/AnimNextAnimGraphEditor/Private/RewindDebugger/EvaluationProgramTrack.cpp
--- a/Engine/Plugins/Experimental/AnimNextAnimGraph/Source/AnimNextAnimGraphEditor/Private/RewindDebugger/EvaluationProgramTrack.cpp
+++ b/Engine/Plugins/Experimental/AnimNextAnimGraph/Source/AnimNextAnimGraphEditor/Private/RewindDebugger/EvaluationProgramTrack.cpp
@@ -37,9 +37,9 @@ void FEvaluationProgramTrackCreator::GetTrackTypesInternal(TArray<RewindDebugger
 	Types.Add({AnimNextModulesName, LOCTEXT("AnimNextModule", "AnimNextModules")});
 }
 
-TSharedPtr<RewindDebugger::FRewindDebuggerTrack> FEvaluationProgramTrackCreator::CreateTrackInternal(uint64 ObjectId) const
+TSharedPtr<RewindDebugger::FRewindDebuggerTrack> FEvaluationProgramTrackCreator::CreateTrackInternal(const RewindDebugger::FObjectId& InObjectId) const
 {
-	return MakeShared<FEvaluationProgramTrack>(ObjectId);
+	return MakeShared<FEvaluationProgramTrack>(InObjectId.GetMainId());
 }
 
 		
@@ -177,7 +177,7 @@ bool FEvaluationProgramTrack::UpdateInternal()
  	return bChanged;
 }
 
-bool FEvaluationProgramTrackCreator::HasDebugInfoInternal(uint64 ObjectId) const
+bool FEvaluationProgramTrackCreator::HasDebugInfoInternal(const RewindDebugger::FObjectId& InObjectId) const
 {
 	// TRACE_CPUPROFILER_EVENT_SCOPE(FEvaluationProgramTrack::HasDebugInfoInternal);
 
diff --git a/Engine/Plugins/Experimental/AnimNextAnimGraph/Source/AnimNextAnimGraphEditor/Private/RewindDebugger/EvaluationProgramTrack.h b/Engine/Plugins/Experimental/AnimNextAnimGraph/Source/AnimNextAnimGraphEditor/Private/RewindDebugger/EvaluationProgramTrack.h
--- a/Engine/Plugins/Experimental/AnimNextAnimGraph/Source/AnimNextAnimGraphEditor/Private/RewindDebugger/EvaluationProgramTrack.h
+++ b/Engine/Plugins/Experimental/AnimNextAnimGraph/Source/AnimNextAnimGraphEditor/Private/RewindDebugger/EvaluationProgramTrack.h
@@ -62,8 +62,8 @@ private:
 	virtual FName GetTargetTypeNameInternal() const;
 	virtual FName GetNameInternal() const override { return "AnimNextModule"; }
 	virtual void GetTrackTypesInternal(TArray<RewindDebugger::FRewindDebuggerTrackType>& Types) const override;
-	virtual TSharedPtr<RewindDebugger::FRewindDebuggerTrack> CreateTrackInternal(uint64 ObjectId) const override;
-	virtual bool HasDebugInfoInternal(uint64 ObjectId) const override;
+	virtual TSharedPtr<RewindDebugger::FRewindDebuggerTrack> CreateTrackInternal(const RewindDebugger::FObjectId& InObjectId) const override;
+	virtual bool HasDebugInfoInternal(const RewindDebugger::FObjectId& InObjectId) const override;
 };
 
 }
\ No newline at end of file
diff --git a/Engine/Plugins/Experimental/AnimNextAnimGraph/Source/AnimNextAnimGraphEditor/Private/RewindDebugger/SequenceInfoTrack.cpp b/Engine/Plugins/Experimental/AnimNextAnimGraph/Source/AnimNextAnimGraphEditor/Private/RewindDebugger/SequenceInfoTrack.cpp
--- a/Engine/Plugins/Experimental/AnimNextAnimGraph/Source/AnimNextAnimGraphEditor/Private/RewindDebugger/SequenceInfoTrack.cpp
+++ b/Engine/Plugins/Experimental/AnimNextAnimGraph/Source/AnimNextAnimGraphEditor/Private/RewindDebugger/SequenceInfoTrack.cpp
@@ -37,9 +37,9 @@ void FSequenceInfoTrackCreator::GetTrackTypesInternal(TArray<RewindDebugger::FRe
 	Types.Add({FSequenceInfoTrack::TrackName, LOCTEXT("AnimNextSequenceInfo", "AnimNextSequenceInfo")});
 }
 
-TSharedPtr<RewindDebugger::FRewindDebuggerTrack> FSequenceInfoTrackCreator::CreateTrackInternal(uint64 ObjectId) const
+TSharedPtr<RewindDebugger::FRewindDebuggerTrack> FSequenceInfoTrackCreator::CreateTrackInternal(const RewindDebugger::FObjectId& InObjectId) const
 {
-	return MakeShared<FSequenceInfoTrack>(ObjectId);
+	return MakeShared<FSequenceInfoTrack>(InObjectId.GetMainId());
 }
 
 
@@ -338,7 +338,7 @@ bool FSequenceInfoTrack::UpdateInternal()
 	return bChanged;
 }
 
-bool FSequenceInfoTrackCreator::HasDebugInfoInternal(uint64 ObjectId) const
+bool FSequenceInfoTrackCreator::HasDebugInfoInternal(const RewindDebugger::FObjectId& InObjectId) const
 {
 	const IRewindDebugger* RewindDebugger = IRewindDebugger::Instance();
 	if (const TraceServices::IAnalysisSession* AnalysisSession = RewindDebugger->GetAnalysisSession())
diff --git a/Engine/Plugins/Experimental/AnimNextAnimGraph/Source/AnimNextAnimGraphEditor/Private/RewindDebugger/SequenceInfoTrack.h b/Engine/Plugins/Experimental/AnimNextAnimGraph/Source/AnimNextAnimGraphEditor/Private/RewindDebugger/SequenceInfoTrack.h
--- a/Engine/Plugins/Experimental/AnimNextAnimGraph/Source/AnimNextAnimGraphEditor/Private/RewindDebugger/SequenceInfoTrack.h
+++ b/Engine/Plugins/Experimental/AnimNextAnimGraph/Source/AnimNextAnimGraphEditor/Private/RewindDebugger/SequenceInfoTrack.h
@@ -114,8 +114,8 @@ private:
 	virtual FName GetTargetTypeNameInternal() const;
 	virtual FName GetNameInternal() const override { return FSequenceInfoTrack::TrackName; }
 	virtual void GetTrackTypesInternal(TArray<RewindDebugger::FRewindDebuggerTrackType>& Types) const override;
-	virtual TSharedPtr<RewindDebugger::FRewindDebuggerTrack> CreateTrackInternal(uint64 ObjectId) const override;
-	virtual bool HasDebugInfoInternal(uint64 ObjectId) const override;
+	virtual TSharedPtr<RewindDebugger::FRewindDebuggerTrack> CreateTrackInternal(const RewindDebugger::FObjectId& InObjectId) const override;
+	virtual bool HasDebugInfoInternal(const RewindDebugger::FObjectId& InObjectId) const override;
 };
 
 }
\ No newline at end of file
diff --git a/Engine/Plugins/Experimental/Animation/TrajectoryTools/Source/TrajectoryTools/Private/TrajectoryRewindDebuggerExtension.cpp b/Engine/Plugins/Experimental/Animation/TrajectoryTools/Source/TrajectoryTools/Private/TrajectoryRewindDebuggerExtension.cpp
--- a/Engine/Plugins/Experimental/Animation/TrajectoryTools/Source/TrajectoryTools/Private/TrajectoryRewindDebuggerExtension.cpp
+++ b/Engine/Plugins/Experimental/Animation/TrajectoryTools/Source/TrajectoryTools/Private/TrajectoryRewindDebuggerExtension.cpp
@@ -351,9 +351,9 @@ void FRewindDebuggerTrajectory::BuildTrajectoryOwnersList(const IRewindDebugger*
 	InGameplayProvider->EnumerateObjects(InTraceStartTime, InTraceEndTime, [&OutObjectInfos,InGameplayProvider, SkelMeshCmpClassInfo, InRewindDebugger](const FObjectInfo& InObjectInfo)
 	{
 		// For all root objects find their skeletal mesh component's id, if any.
-		if (InObjectInfo.OuterId == 0)
+		if (!InObjectInfo.GetOuterId().IsSet())
 		{
-			InGameplayProvider->EnumerateSubobjects(InObjectInfo.Id, [&OutObjectInfos, InGameplayProvider, SkelMeshCmpClassInfo](uint64 InSubObjectId)
+			InGameplayProvider->EnumerateSubobjects(InObjectInfo.GetId(), [&OutObjectInfos, InGameplayProvider, SkelMeshCmpClassInfo](const RewindDebugger::FObjectId& InSubObjectId)
 			{
 				const FObjectInfo& SubObjectInfo = InGameplayProvider->GetObjectInfo(InSubObjectId);
 
@@ -369,7 +369,7 @@ void FRewindDebuggerTrajectory::BuildTrajectoryOwnersList(const IRewindDebugger*
 			const bool bIsObjectASkelMesh = InObjectInfo.ClassId == SkelMeshCmpClassInfo->Id || InGameplayProvider->IsSubClassOf(InObjectInfo.ClassId, SkelMeshCmpClassInfo->Id);
 			if (bIsObjectASkelMesh)
 			{
-				const bool bHasOwningActor = InRewindDebugger->FindOwningActorInfo(InGameplayProvider, InObjectInfo.Id) != nullptr;
+				const bool bHasOwningActor = InRewindDebugger->FindOwningActorInfo(InGameplayProvider, InObjectInfo.GetUObjectId()) != nullptr;
 				if (bHasOwningActor)
 				{
 					OutObjectInfos.Push(InObjectInfo);
@@ -388,7 +388,7 @@ void FRewindDebuggerTrajectory::BuildTrajectorySkeletalMeshInfoList(const IGamep
 	{
 		const FObjectInfo& SkeletalMeshComponentInfo = InSkeletalMeshComponentObjectInfos[SkelMeshCmpIndex];
 		
-		InAnimationProvider->ReadSkeletalMeshPoseTimeline(SkeletalMeshComponentInfo.Id, [InAnimationProvider, InGameplayProvider, InTraceStartTime, InTraceEndTime, SkelMeshCmpIndex, &SkeletalMeshComponentInfo, &OutSkeletalMeshInfos](const IAnimationProvider::SkeletalMeshPoseTimeline& TimelineData, bool bHasCurves)
+		InAnimationProvider->ReadSkeletalMeshPoseTimeline(SkeletalMeshComponentInfo.GetUObjectId(), [InAnimationProvider, InGameplayProvider, InTraceStartTime, InTraceEndTime, SkelMeshCmpIndex, &SkeletalMeshComponentInfo, &OutSkeletalMeshInfos](const IAnimationProvider::SkeletalMeshPoseTimeline& TimelineData, bool bHasCurves)
 		{
 			TimelineData.EnumerateEvents(InTraceStartTime, InTraceEndTime, [InAnimationProvider, InGameplayProvider, SkelMeshCmpIndex, &SkeletalMeshComponentInfo, &OutSkeletalMeshInfos](double InStartTime, double InEndTime, uint32 InDepth, const FSkeletalMeshPoseMessage& InPoseMessage) -> TraceServices::EEventEnumerate
 			{
@@ -465,7 +465,7 @@ void FRewindDebuggerTrajectory::BuildTrajectories(const IGameplayProvider* InGam
 		// BoneCount split tracker.
 		int64 PrevSampleBoneCount = INDEX_NONE;
 		
-		InAnimationProvider->ReadSkeletalMeshPoseTimeline(ObjectInfo.Id, [InGameplayProvider, InAnimationProvider, InTraceStartTime, InTraceEndTime, &OutTrajectories, &ObjectInfo, &PrevSampleBoneCount](const IAnimationProvider::SkeletalMeshPoseTimeline& InTimeline, bool bHasCurves)
+		InAnimationProvider->ReadSkeletalMeshPoseTimeline(ObjectInfo.GetUObjectId(), [InGameplayProvider, InAnimationProvider, InTraceStartTime, InTraceEndTime, &OutTrajectories, &ObjectInfo, &PrevSampleBoneCount](const IAnimationProvider::SkeletalMeshPoseTimeline& InTimeline, bool bHasCurves)
 		{
 			InTimeline.EnumerateEvents(InTraceStartTime, InTraceEndTime, [InGameplayProvider, InAnimationProvider, &OutTrajectories, &ObjectInfo, &PrevSampleBoneCount](double InStartTime, double InEndTime, uint32 InDepth, const FSkeletalMeshPoseMessage& InPoseMessage)
 			{
@@ -554,7 +554,7 @@ void FRewindDebuggerTrajectory::UpdateDebugInfos(const IGameplayProvider* InGame
 	{
 		const FExtensionState::FDebugInfo* MatchedDebugInfo = PrevDebugInfos.FindByPredicate([&ObjectInfo](const FExtensionState::FDebugInfo& Item)
 		{
-			return Item.OwnerId == ObjectInfo.Id;
+			return Item.OwnerId == ObjectInfo.GetUObjectId();
 		});
 
 		// Use cached version to keep track of stateful date (i.e. ShouldDraw checkbox)
@@ -564,7 +564,7 @@ void FRewindDebuggerTrajectory::UpdateDebugInfos(const IGameplayProvider* InGame
 		}
 		else
 		{
-			InOutDebugInfos.Push({ObjectInfo.Id, GetFullNameForDebugInfoOwner(InGameplayProvider, ObjectInfo.Id), false});
+			InOutDebugInfos.Push({ObjectInfo.GetUObjectId(), GetFullNameForDebugInfoOwner(InGameplayProvider, ObjectInfo.GetUObjectId()), false});
 		}
 	}
 }
@@ -603,7 +603,7 @@ FName FRewindDebuggerTrajectory::GetFullNameForDebugInfoOwner(const IGameplayPro
 		// We are a component, keep traversing up the tree.
 		else
 		{
-			if (ObjectInfo.OuterId != 0)
+			if (ObjectInfo.GetOuterId().IsSet())
 			{
 				if (ObjectId != InOwnerObjectId)
 				{
@@ -611,7 +611,7 @@ FName FRewindDebuggerTrajectory::GetFullNameForDebugInfoOwner(const IGameplayPro
 				}
 				Result.Prepend(ObjectInfo.Name);
 				
-				ObjectId = ObjectInfo.OuterId;
+				ObjectId = ObjectInfo.GetOuterUObjectId();
 			}
 			else
 			{
diff --git a/Engine/Source/Editor/RewindDebuggerInterface/Public/IRewindDebugger.h b/Engine/Source/Editor/RewindDebuggerInterface/Public/IRewindDebugger.h
--- a/Engine/Source/Editor/RewindDebuggerInterface/Public/IRewindDebugger.h
+++ b/Engine/Source/Editor/RewindDebuggerInterface/Public/IRewindDebugger.h
@@ -4,6 +4,7 @@
 #if UE_ENABLE_INCLUDE_ORDER_DEPRECATED_IN_5_7
 #include "CoreMinimal.h"
 #endif // UE_ENABLE_INCLUDE_ORDER_DEPRECATED_IN_5_7
+#include "RewindDebuggerTypes.h"
 #include "TraceServices/Model/Frames.h"
 #include "IRewindDebugger.generated.h"
 
@@ -21,15 +22,28 @@ struct FObjectInfo;
 struct FDebugObjectInfo
 {
 	FDebugObjectInfo() = default;
-	FDebugObjectInfo(const uint64 Id, const FString& Name): ObjectId(Id), ObjectName(Name), bExpanded(true)
+	FDebugObjectInfo(const RewindDebugger::FObjectId& Id, const FString& Name)
+		: Id(Id)
+		, ObjectName(Name)
+		, bExpanded(true)
 	{
 	}
 
-	uint64 ObjectId = 0;
+	/** @return Part of the Id representing the UObject */
+	uint64 GetUObjectId() const
+	{
+		return Id.GetMainId();
+	}
+
+	RewindDebugger::FObjectId Id;
 	FString ObjectName;
+	TArray<TSharedPtr<FDebugObjectInfo>> Children;
 	bool bExpanded = false;
 
-	TArray<TSharedPtr<FDebugObjectInfo>> Children;
+#if WITH_EDITORONLY_DATA
+	UE_DEPRECATED(5.7, "Use ObjectIdentifier instead")
+	uint64 ObjectId = 0;
+#endif
 };
 
 UCLASS()
diff --git a/Engine/Source/Editor/RewindDebuggerInterface/Public/IRewindDebuggerTrackCreator.h b/Engine/Source/Editor/RewindDebuggerInterface/Public/IRewindDebuggerTrackCreator.h
--- a/Engine/Source/Editor/RewindDebuggerInterface/Public/IRewindDebuggerTrackCreator.h
+++ b/Engine/Source/Editor/RewindDebuggerInterface/Public/IRewindDebuggerTrackCreator.h
@@ -20,7 +20,7 @@ namespace RewindDebugger
 {
 
 class FRewindDebuggerTrack;
-
+struct FObjectId;
 struct FRewindDebuggerTrackType
 {
 	FName Name;
@@ -59,18 +59,34 @@ public:
 	 * Optional filter to prevent tracks from being listed if there is no data associated to the provided object identifier.
 	 * @return Whether the track should be created to represent data associated to the provided object identifier.
 	 */
+	bool HasDebugInfo(const FObjectId& ObjectId) const
+	{
+		return HasDebugInfoInternal(ObjectId);
+	};
+
+	UE_DEPRECATED(5.7, "Use the version taking FObjectId instead")
 	bool HasDebugInfo(uint64 ObjectId) const
 	{
+PRAGMA_DISABLE_DEPRECATION_WARNINGS
 		return HasDebugInfoInternal(ObjectId);
+PRAGMA_ENABLE_DEPRECATION_WARNINGS
 	};
 
 	/**
 	 * Creates a track which will be shown in the timeline view and tree view, as a child track of the Object
 	 * @return the created child track
 	 */
+	TSharedPtr<RewindDebugger::FRewindDebuggerTrack> CreateTrack(const FObjectId& ObjectId) const
+	{
+		return CreateTrackInternal(ObjectId);
+	}
+
+	UE_DEPRECATED(5.7, "Use the version taking FObjectId instead")
 	TSharedPtr<RewindDebugger::FRewindDebuggerTrack> CreateTrack(uint64 ObjectId) const
 	{
+PRAGMA_DISABLE_DEPRECATION_WARNINGS
 		return CreateTrackInternal(ObjectId);
+PRAGMA_ENABLE_DEPRECATION_WARNINGS
 	}
 
 	void GetTrackTypes(TArray<FRewindDebuggerTrackType>& Types) const
@@ -78,6 +94,12 @@ public:
 		return GetTrackTypesInternal(Types);
 	};
 
+	/** @return Whether the track created can also represent non-UObject based children */
+	bool IsChildElementSupported() const
+	{
+		return IsChildElementSupportedInternal();
+	}
+
 private:
 
 	/** get the UObject type name this Creator will create child tracks for */
@@ -104,15 +126,32 @@ private:
 	};
 
 	/** Returns true if this creator's track should be shown for this Object, false if there is no data and they should be hidden. */
-	virtual bool HasDebugInfoInternal(uint64 InObjectId) const
+	virtual bool HasDebugInfoInternal(const FObjectId& InObjectId) const
 	{
 		return true;
 	};
 
-	virtual TSharedPtr<FRewindDebuggerTrack> CreateTrackInternal(uint64 InObjectId) const
+	UE_DEPRECATED(5.7, "Use the version taking FObjectId instead")
+	virtual bool HasDebugInfoInternal(uint64 InObjectId) const final
+	{
+		return true;
+	};
+
+	virtual TSharedPtr<FRewindDebuggerTrack> CreateTrackInternal(const FObjectId& InObjectId) const
 	{
 		return TSharedPtr<FRewindDebuggerTrack>();
 	}
+
+	UE_DEPRECATED(5.7, "Use the version taking FObjectId instead")
+	virtual TSharedPtr<FRewindDebuggerTrack> CreateTrackInternal(uint64 InObjectId) const final
+	{
+		return TSharedPtr<FRewindDebuggerTrack>();
+	}
+
+	virtual bool IsChildElementSupportedInternal() const
+	{
+		return false;
+	}
 };
 
 } // namespace RewindDebugger
diff --git a/Engine/Source/Editor/RewindDebuggerInterface/Public/RewindDebuggerTrack.h b/Engine/Source/Editor/RewindDebuggerInterface/Public/RewindDebuggerTrack.h
--- a/Engine/Source/Editor/RewindDebuggerInterface/Public/RewindDebuggerTrack.h
+++ b/Engine/Source/Editor/RewindDebuggerInterface/Public/RewindDebuggerTrack.h
@@ -111,13 +111,25 @@ public:
 		return GetDisplayNameInternal();
 	}
 
-	/** insights object id for an object associated with this track */
-	uint64 GetObjectId() const
+	/** insights main UObject id for an object associated with this track */
+	uint64 GetUObjectId() const
 	{
 		return GetObjectIdInternal();
 	}
 
-	/** iterate over all sub-tracks of this track and call Iterator function */
+	UE_DEPRECATED(5.7, "Use GetUObjectID instead.")
+	uint64 GetObjectId() const
+	{
+		return GetUObjectId ();
+	}
+
+	/** insights object id for an object associated with this track */
+	FObjectId GetAssociatedObjectId() const
+	{
+		return FObjectId(GetObjectIdInternal(), GetChildElementIdInternal());
+	}
+
+	/** iterates over all sub-tracks of this track and call Iterator function */
 	void IterateSubTracks(TFunction<void(TSharedPtr<FRewindDebuggerTrack> SubTrack)> IteratorFunction)
 	{
 		IterateSubTracksInternal(IteratorFunction);
@@ -185,7 +197,7 @@ private:
 		return {};
 	}
 
-	virtual uint64 GetSubObjectIdInternal() const
+	virtual uint64 GetChildElementIdInternal() const
 	{
 		return INDEX_NONE;
 	}
diff --git a/Engine/Source/Editor/RewindDebuggerInterface/Public/RewindDebuggerTypes.h b/Engine/Source/Editor/RewindDebuggerInterface/Public/RewindDebuggerTypes.h
new file mode 100644
--- /dev/null
+++ b/Engine/Source/Editor/RewindDebuggerInterface/Public/RewindDebuggerTypes.h
@@ -0,0 +1,94 @@
+// Copyright Epic Games, Inc. All Rights Reserved.
+
+#pragma once
+
+#include "Templates/TypeHash.h"
+#include "Containers/UnrealString.h"
+
+namespace RewindDebugger
+{
+
+/**
+ * Struct representing an object (single UObject or a child element of a UObject) tracked by Unreal Insight traces.
+ * This struct is used to represent the data traced by UE::ObjectTrace::FObjectId
+ */
+struct FObjectId
+{
+	static constexpr uint64 InvalidId = 0;
+
+	FObjectId() = default;
+
+	explicit FObjectId(const uint64 MainId, const uint64 ChildElementId = InvalidId)
+		: ObjectId(MainId)
+		, ChildElementId(ChildElementId)
+	{
+	}
+
+	/**
+	 * Identifier associated to a UObject recorded using TRACE_OBJECT* macros.
+	 * @return Identifier associated to a given UObject in traces.
+	 */
+	uint64 GetMainId() const
+	{
+		return ObjectId;
+	}
+
+	/**
+	 * Optional identifier to allow multiple non-UObject based child elements of a given UObject to be represented in the traces.
+	 * @return Identifier of the child element if specified, InvalidId otherwise.
+	 * @see GetObjectId
+	 */
+	uint64 GetChildElementId() const
+	{
+		return ChildElementId;
+	}
+
+	/**
+	 * @return Whether the current gameplay object identifier is used to represent a child element.
+	 */
+	bool IsChildElement() const
+	{
+		return ChildElementId != InvalidId;
+	}
+
+	/**
+	 * @return Whether the current gameplay object identifier was set to represent a UObject, or one of its children.
+	 */
+	bool IsSet() const
+	{
+		return ObjectId != InvalidId;
+	}
+
+	bool operator==(const FObjectId& Other) const = default;
+
+	bool operator<(const FObjectId& Other) const
+	{
+		return ObjectId < Other.ObjectId && ChildElementId < Other.ChildElementId;
+	}
+
+	friend FString LexToString(const FObjectId& ObjectIdentifier)
+	{
+		if (ObjectIdentifier.ObjectId == InvalidId)
+		{
+			return TEXT("Invalid");
+		}
+
+		if (ObjectIdentifier.ChildElementId == InvalidId)
+		{
+			return FString::Printf(TEXT("[%llX]"), ObjectIdentifier.ObjectId);
+		}
+
+		return FString::Printf(TEXT("[%llX : %llX]"), ObjectIdentifier.ObjectId, ObjectIdentifier.ChildElementId);
+	}
+
+	friend uint32 GetTypeHash(const FObjectId Identifier)
+	{
+		return HashCombine(GetTypeHash(Identifier.ObjectId), GetTypeHash(Identifier.ChildElementId));
+	}
+
+private:
+	uint64 ObjectId = InvalidId;
+	uint64 ChildElementId = InvalidId;
+};
+
+}
\ No newline at end of file
diff --git a/Engine/Source/Runtime/Engine/Private/Animation/AnimAttributes.cpp b/Engine/Source/Runtime/Engine/Private/Animation/AnimAttributes.cpp
--- a/Engine/Source/Runtime/Engine/Private/Animation/AnimAttributes.cpp
+++ b/Engine/Source/Runtime/Engine/Private/Animation/AnimAttributes.cpp
@@ -1,7 +1,6 @@
 // Copyright Epic Games, Inc. All Rights Reserved.
 
 #include "Animation/AnimAttributes.h"
-#include "UObject/NameTypes.h"
 
 const FName UE::Anim::FAttributes::Pose("Pose");
 const FName UE::Anim::FAttributes::Curves("Curves");
diff --git a/Engine/Source/Runtime/Engine/Private/ObjectTrace.cpp b/Engine/Source/Runtime/Engine/Private/ObjectTrace.cpp
--- a/Engine/Source/Runtime/Engine/Private/ObjectTrace.cpp
+++ b/Engine/Source/Runtime/Engine/Private/ObjectTrace.cpp
@@ -95,6 +95,79 @@ UE_TRACE_EVENT_BEGIN(Object, View)
 	UE_TRACE_EVENT_FIELD(float, AspectRatio)
 UE_TRACE_EVENT_END()
 
+UE_TRACE_EVENT_BEGIN(Object, ChildElement)
+	UE_TRACE_EVENT_FIELD(uint64, Id)
+	UE_TRACE_EVENT_FIELD(uint64, ChildElementId)
+	UE_TRACE_EVENT_FIELD(uint64, OuterId)
+	UE_TRACE_EVENT_FIELD(uint64, OuterChildElementId)
+	UE_TRACE_EVENT_FIELD(UE::Trace::WideString, Name)
+UE_TRACE_EVENT_END()
+
+UE_TRACE_EVENT_BEGIN(Object, ChildElementLifetimeBegin)
+	UE_TRACE_EVENT_FIELD(uint64, Cycle)
+	UE_TRACE_EVENT_FIELD(double, RecordingTime)
+	UE_TRACE_EVENT_FIELD(uint64, Id)
+	UE_TRACE_EVENT_FIELD(uint64, ChildElementId)
+UE_TRACE_EVENT_END()
+
+UE_TRACE_EVENT_BEGIN(Object, ChildElementLifetimeEnd)
+	UE_TRACE_EVENT_FIELD(uint64, Cycle)
+	UE_TRACE_EVENT_FIELD(double, RecordingTime)
+	UE_TRACE_EVENT_FIELD(uint64, Id)
+	UE_TRACE_EVENT_FIELD(uint64, ChildElementId)
+UE_TRACE_EVENT_END()
+
+UE_TRACE_EVENT_BEGIN(Object, ChildElementEvent)
+	UE_TRACE_EVENT_FIELD(uint64, Cycle)
+	UE_TRACE_EVENT_FIELD(uint64, Id)
+	UE_TRACE_EVENT_FIELD(uint64, ChildElementId)
+	UE_TRACE_EVENT_FIELD(UE::Trace::WideString, Event)
+UE_TRACE_EVENT_END()
+
+namespace UE::ObjectTrace
+{
+bool IsTracingDisabledForObject(const UObject* InObject)
+{
+	if (InObject == nullptr)
+	{
+		return true;
+	}
+
+	if (!UE_TRACE_CHANNELEXPR_IS_ENABLED(ObjectChannel))
+	{
+		return true;
+	}
+
+	if (InObject->HasAnyFlags(RF_ClassDefaultObject))
+	{
+		return true;
+	}
+
+	if (CANNOT_TRACE_OBJECT(InObject->GetWorld()))
+	{
+		return true;
+	}
+
+	return false;
+}
+
+struct FObjectId
+{
+	static constexpr uint64 InvalidId = 0;
+
+	uint64 ObjectId = InvalidId;
+	uint64 ChildElementId = InvalidId;
+
+	bool operator==(const FObjectId& Other) const = default;
+
+	friend uint32 GetTypeHash(const FObjectId Identifier)
+	{
+		return HashCombine(GetTypeHash(Identifier.ObjectId), GetTypeHash(Identifier.ChildElementId));
+	}
+};
+
+}
+
 // Object annotations used for tracing
 struct FObjectIdAnnotation
 {
@@ -128,7 +201,7 @@ int32 GetTypeHash(const FObjectIdAnnotation& Annotation)
 // GObjectTracedSet must only be accessed in the open, and must be protected by the GObjectTracedSetLock.
 // Changing GObjectTracedSet in the closed is likely to cause a validation error.
 FRWLock GObjectTracedSetLock;
-TSet<uint64> GObjectTracedSet;
+TSet<UE::ObjectTrace::FObjectId> GObjectTracedSet;
 FUObjectAnnotationSparseSearchable<FObjectIdAnnotation, true> GObjectIdAnnotations;
 
 // Handle used to hook to world tick
@@ -159,7 +232,7 @@ static void TickObjectTraceWorldSubsystem(UWorld* InWorld, ELevelTick InTickType
 // Returns true if the object ID was added to the traced set.
 // Returns false if the object ID already existed in the set.
 // If the object ID was already in the set, we don't need to log it again.
-UE_AUTORTFM_ALWAYS_OPEN static bool AddObjectToTracedSet(uint64 ObjectId)
+UE_AUTORTFM_ALWAYS_OPEN static bool AddObjectToTracedSet(const UE::ObjectTrace::FObjectId& ObjectId)
 {
 	{
 		FRWScopeLock ScopeLock(GObjectTracedSetLock, SLT_ReadOnly);
@@ -335,7 +408,7 @@ void FObjectTrace::OutputClass(const UClass* InClass)
 		return;
 	}
 	const uint64 ObjectId = GetObjectId(InClass);
-	if (!AddObjectToTracedSet(ObjectId))
+	if (!AddObjectToTracedSet(UE::ObjectTrace::FObjectId{ObjectId}))
 	{
 		return;
 	}
@@ -355,13 +428,7 @@ void FObjectTrace::OutputClass(const UClass* InClass)
 
 void FObjectTrace::OutputView(const UObject* InPlayer, const FSceneView* InView)
 {
-	const bool bChannelEnabled = UE_TRACE_CHANNELEXPR_IS_ENABLED(ObjectChannel);
-	if (!bChannelEnabled || InPlayer == nullptr)
-	{
-		return;
-	}
-
-	if (CANNOT_TRACE_OBJECT(InPlayer->GetWorld()))
+	if (UE::ObjectTrace::IsTracingDisabledForObject(InPlayer))
 	{
 		return;
 	}
@@ -387,24 +454,13 @@ void FObjectTrace::OutputView(const UObject* InPlayer, const FSceneView* InView)
 
 void FObjectTrace::OutputObject(const UObject* InObject)
 {
-	const bool bChannelEnabled = UE_TRACE_CHANNELEXPR_IS_ENABLED(ObjectChannel);
-	if (!bChannelEnabled || InObject == nullptr)
-	{
-		return;
-	}
-
-	if (InObject->HasAnyFlags(RF_ClassDefaultObject))
-	{
-		return;
-	}
-
-	if (CANNOT_TRACE_OBJECT(InObject->GetWorld()))
+	if (UE::ObjectTrace::IsTracingDisabledForObject(InObject))
 	{
 		return;
 	}
 
 	const uint64 ObjectId = GetObjectId(InObject);
-	if (!AddObjectToTracedSet(ObjectId))
+	if (!AddObjectToTracedSet(UE::ObjectTrace::FObjectId{ObjectId}))
 	{
 		return;
 	}
@@ -433,18 +489,7 @@ void FObjectTrace::OutputObject(const UObject* InObject)
 
 void FObjectTrace::OutputObjectTransform(const UObject* InObject, const FTransform& InTransform, ETeleportType InTeleportType)
 {
-	const bool bChannelEnabled = UE_TRACE_CHANNELEXPR_IS_ENABLED(ObjectChannel);
-	if (!bChannelEnabled || InObject == nullptr)
-	{
-		return;
-	}
-
-	if (InObject->HasAnyFlags(RF_ClassDefaultObject))
-	{
-		return;
-	}
-
-	if (CANNOT_TRACE_OBJECT(InObject->GetWorld()))
+	if (UE::ObjectTrace::IsTracingDisabledForObject(InObject))
 	{
 		return;
 	}
@@ -461,18 +506,7 @@ void FObjectTrace::OutputObjectTransform(const UObject* InObject, const FTransfo
 
 void FObjectTrace::OutputObjectEvent(const UObject* InObject, const TCHAR* InEvent)
 {
-	const bool bChannelEnabled = UE_TRACE_CHANNELEXPR_IS_ENABLED(ObjectChannel);
-	if (!bChannelEnabled || InObject == nullptr)
-	{
-		return;
-	}
-
-	if (InObject->HasAnyFlags(RF_ClassDefaultObject))
-	{
-		return;
-	}
-
-	if (CANNOT_TRACE_OBJECT(InObject->GetWorld()))
+	if (UE::ObjectTrace::IsTracingDisabledForObject(InObject))
 	{
 		return;
 	}
@@ -492,18 +526,7 @@ void FObjectTrace::OutputObjectLifetimeBegin(const UObject* InObject)
 
 void FObjectTrace::OutputObjectLifetimeEnd(const UObject* InObject)
 {
-	const bool bChannelEnabled = UE_TRACE_CHANNELEXPR_IS_ENABLED(ObjectChannel);
-	if (!bChannelEnabled || InObject == nullptr)
-	{
-		return;
-	}
-
-	if (InObject->HasAnyFlags(RF_ClassDefaultObject))
-	{
-		return;
-	}
-
-	if (CANNOT_TRACE_OBJECT(InObject->GetWorld()))
+	if (UE::ObjectTrace::IsTracingDisabledForObject(InObject))
 	{
 		return;
 	}
@@ -518,8 +541,7 @@ void FObjectTrace::OutputObjectLifetimeEnd(const UObject* InObject)
 
 void FObjectTrace::OutputPawnPossess(const UObject* InController, const UObject* InPawn)
 {
-	const bool bChannelEnabled = UE_TRACE_CHANNELEXPR_IS_ENABLED(ObjectChannel);
-	if (!bChannelEnabled || InController == nullptr)
+	if (UE::ObjectTrace::IsTracingDisabledForObject(InController))
 	{
 		return;
 	}
@@ -566,4 +588,75 @@ void FObjectTrace::OutputWorld(const UWorld* InWorld)
 	TRACE_OBJECT(InWorld);
 }
 
+void FObjectTrace::OutputChildElement(const UObject* InObject, const uint64 InChildElementId, const UObject* InOuterObject, const uint64 InOuterChildElementId)
+{
+	if (UE::ObjectTrace::IsTracingDisabledForObject(InObject))
+	{
+		return;
+	}
+
+	const uint64 ObjectId = GetObjectId(InObject);
+	if (!AddObjectToTracedSet(UE::ObjectTrace::FObjectId(ObjectId, InChildElementId)))
+	{
+		return;
+	}
+
+	if (InOuterChildElementId == INDEX_NONE)
+	{
+		OutputObject(InOuterObject);
+	}
+	else
+	{
+		OutputChildElement(InOuterObject, InOuterChildElementId);
+	}
+
+	FString ObjectName = InObject->GetName();
+	ObjectName += FString::Printf(TEXT(" : [0x%llX]"), InChildElementId);
+
+	UE_TRACE_LOG(Object, ChildElement, ObjectChannel)
+		<< ChildElement.Id(ObjectId)
+		<< ChildElement.ChildElementId(InChildElementId)
+		<< ChildElement.OuterId(GetObjectId(InOuterObject))
+		<< ChildElement.OuterChildElementId(InOuterChildElementId)
+		<< ChildElement.Name(*ObjectName, ObjectName.Len());
+
+	UE_TRACE_LOG(Object, ChildElementLifetimeBegin, ObjectChannel)
+		<< ChildElementLifetimeBegin.Cycle(FPlatformTime::Cycles64())
+		<< ChildElementLifetimeBegin.RecordingTime(FObjectTrace::GetWorldElapsedTime(InObject->GetWorld()))
+		<< ChildElementLifetimeBegin.Id(ObjectId)
+		<< ChildElementLifetimeBegin.ChildElementId(InChildElementId);
+}
+
+void FObjectTrace::OutputChildElementLifetimeEnd(const UObject* InObject, const uint64 InChildElementId)
+{
+	if (UE::ObjectTrace::IsTracingDisabledForObject(InObject))
+	{
+		return;
+	}
+
+	TRACE_CHILD_ELEMENT(InObject, InChildElementId);
+
+	UE_TRACE_LOG(Object, ChildElementLifetimeEnd, ObjectChannel)
+		<< ChildElementLifetimeEnd.Cycle(FPlatformTime::Cycles64())
+		<< ChildElementLifetimeEnd.RecordingTime(FObjectTrace::GetWorldElapsedTime(InObject->GetWorld()))
+		<< ChildElementLifetimeEnd.Id(GetObjectId(InObject))
+		<< ChildElementLifetimeEnd.ChildElementId(InChildElementId);
+}
+
+void FObjectTrace::OutputChildElementEvent(const UObject* InObject, const uint64 InChildElementId, const TCHAR* InEvent)
+{
+	if (UE::ObjectTrace::IsTracingDisabledForObject(InObject))
+	{
+		return;
+	}
+
+	TRACE_CHILD_ELEMENT(InObject, InChildElementId);
+
+	UE_TRACE_LOG(Object, ChildElementEvent, ObjectChannel)
+		<< ChildElementEvent.Cycle(FPlatformTime::Cycles64())
+		<< ChildElementEvent.Id(GetObjectId(InObject))
+		<< ChildElementEvent.ChildElementId(InChildElementId)
+		<< ChildElementEvent.Event(InEvent);
+}
+
 #endif
diff --git a/Engine/Source/Runtime/Engine/Public/ObjectTrace.h b/Engine/Source/Runtime/Engine/Public/ObjectTrace.h
--- a/Engine/Source/Runtime/Engine/Public/ObjectTrace.h
+++ b/Engine/Source/Runtime/Engine/Public/ObjectTrace.h
@@ -107,6 +107,15 @@ struct FObjectTrace
 
 	/** Helper function to output a world */
 	ENGINE_API static void OutputWorld(const UWorld* InWorld);
+
+	/** Helper function to output a child element of a UObject */
+	ENGINE_API static void OutputChildElement(const UObject* InObject, uint64 InChildElementId, const UObject* InOuterObject = nullptr , uint64 InOuterChildElementId = INDEX_NONE);
+
+	/** Helper function to output destruction event of a child element */
+	ENGINE_API static void OutputChildElementLifetimeEnd(const UObject* InObject, uint64 InChildElementId);
+
+	/** Helper function to output an event for a child element */
+	ENGINE_API static void OutputChildElementEvent(const UObject* InObject, uint64 InChildElementId, const TCHAR* InEvent);
 };
 
 #define TRACE_CLASS(Class) \
@@ -133,22 +142,62 @@ struct FObjectTrace
 	if (CAN_TRACE_OBJECT(Object)) { FObjectTrace::OutputObjectLifetimeEnd(Object); }
 
 #define TRACE_PAWN_POSSESS(Controller, Pawn)\
-	if (CAN_TRACE_OBJECT(Controller) && (Pawn==nullptr || (CAN_TRACE_OBJECT(Pawn)))) FObjectTrace::OutputPawnPossess(Controller, Pawn);
+	if (CAN_TRACE_OBJECT(Controller) && (Pawn==nullptr || (CAN_TRACE_OBJECT(Pawn)))) { FObjectTrace::OutputPawnPossess(Controller, Pawn); }
+
+#define TRACE_CHILD_ELEMENT(Object, ChildElementId) \
+	if (CAN_TRACE_OBJECT(Object)) { FObjectTrace::OutputChildElement(Object, ChildElementId, (Object ? Object->GetOuter() : nullptr), INDEX_NONE); }
+
+#define TRACE_CHILD_ELEMENT_WITH_OUTER(Object, ChildElementId, OuterObject, OuterChildElementId) \
+	if (CAN_TRACE_OBJECT(Object)) { FObjectTrace::OutputChildElement(Object, ChildElementId, OuterObject, OuterChildElementId); }
+
+#define TRACE_CHILD_ELEMENT_LIFETIME_BEGIN(Object, ChildElementId) \
+	if (CAN_TRACE_OBJECT(Object)) { FObjectTrace::OutputChildElement(Object, ChildElementId, (Object ? Object->GetOuter() : nullptr), INDEX_NONE); }
+
+#define TRACE_CHILD_ELEMENT_WITH_OUTER_LIFETIME_BEGIN(Object, ChildElementId, OuterObject, OuterChildElementId) \
+	if (CAN_TRACE_OBJECT(Object)) { FObjectTrace::OutputChildElement(Object, ChildElementId, OuterObject, OuterChildElementId); }
+
+#define TRACE_CHILD_ELEMENT_LIFETIME_END(Object, ChildElementId) \
+	if (CAN_TRACE_OBJECT(Object)) { FObjectTrace::OutputChildElementLifetimeEnd(Object, ChildElementId); }
+
+#define TRACE_CHILD_ELEMENT_EVENT(Object, ChildElementId, Event) \
+	if (CAN_TRACE_OBJECT(Object)) { FObjectTrace::OutputChildElementEvent(Object, ChildElementId, Event); }
 
 #else
 
 #define TRACE_OBJECT_EVENT(Object, Event) \
 	UNCONDITIONAL_TRACE_OBJECT_EVENT(Object, Event);
 
+#define TRACE_OBJECT_TRANSFORM(Object, Transform, Teleport) \
+	FObjectTrace::OutputObjectTransform(Object, Transform, Teleport);
+
 #define TRACE_OBJECT_LIFETIME_BEGIN(Object) \
 	FObjectTrace::OutputObjectLifetimeBegin(Object);
-	
+
 #define TRACE_OBJECT_LIFETIME_END(Object) \
 	FObjectTrace::OutputObjectLifetimeEnd(Object);
-	
+
 #define TRACE_PAWN_POSSESS(Controller, Pawn)\
 	FObjectTrace::OutputPawnPossess(Controller, Pawn);
 
+#define TRACE_CHILD_ELEMENT(Object, ChildElementId) \
+	FObjectTrace::OutputChildElement(Object, ChildElementId, (Object ? Object->GetOuter() : nullptr), INDEX_NONE);
+
+#define TRACE_CHILD_ELEMENT_WITH_OUTER(Object, ChildElementId, OuterObject, OuterChildElementId) \
+	FObjectTrace::OutputChildElement(Object, ChildElementId, OuterObject, OuterChildElementId);
+
+#define TRACE_CHILD_ELEMENT_LIFETIME_BEGIN(Object, ChildElementId) \
+	FObjectTrace::OutputChildElement(Object, ChildElementId, (Object ? Object->GetOuter() : nullptr), INDEX_NONE);
+
+#define TRACE_CHILD_ELEMENT_WITH_OUTER_LIFETIME_BEGIN(Object, ChildElementId, OuterObject, OuterChildElementId) \
+	FObjectTrace::OutputChildElement(Object, ChildElementId, OuterObject, OuterChildElementId);
+
+#define TRACE_CHILD_ELEMENT_LIFETIME_END(Object, ChildElementId) \
+	FObjectTrace::OutputChildElementLifetimeEnd(Object, ChildElementId);
+
+#define TRACE_CHILD_ELEMENT_EVENT(Object, ChildElementId, Event) \
+	FObjectTrace::OutputChildElementEvent(Object, ChildElementId, Event);
+
+
 #endif
 	
 #define UNCONDITIONAL_TRACE_OBJECT_EVENT(Object, Event) \
@@ -175,5 +224,11 @@ struct FObjectTrace
 #define TRACE_VIEW(...)
 #define TRACE_OBJECT_LIFETIME_BEGIN(...)
 #define TRACE_OBJECT_LIFETIME_END(...)
+#define TRACE_CHILD_ELEMENT(...)
+#define TRACE_CHILD_ELEMENT_WITH_OUTER(...)
+#define TRACE_CHILD_ELEMENT_LIFETIME_BEGIN(...)
+#define TRACE_CHILD_ELEMENT_WITH_OUTER_LIFETIME_BEGIN(...)
+#define TRACE_CHILD_ELEMENT_LIFETIME_END(...)
+#define TRACE_CHILD_ELEMENT_EVENT(...)
 
 #endif
\ No newline at end of file
