[RewindDebugger] track creator can now indicate if the created child track can be used by the object track as an override for its representation and timeline values
- Replaced FRewindDebuggerTrack visible flag by a mask to allow distinction between tracks hidden by UI filtering or from code
- fixed Object track existence window for still active objects (upperbound == +inf)
- StateTree instance track is allowed on child element on any type of parent object (MassStateTreeSubsystem, AIController, Pawn, etc.)

diff --git a/Engine/Plugins/Animation/GameplayInsights/Source/RewindDebugger/Private/RewindDebuggerObjectTrack.cpp b/Engine/Plugins/Animation/GameplayInsights/Source/RewindDebugger/Private/RewindDebuggerObjectTrack.cpp
--- a/Engine/Plugins/Animation/GameplayInsights/Source/RewindDebugger/Private/RewindDebuggerObjectTrack.cpp
+++ b/Engine/Plugins/Animation/GameplayInsights/Source/RewindDebugger/Private/RewindDebuggerObjectTrack.cpp
@@ -49,7 +49,7 @@ FRewindDebuggerObjectTrack::FRewindDebuggerObjectTrack(const FObjectId& InObject
 
 	if (TArray<const TCHAR*> TypeNames; GetTypeHierarchyNames(InObjectId.GetMainId(), *Session, TypeNames))
 	{
-		FRewindDebuggerTrackCreators::EnumerateCreators([&TrackChildren = TrackChildren, &TypeNames, IsChildElementTrack = InObjectId.IsChildElement()](const IRewindDebuggerTrackCreator* Creator)
+		FRewindDebuggerTrackCreators::EnumerateCreators([&Creators = ChildTrackCreators, &TypeNames, IsChildElementTrack = InObjectId.IsChildElement()](const IRewindDebuggerTrackCreator* Creator)
 			{
 				if (IsChildElementTrack && !Creator->IsChildElementSupported())
 				{
@@ -58,13 +58,13 @@ FRewindDebuggerObjectTrack::FRewindDebuggerObjectTrack(const FObjectId& InObject
 
 				if (TypeNames.FindByPredicate([TargetTypeName = Creator->GetTargetTypeName()](const TCHAR* TypeName) { return TargetTypeName == TypeName; }))
 				{
-					TrackChildren.Push({ Creator, /*Track*/nullptr });
+					Creators.Push({ Creator, /*Track*/nullptr });
 				}
 			});
 	}
 
 	// sort by creators by priority + name
-	TrackChildren.Sort([](const FTrackCreatorAndTrack& A, const FTrackCreatorAndTrack& B)
+	ChildTrackCreators.Sort([](const FTrackCreatorAndTrack& A, const FTrackCreatorAndTrack& B)
 		{
 			const int32 SortOrderPriorityA = A.Creator->GetSortOrderPriority();
 			const int32 SortOrderPriorityB = B.Creator->GetSortOrderPriority();
@@ -80,6 +80,14 @@ FRewindDebuggerObjectTrack::FRewindDebuggerObjectTrack(const FObjectId& InObject
 
 TSharedPtr<SWidget> FRewindDebuggerObjectTrack::GetTimelineViewInternal()
 {
+	if (PrimaryChildTrack.IsValid())
+	{
+		if (TSharedPtr<SWidget> TimelineViewOverride = PrimaryChildTrack->GetTimelineView())
+		{
+			return TimelineViewOverride;
+		}
+	}
+
 	return SNew(SEventTimelineView)
 		.ViewRange_Lambda([]() { return IRewindDebugger::Instance()->GetCurrentViewRange(); })
 		.EventData_Raw(this, &FRewindDebuggerObjectTrack::GetExistenceRange);
@@ -87,6 +95,15 @@ TSharedPtr<SWidget> FRewindDebuggerObjectTrack::GetTimelineViewInternal()
 
 bool FRewindDebuggerObjectTrack::HandleDoubleClickInternal()
 {
+	// Allow primary child track to handle double-click
+	if (PrimaryChildTrack.IsValid())
+	{
+		if (PrimaryChildTrack->HandleDoubleClick())
+		{
+			return true;
+		}
+	}
+
 	IModularFeatures& ModularFeatures = IModularFeatures::Get();
 	static const FName HandlerFeatureName = IRewindDebuggerDoubleClickHandler::ModularFeatureName;
 
@@ -128,10 +145,9 @@ bool FRewindDebuggerObjectTrack::HandleDoubleClickInternal()
 	return true;
 }
 
-
 void FRewindDebuggerObjectTrack::IterateSubTracksInternal(TFunction<void(TSharedPtr<FRewindDebuggerTrack> SubTrack)> IteratorFunction)
 {
-	for (FTrackCreatorAndTrack& TrackChild : TrackChildren)
+	for (FTrackCreatorAndTrack& TrackChild : ChildTrackCreators)
 	{
 		if (TrackChild.Track.IsValid())
 		{
@@ -143,10 +159,47 @@ void FRewindDebuggerObjectTrack::IterateSubTracksInternal(TFunction<void(TShared
 	{
 		IteratorFunction(Track);
 	}
+}
+
+FName FRewindDebuggerObjectTrack::GetNameInternal() const
+{
+	if (PrimaryChildTrack.IsValid())
+	{
+		const FName NameOverride = PrimaryChildTrack->GetName();
+		if (!NameOverride.IsNone())
+		{
+			return NameOverride;
+		}
+	}
+
+	return NAME_None;
+}
+
+FSlateIcon FRewindDebuggerObjectTrack::GetIconInternal()
+{
+	if (PrimaryChildTrack.IsValid())
+	{
+		const FSlateIcon IconOverride = PrimaryChildTrack->GetIcon();
+		if (IconOverride.IsSet())
+		{
+			return IconOverride;
+		}
+	}
+
+	return Icon;
 };
 
 FText FRewindDebuggerObjectTrack::GetDisplayNameInternal() const
 {
+	if (PrimaryChildTrack.IsValid())
+	{
+		const FText DisplayNameOverride = PrimaryChildTrack->GetDisplayName();
+		if (!DisplayNameOverride.IsEmpty())
+		{
+			return DisplayNameOverride;
+		}
+	}
+
 	if (!bDisplayNameValid)
 	{
 		if (ObjectId.IsSet())
@@ -180,6 +233,25 @@ FText FRewindDebuggerObjectTrack::GetDisplayNameInternal() const
 	return DisplayName;
 }
 
+uint64 FRewindDebuggerObjectTrack::GetObjectIdInternal() const
+{
+	return ObjectId.GetMainId();
+}
+
+uint64 FRewindDebuggerObjectTrack::GetChildElementIdInternal() const
+{
+	return ObjectId.GetChildElementId();
+}
+
+bool FRewindDebuggerObjectTrack::HasDebugDataInternal() const
+{
+	if (PrimaryChildTrack.IsValid())
+	{
+		return PrimaryChildTrack->HasDebugData();
+	}
+	return false;
+}
+
 bool FRewindDebuggerObjectTrack::UpdateInternal()
 {
 	TRACE_CPUPROFILER_EVENT_SCOPE(FRewindDebuggerObjectTrack::UpdateInternal);
@@ -196,7 +268,12 @@ bool FRewindDebuggerObjectTrack::UpdateInternal()
 	ExistenceRange->Windows.SetNum(0, EAllowShrinking::No);
 	if (Existence.HasLowerBound() && Existence.HasUpperBound())
 	{
-		ExistenceRange->Windows.Add({ Existence.GetLowerBoundValue(), Existence.GetUpperBoundValue(), LOCTEXT("Object Existence","Object Existence"), LOCTEXT("Object Existence","Object Existence"), FLinearColor(0.1f,0.11f,0.1f) });
+		double EndTime = Existence.GetUpperBoundValue();
+		if (EndTime == std::numeric_limits<double>::infinity())
+		{
+			EndTime = GameplayProvider->GetRecordingDuration();
+		}
+		ExistenceRange->Windows.Add({ Existence.GetLowerBoundValue(), EndTime, LOCTEXT("Object Existence","Object Existence"), LOCTEXT("Object Existence","Object Existence"), FLinearColor(0.1f,0.11f,0.1f) });
 	}
 
 	if (!bIconSearched)
@@ -219,18 +296,27 @@ bool FRewindDebuggerObjectTrack::UpdateInternal()
 
 	FoundObjects.Add(ObjectId); // prevent debug views from being removed
 
-	// add debug views as children
+	// add debug views as children and make sure primary child track is up-to-date
+	PrimaryChildTrack.Reset();
 
-	for (FTrackCreatorAndTrack& TrackChild : TrackChildren)
+	for (FTrackCreatorAndTrack& CreatorAndTrack : ChildTrackCreators)
 	{
-		const bool bHasDebugInfo = TrackChild.Creator->HasDebugInfo(ObjectId);
+		const bool bHasDebugInfo = CreatorAndTrack.Creator->HasDebugInfo(ObjectId);
 
-		if (TrackChild.Track.IsValid())
+		if (CreatorAndTrack.Track.IsValid())
 		{
+			if (CreatorAndTrack.Creator->IsCreatingPrimaryChildTrack())
+			{
+				ensureMsgf(!PrimaryChildTrack.IsValid(), TEXT("Multiple primary tracks is not supported yet."));
+				PrimaryChildTrack = CreatorAndTrack.Track;
+				// Hide the track since its data will be exposed by the current Object track
+				PrimaryChildTrack->SetHiddenFlag(ETrackHiddenFlags::HiddenByCode);
+			}
+
 			if (!bHasDebugInfo)
 			{
 				bChanged = true;
-				TrackChild.Track.Reset();
+				CreatorAndTrack.Track.Reset();
 			}
 		}
 		else
@@ -238,7 +324,7 @@ bool FRewindDebuggerObjectTrack::UpdateInternal()
 			if (bHasDebugInfo)
 			{
 				bChanged = true;
-				TrackChild.Track = TrackChild.Creator->CreateTrack(ObjectId);
+				CreatorAndTrack.Track = CreatorAndTrack.Creator->CreateTrack(ObjectId);
 			}
 		}
 	}
@@ -362,7 +448,7 @@ bool FRewindDebuggerObjectTrack::UpdateInternal()
 	}
 	{
 		TRACE_CPUPROFILER_EVENT_SCOPE(FRewindDebuggerObjectTrack::UpdateTrackChilden);
-		for (auto& TrackChild : TrackChildren)
+		for (auto& TrackChild : ChildTrackCreators)
 		{
 			if (TrackChild.Track.IsValid())
 			{
diff --git a/Engine/Plugins/Animation/GameplayInsights/Source/RewindDebugger/Private/RewindDebuggerObjectTrack.h b/Engine/Plugins/Animation/GameplayInsights/Source/RewindDebugger/Private/RewindDebuggerObjectTrack.h
--- a/Engine/Plugins/Animation/GameplayInsights/Source/RewindDebugger/Private/RewindDebuggerObjectTrack.h
+++ b/Engine/Plugins/Animation/GameplayInsights/Source/RewindDebugger/Private/RewindDebuggerObjectTrack.h
@@ -31,34 +31,12 @@ private:
 	virtual TSharedPtr<SWidget> GetTimelineViewInternal() override;
 	virtual bool UpdateInternal() override;
 	virtual void IterateSubTracksInternal(TFunction<void(TSharedPtr<FRewindDebuggerTrack> SubTrack)> IteratorFunction) override;
-
-	virtual FName GetNameInternal() const override
-	{
-		return "";
-	}
-
-	virtual FSlateIcon GetIconInternal() override
-	{
-		return Icon;
-	}
-
+	virtual FName GetNameInternal() const override;
+	virtual FSlateIcon GetIconInternal() override;
 	virtual FText GetDisplayNameInternal() const override;
-
-	virtual uint64 GetObjectIdInternal() const override
-	{
-		return ObjectId.GetMainId();
-	}
-
-	virtual uint64 GetChildElementIdInternal() const override
-	{
-		return ObjectId.GetChildElementId();
-	}
-
-	virtual bool HasDebugDataInternal() const override
-	{
-		return false;
-	}
-
+	virtual uint64 GetObjectIdInternal() const override;
+	virtual uint64 GetChildElementIdInternal() const override;
+	virtual bool HasDebugDataInternal() const override;
 	virtual bool HandleDoubleClickInternal() override;
 
 	mutable FText DisplayName;
@@ -66,9 +44,12 @@ private:
 	FSlateIcon Icon;
 	TSharedPtr<SEventTimelineView::FTimelineEventData> ExistenceRange;
 	FObjectId ObjectId;
-	TArray<FTrackCreatorAndTrack> TrackChildren;
+	TArray<FTrackCreatorAndTrack> ChildTrackCreators;
 	TArray<TSharedPtr<FRewindDebuggerTrack>> Children;
 
+	/** This child track will be used by the current object track to replace its own appearance and timeline values */
+	TSharedPtr<FRewindDebuggerTrack> PrimaryChildTrack;
+
 	bool bAddController;
 	mutable bool bDisplayNameValid;
 	bool bIconSearched;
diff --git a/Engine/Plugins/Animation/GameplayInsights/Source/RewindDebugger/Private/SRewindDebugger.cpp b/Engine/Plugins/Animation/GameplayInsights/Source/RewindDebugger/Private/SRewindDebugger.cpp
--- a/Engine/Plugins/Animation/GameplayInsights/Source/RewindDebugger/Private/SRewindDebugger.cpp
+++ b/Engine/Plugins/Animation/GameplayInsights/Source/RewindDebugger/Private/SRewindDebugger.cpp
@@ -482,7 +482,7 @@ void SRewindDebugger::Construct(const FArguments& InArgs, TSharedRef<FUICommandL
 										[
 											TimelinesView.ToSharedRef()
 										]
-										+ SOverlay::Slot().HAlign(EHorizontalAlignment::HAlign_Right)
+										+ SOverlay::Slot().HAlign(HAlign_Right)
 										[
 											ScrollBar.ToSharedRef()
 										]
@@ -495,9 +495,16 @@ void SRewindDebugger::Construct(const FArguments& InArgs, TSharedRef<FUICommandL
 
 bool FilterTrack(TSharedPtr<RewindDebugger::FRewindDebuggerTrack>& Track, const FString& FilterString, bool bRemoveNoData, const TArray<FName>& FilteredTrackTypes, bool bParentFilterPassed = false)
 {
+	using namespace RewindDebugger;
 	if (FilteredTrackTypes.Contains(Track->GetName()))
 	{
-		Track->SetIsVisible(false);
+		Track->SetHiddenFlag(ETrackHiddenFlags::HiddenByUI);
+		return false;
+	}
+
+	// No need to test for UI filtering if track is hidden for any other reason
+	if (Track->HasAnyFlags(ETrackHiddenFlags::AnyFlag & ~ETrackHiddenFlags::HiddenByUI))
+	{
 		return false;
 	}
 
@@ -507,16 +514,23 @@ bool FilterTrack(TSharedPtr<RewindDebugger::FRewindDebuggerTrack>& Track, const
 	const bool bThisFilterPassed = (!bStringFilterEmpty && bStringFilterPassed);
 
 	bool bAnyChildVisible = false;
-	Track->IterateSubTracks([&bAnyChildVisible, bThisFilterPassed, FilterString, bRemoveNoData, FilteredTrackTypes](TSharedPtr<RewindDebugger::FRewindDebuggerTrack> ChildTrack)
+	Track->IterateSubTracks([&bAnyChildVisible, bThisFilterPassed, FilterString, bRemoveNoData, FilteredTrackTypes](TSharedPtr<FRewindDebuggerTrack> ChildTrack)
 		{
 			const bool bChildIsVisible = FilterTrack(ChildTrack, FilterString, bRemoveNoData, FilteredTrackTypes, bThisFilterPassed);
 			bAnyChildVisible |= bChildIsVisible;
 		});
 
 	const bool bVisible = bAnyChildVisible || ((!bRemoveNoData || Track->HasDebugData()) && bStringFilterPassed);
+	if (bVisible)
+	{
+		Track->UnsetHiddenFlag(ETrackHiddenFlags::HiddenByUI);
+	}
+	else
+	{
+		Track->SetHiddenFlag(ETrackHiddenFlags::HiddenByUI);
+	}
 
-	Track->SetIsVisible(bVisible);
-	return bVisible;
+	return Track->IsVisible();
 }
 
 void SRewindDebugger::RefreshTracks()
diff --git a/Engine/Plugins/Runtime/StateTree/Source/StateTreeEditorModule/Private/Debugger/StateTreeRewindDebuggerTrack.cpp b/Engine/Plugins/Runtime/StateTree/Source/StateTreeEditorModule/Private/Debugger/StateTreeRewindDebuggerTrack.cpp
--- a/Engine/Plugins/Runtime/StateTree/Source/StateTreeEditorModule/Private/Debugger/StateTreeRewindDebuggerTrack.cpp
+++ b/Engine/Plugins/Runtime/StateTree/Source/StateTreeEditorModule/Private/Debugger/StateTreeRewindDebuggerTrack.cpp
@@ -7,6 +7,7 @@
 #include "Editor.h"
 #include "Modules/ModuleManager.h"
 #include "PropertyEditorModule.h"
+#include "Styling/SlateIconFinder.h"
 #include "Subsystems/AssetEditorSubsystem.h"
 
 #define LOCTEXT_NAMESPACE "StateTreeRewindDebuggerTrack"
@@ -58,9 +59,17 @@ bool FRewindDebuggerTrackCreator::IsChildElementSupportedInternal() const
 	return true;
 }
 
+bool FRewindDebuggerTrackCreator::IsCreatingPrimaryChildTrackInternal() const
+{
+	return true;
+}
+
 FName FRewindDebuggerTrackCreator::GetTargetTypeNameInternal() const
 {
-	static const FName AssociatedObjectName("MassStateTreeSubsystem");
+	// Statetree instance can be executed by multiple systems so 
+	// we want to process all objects tracks but we early out in 
+	// FRewindDebuggerTrackCreator::HasDebugInfoInternal if not a child element
+	static const FName AssociatedObjectName("Object");
 	return AssociatedObjectName;
 }
 
@@ -70,6 +79,7 @@ FName FRewindDebuggerTrackCreator::GetTargetTypeNameInternal() const
 FRewindDebuggerTrack::FRewindDebuggerTrack(const RewindDebugger::FObjectId& InObjectId)
 	: EventCollection(FStateTreeInstanceDebugId(InObjectId.GetChildElementId()))
 	, InstanceTrackHelper(new FInstanceTrackHelper(EventCollection.InstanceId, IRewindDebugger::Instance()->GetCurrentViewRange()))
+	, Icon(FSlateIconFinder::FindIconForClass(UStateTree::StaticClass()))
 	, ObjectId(InObjectId)
 {
 	FPropertyEditorModule& PropertyEditorModule = FModuleManager::GetModuleChecked<FPropertyEditorModule>("PropertyEditor");
@@ -95,7 +105,8 @@ FText FRewindDebuggerTrack::GetDisplayNameInternal() const
 {
 	if (Descriptor.IsValid())
 	{
-		return FText::FromString(Descriptor->Name);
+		return FText::FromString(FString::Printf(TEXT("StateTree (%s)")
+			, *GetNameSafe(Descriptor->StateTree.Get())));
 	}
 
 	return LOCTEXT("StateTreeRewindDebuggerTrackName", "StateTreeInstance");
diff --git a/Engine/Plugins/Runtime/StateTree/Source/StateTreeEditorModule/Private/Debugger/StateTreeRewindDebuggerTrack.h b/Engine/Plugins/Runtime/StateTree/Source/StateTreeEditorModule/Private/Debugger/StateTreeRewindDebuggerTrack.h
--- a/Engine/Plugins/Runtime/StateTree/Source/StateTreeEditorModule/Private/Debugger/StateTreeRewindDebuggerTrack.h
+++ b/Engine/Plugins/Runtime/StateTree/Source/StateTreeEditorModule/Private/Debugger/StateTreeRewindDebuggerTrack.h
@@ -26,6 +26,7 @@ protected:
 	virtual TSharedPtr<RewindDebugger::FRewindDebuggerTrack> CreateTrackInternal(const RewindDebugger::FObjectId& InObjectId) const override;
 	virtual bool HasDebugInfoInternal(const RewindDebugger::FObjectId& InObjectId) const override;
 	virtual bool IsChildElementSupportedInternal() const override;
+	virtual bool IsCreatingPrimaryChildTrackInternal() const override;
 };
 
 /**
@@ -37,7 +38,7 @@ struct FRewindDebuggerTrack : RewindDebugger::FRewindDebuggerTrack, ITraceReader
 
 private:
 	//~ Begin UE::StateTreeDebugger::ITraceReader interface
-	virtual UE::StateTreeDebugger::FInstanceEventCollection* GetOrCreateEventCollection(FStateTreeInstanceDebugId InstanceId) override;
+	virtual FInstanceEventCollection* GetOrCreateEventCollection(FStateTreeInstanceDebugId InstanceId) override;
 	//~ End UE::StateTreeDebugger::ITraceReader interface
 
 	//~ Begin RewindDebugger::FRewindDebuggerTrack interface
diff --git a/Engine/Source/Editor/RewindDebuggerInterface/Public/IRewindDebuggerTrackCreator.h b/Engine/Source/Editor/RewindDebuggerInterface/Public/IRewindDebuggerTrackCreator.h
--- a/Engine/Source/Editor/RewindDebuggerInterface/Public/IRewindDebuggerTrackCreator.h
+++ b/Engine/Source/Editor/RewindDebuggerInterface/Public/IRewindDebuggerTrackCreator.h
@@ -100,6 +100,12 @@ PRAGMA_ENABLE_DEPRECATION_WARNINGS
 		return IsChildElementSupportedInternal();
 	}
 
+	/** @return Whether the track created will be used by its parent track to replace its representation and timeline values */
+	bool IsCreatingPrimaryChildTrack() const
+	{
+		return IsCreatingPrimaryChildTrackInternal();
+	}
+
 private:
 
 	/** get the UObject type name this Creator will create child tracks for */
@@ -152,6 +158,11 @@ private:
 	{
 		return false;
 	}
+
+	virtual bool IsCreatingPrimaryChildTrackInternal() const
+	{
+		return false;
+	}
 };
 
 } // namespace RewindDebugger
diff --git a/Engine/Source/Editor/RewindDebuggerInterface/Public/RewindDebuggerTrack.h b/Engine/Source/Editor/RewindDebuggerInterface/Public/RewindDebuggerTrack.h
--- a/Engine/Source/Editor/RewindDebuggerInterface/Public/RewindDebuggerTrack.h
+++ b/Engine/Source/Editor/RewindDebuggerInterface/Public/RewindDebuggerTrack.h
@@ -18,6 +18,14 @@ struct FToolMenuSection;
 
 namespace RewindDebugger
 {
+enum class ETrackHiddenFlags : uint8
+{
+	None = 0,
+	HiddenByCode = 1 << 0,
+	HiddenByUI = 1 << 1,
+	AnyFlag = 0xFF
+};
+ENUM_CLASS_FLAGS(ETrackHiddenFlags);
 
 class FRewindDebuggerTrack
 {
@@ -147,14 +155,37 @@ public:
 		return HandleDoubleClickInternal();
 	}
 
+	void SetHiddenFlag(const ETrackHiddenFlags Flag)
+	{
+		EnumAddFlags(HiddenFlags, Flag);
+	}
+
+	void UnsetHiddenFlag(const ETrackHiddenFlags Flag)
+	{
+		EnumRemoveFlags(HiddenFlags, Flag);
+	}
+
+	bool HasAnyFlags(const ETrackHiddenFlags Flags) const
+	{
+		return EnumHasAnyFlags(HiddenFlags, Flags);
+	}
+
 	bool IsVisible() const
 	{
-		return bVisible;
+		return HiddenFlags == ETrackHiddenFlags::None;
 	}
 
+	UE_DEPRECATED(5.7, "Use SetHiddenFlag/UnsetHiddenFlag instead.")
 	void SetIsVisible(const bool bInIsVisible)
 	{
-		bVisible = bInIsVisible;
+		if (bInIsVisible)
+		{
+			UnsetHiddenFlag(ETrackHiddenFlags::HiddenByCode);
+		}
+		else
+		{
+			SetHiddenFlag(ETrackHiddenFlags::HiddenByCode);
+		}
 	}
 
 	/** Called to generate context menu for the current selected track */
@@ -217,11 +248,12 @@ private:
 		return true;
 	};
 
+	ETrackHiddenFlags HiddenFlags = ETrackHiddenFlags::None;
+
 	bool bSelected : 1 = false;
 	bool bTrackHovered : 1 = false;
 	bool bTreeHovered : 1 = false;
 	bool bExpanded : 1 = true;
-	bool bVisible : 1 = true;
 };
 
 } // namespace RewindDebugger
\ No newline at end of file
