[RewindDebugger] Allow primary track to replace ObjectTrack details view

diff --git a/Engine/Plugins/Animation/GameplayInsights/Source/RewindDebugger/Private/RewindDebuggerObjectTrack.cpp b/Engine/Plugins/Animation/GameplayInsights/Source/RewindDebugger/Private/RewindDebuggerObjectTrack.cpp
--- a/Engine/Plugins/Animation/GameplayInsights/Source/RewindDebugger/Private/RewindDebuggerObjectTrack.cpp
+++ b/Engine/Plugins/Animation/GameplayInsights/Source/RewindDebugger/Private/RewindDebuggerObjectTrack.cpp
@@ -78,6 +78,19 @@ FRewindDebuggerObjectTrack::FRewindDebuggerObjectTrack(const FObjectId& InObject
 		});
 }
 
+TSharedPtr<SWidget> FRewindDebuggerObjectTrack::GetDetailsViewInternal()
+{
+	if (PrimaryChildTrack.IsValid())
+	{
+		if (TSharedPtr<SWidget> DetailsViewOverride = PrimaryChildTrack->GetDetailsView())
+		{
+			return DetailsViewOverride;
+		}
+	}
+
+	return TSharedPtr<SWidget>();
+}
+
 TSharedPtr<SWidget> FRewindDebuggerObjectTrack::GetTimelineViewInternal()
 {
 	if (PrimaryChildTrack.IsValid())
diff --git a/Engine/Plugins/Animation/GameplayInsights/Source/RewindDebugger/Private/RewindDebuggerObjectTrack.h b/Engine/Plugins/Animation/GameplayInsights/Source/RewindDebugger/Private/RewindDebuggerObjectTrack.h
--- a/Engine/Plugins/Animation/GameplayInsights/Source/RewindDebugger/Private/RewindDebuggerObjectTrack.h
+++ b/Engine/Plugins/Animation/GameplayInsights/Source/RewindDebugger/Private/RewindDebuggerObjectTrack.h
@@ -28,6 +28,7 @@ public:
 	}
 
 private:
+	virtual TSharedPtr<SWidget> GetDetailsViewInternal() override;
 	virtual TSharedPtr<SWidget> GetTimelineViewInternal() override;
 	virtual bool UpdateInternal() override;
 	virtual void IterateSubTracksInternal(TFunction<void(TSharedPtr<FRewindDebuggerTrack> SubTrack)> IteratorFunction) override;
