[VisualLogger] exposed some types internally to be used by RewindDebugger track

diff --git a/Engine/Plugins/Animation/GameplayInsights/Source/RewindDebuggerVLog/Private/VLogDetailsObject.cpp b/Engine/Plugins/Animation/GameplayInsights/Source/RewindDebuggerVLog/Private/VLogDetailsObject.cpp
deleted file mode 100644
diff --git a/Engine/Plugins/Animation/GameplayInsights/Source/RewindDebuggerVLog/Private/VLogDetailsObject.h b/Engine/Plugins/Animation/GameplayInsights/Source/RewindDebuggerVLog/Private/VLogDetailsObject.h
deleted file mode 100644
diff --git a/Engine/Plugins/Animation/GameplayInsights/Source/RewindDebuggerVLog/Private/VisualLogTrack.cpp b/Engine/Plugins/Animation/GameplayInsights/Source/RewindDebuggerVLog/Private/VisualLogTrack.cpp
--- a/Engine/Plugins/Animation/GameplayInsights/Source/RewindDebuggerVLog/Private/VisualLogTrack.cpp
+++ b/Engine/Plugins/Animation/GameplayInsights/Source/RewindDebuggerVLog/Private/VisualLogTrack.cpp
@@ -5,6 +5,7 @@
 #include "Modules/ModuleManager.h"
 #include "PropertyEditorModule.h"
 #include "VisualLoggerProvider.h"
+#include "SVisualLoggerTableRow.h"
 #include "UObject/Package.h"
 
 #define LOCTEXT_NAMESPACE "VisualLogTrack"
@@ -29,15 +30,11 @@ FText& GetCategoryNameAsText(FName InCategory)
 	return CachedEntry->Value;
 }
 
-UVLogDetailsObject* FVisualLogCategoryTrack::InitializeDetailsObject()
-{
-	UVLogDetailsObject* DetailsObject = NewObject<UVLogDetailsObject>();
-	DetailsObject->SetFlags(RF_Standalone);
-	DetailsObjectWeakPtr = MakeWeakObjectPtr(DetailsObject);
-	DetailsView->SetObject(DetailsObject);
-	return DetailsObject;
-}
+static const FLazyName VisualLogName("Visual Logging");
 
+//----------------------------------------------------------------------//
+// FVisualLogCategoryTrack
+//----------------------------------------------------------------------//
 FVisualLogCategoryTrack::FVisualLogCategoryTrack(uint64 InObjectId, const FName& InCategory) :
 	ObjectId(InObjectId),
 	Category(InCategory)
@@ -46,20 +43,23 @@ FVisualLogCategoryTrack::FVisualLogCategoryTrack(uint64 InObjectId, const FName&
 	EventData = MakeShared<SEventTimelineView::FTimelineEventData>();
 	Icon = FSlateIcon("EditorStyle", "Sequencer.Tracks.Event", "Sequencer.Tracks.Event");
 
-	FPropertyEditorModule& PropertyEditorModule = FModuleManager::GetModuleChecked<FPropertyEditorModule>("PropertyEditor");
-	FDetailsViewArgs DetailsViewArgs;
-	DetailsViewArgs.NameAreaSettings = FDetailsViewArgs::HideNameArea;
-	DetailsView = PropertyEditorModule.CreateDetailView(DetailsViewArgs);
-
-	InitializeDetailsObject();
-}
-
-FVisualLogCategoryTrack::~FVisualLogCategoryTrack()
-{
-	if (UVLogDetailsObject* DetailsObject = DetailsObjectWeakPtr.Get())
-	{
-		DetailsObject->ClearFlags(RF_Standalone);
-	}
+	SAssignNew(ListView, SListView<TSharedPtr<FLogEntryItem>>)
+		.ListItemsSource(&LogEntries)
+		.SelectionMode(ESelectionMode::None)
+		.OnGenerateRow_Lambda([](TSharedPtr<FLogEntryItem> Item, const TSharedRef<STableViewBase>& OwnerTable)
+			{
+				return SNew(SVisualLoggerTableRow, OwnerTable, Item);
+			})
+		.HeaderRow
+		(
+			SNew(SHeaderRow)
+			+ SHeaderRow::Column(SVisualLoggerTableRow::ColumnId_VerbosityLabel)
+				.FillWidth(0.1f)
+				.DefaultLabel(LOCTEXT("VLoggerVerbosityLabel", "Verbosity"))
+			+ SHeaderRow::Column(SVisualLoggerTableRow::ColumnId_LogLabel)
+				.FillWidth(0.8f)
+				.DefaultLabel(LOCTEXT("VLoggerLogLabel", "Log"))
+		);
 }
 
 TSharedPtr<SEventTimelineView::FTimelineEventData> FVisualLogCategoryTrack::GetEventData() const
@@ -74,17 +74,9 @@ TSharedPtr<SEventTimelineView::FTimelineEventData> FVisualLogCategoryTrack::GetE
 	return EventData;
 }
 
-static FLinearColor MakeNotifyColor(uint32 InSeed, bool bInLine = false)
-{
-	FRandomStream Stream(InSeed);
-	const uint8 Hue = (uint8)(Stream.FRand() * 255.0f);
-	const uint8 SatVal = bInLine ? 196 : 128;
-	return FLinearColor::MakeFromHSV8(Hue, SatVal, SatVal);
-}
-
 TSharedPtr<SWidget> FVisualLogCategoryTrack::GetDetailsViewInternal()
 {
-	return DetailsView;
+	return ListView;
 }
 
 bool FVisualLogCategoryTrack::UpdateInternal()
@@ -138,55 +130,61 @@ bool FVisualLogCategoryTrack::UpdateInternal()
 		if (PreviousScrubTime != CurrentScrubTime)
 		{
 			PreviousScrubTime = CurrentScrubTime;
-
-			UVLogDetailsObject* DetailsObject = DetailsObjectWeakPtr.Get();
-			if (DetailsObject == nullptr)
-			{
-				// this should not happen unless the object was garbage collected (which should not happen since it's marked as Standalone)
-				DetailsObject = InitializeDetailsObject();
-			}
-
-			DetailsObject->VisualLogDetails.SetNum(0, EAllowShrinking::No);
+			LogEntries.Reset();
 
 			const TraceServices::IFrameProvider& FramesProvider = TraceServices::ReadFrameProvider(*AnalysisSession);
 			TraceServices::FAnalysisSessionReadScope SessionReadScope(*AnalysisSession);
 
 			TraceServices::FFrame MarkerFrame;
-			if (FramesProvider.GetFrameFromTime(ETraceFrameType::TraceFrameType_Game, CurrentScrubTime, MarkerFrame))
+			if (FramesProvider.GetFrameFromTime(TraceFrameType_Game, CurrentScrubTime, MarkerFrame))
 			{
-				VisLogProvider->ReadVisualLogEntryTimeline(ObjectId, [this, DetailsObject, &MarkerFrame, EndTime, VisLogProvider, AnalysisSession](const FVisualLoggerProvider::VisualLogEntryTimeline& InTimeline)
+				VisLogProvider->ReadVisualLogEntryTimeline(ObjectId,
+					[this, &Entries = LogEntries, &MarkerFrame, EndTime, VisLogProvider, AnalysisSession](const FVisualLoggerProvider::VisualLogEntryTimeline& InTimeline)
 					{
-						InTimeline.EnumerateEvents(MarkerFrame.StartTime, MarkerFrame.EndTime, [this, DetailsObject, VisLogProvider, AnalysisSession](double InStartTime, double InEndTime, uint32 InDepth, const FVisualLogEntry& InMessage)
+						InTimeline.EnumerateEvents(MarkerFrame.StartTime, MarkerFrame.EndTime,
+							[this, &Entries, VisLogProvider, AnalysisSession](double InStartTime, double InEndTime, uint32 InDepth, const FVisualLogEntry& InMessage)
 							{
 								for (const FVisualLogShapeElement& Element : InMessage.ElementsToDraw)
 								{
-									if (Element.Category == Category)
+									if (Element.Category == Category
+										&& !Element.Description.IsEmpty())
 									{
-										DetailsObject->VisualLogDetails.Add({ Element.Category, Element.Description });
+										Entries.Add(MakeShared<FLogEntryItem>(
+											FLogEntryItem
+											{
+												.Category = Element.Category.ToString(),
+												.CategoryColor = Element.GetFColor(),
+												.Verbosity = Element.Verbosity,
+												.Line = Element.Description
+											}));
 									}
 								}
 								for (const FVisualLogLine& Line : InMessage.LogLines)
 								{
 									if (Line.Category == Category)
 									{
-										DetailsObject->VisualLogDetails.Add({ Line.Category, Line.Line });
+										Entries.Add(MakeShared<FLogEntryItem>(
+											FLogEntryItem
+											{
+												.Category = Line.Category.ToString(),
+												.CategoryColor = Line.Color,
+												.Verbosity = Line.Verbosity,
+												.Line = Line.Line,
+												.UserData = Line.UserData,
+												.TagName = Line.TagName
+											}));
 									}
 								}
 								return TraceServices::EEventEnumerate::Continue;
 							});
 					});
 			}
+			ListView->RebuildList();
 		}
 	}
 
-	bool bChanged = false;
+	constexpr bool bChanged = false;
 	return bChanged;
-
-}
-
-TSharedPtr<SWidget> FVisualLogTrack::GetDetailsViewInternal()
-{
-	return nullptr;
 }
 
 TSharedPtr<SWidget> FVisualLogCategoryTrack::GetTimelineViewInternal()
@@ -196,37 +194,16 @@ TSharedPtr<SWidget> FVisualLogCategoryTrack::GetTimelineViewInternal()
 		.EventData_Raw(this, &FVisualLogCategoryTrack::GetEventData);
 }
 
-static const FName VisualLogName("Visual Logging");
-
-FName FVisualLogTrackCreator::GetTargetTypeNameInternal() const
-{
-	static const FName AnimInstanceName("Object");
-	return AnimInstanceName;
-}
-
-FName FVisualLogTrackCreator::GetNameInternal() const
-{
-	return VisualLogName;
-}
+//----------------------------------------------------------------------//
+// FVisualLogTrack
+//----------------------------------------------------------------------//
 
-void FVisualLogTrackCreator::GetTrackTypesInternal(TArray<FRewindDebuggerTrackType>& Types) const
-{
-	Types.Add({ VisualLogName, LOCTEXT("Visual Logging", "Visual Logging") });
-}
-
-TSharedPtr<RewindDebugger::FRewindDebuggerTrack> FVisualLogTrackCreator::CreateTrackInternal(const FObjectId& InObjectId) const
-{
-	return MakeShared<RewindDebugger::FVisualLogTrack>(InObjectId.GetMainId());
-}
-
-
-FVisualLogTrack::FVisualLogTrack(uint64 InObjectId) :
+FVisualLogTrack::FVisualLogTrack(const uint64 InObjectId) :
 	ObjectId(InObjectId)
 {
 	Icon = FSlateIcon("EditorStyle", "Sequencer.Tracks.Event", "Sequencer.Tracks.Event");
 }
 
-
 bool FVisualLogTrack::UpdateInternal()
 {
 	TRACE_CPUPROFILER_EVENT_SCOPE(FVisualLogTrack::UpdateInternal);
@@ -275,7 +252,7 @@ bool FVisualLogTrack::UpdateInternal()
 		Children.SetNum(UniqueTrackIds.Num());
 		for (int i = 0; i < TrackCount; i++)
 		{
-			if (!Children[i].IsValid() || !(Children[i].Get()->GetName() == UniqueTrackIds[i]))
+			if (!Children[i].IsValid() || Children[i].Get()->GetName() != UniqueTrackIds[i])
 			{
 				Children[i] = MakeShared<FVisualLogCategoryTrack>(ObjectId, UniqueTrackIds[i]);
 				bChanged = true;
@@ -299,6 +276,35 @@ void FVisualLogTrack::IterateSubTracksInternal(TFunction<void(TSharedPtr<FRewind
 	}
 };
 
+TSharedPtr<SWidget> FVisualLogTrack::GetDetailsViewInternal()
+{
+	return nullptr;
+}
+
+//----------------------------------------------------------------------//
+// FVisualLogTrackCreator
+//----------------------------------------------------------------------//
+FName FVisualLogTrackCreator::GetTargetTypeNameInternal() const
+{
+	static const FName TargetTypeName("Object");
+	return TargetTypeName;
+}
+
+FName FVisualLogTrackCreator::GetNameInternal() const
+{
+	return VisualLogName;
+}
+
+void FVisualLogTrackCreator::GetTrackTypesInternal(TArray<FRewindDebuggerTrackType>& Types) const
+{
+	Types.Add({VisualLogName, LOCTEXT("Visual Logging", "Visual Logging")});
+}
+
+TSharedPtr<FRewindDebuggerTrack> FVisualLogTrackCreator::CreateTrackInternal(const FObjectId& InObjectId) const
+{
+	return MakeShared<FVisualLogTrack>(InObjectId.GetMainId());
+}
+
 bool FVisualLogTrackCreator::HasDebugInfoInternal(const FObjectId& InObjectId) const
 {
 	TRACE_CPUPROFILER_EVENT_SCOPE(FVisualLogTrack::HasDebugInfoInternal);
@@ -316,6 +322,6 @@ bool FVisualLogTrackCreator::HasDebugInfoInternal(const FObjectId& InObjectId) c
 	return bHasData;
 }
 
-}
+} // RewindDebugger
 
 #undef LOCTEXT_NAMESPACE
\ No newline at end of file
diff --git a/Engine/Plugins/Animation/GameplayInsights/Source/RewindDebuggerVLog/Private/VisualLogTrack.h b/Engine/Plugins/Animation/GameplayInsights/Source/RewindDebuggerVLog/Private/VisualLogTrack.h
--- a/Engine/Plugins/Animation/GameplayInsights/Source/RewindDebuggerVLog/Private/VisualLogTrack.h
+++ b/Engine/Plugins/Animation/GameplayInsights/Source/RewindDebuggerVLog/Private/VisualLogTrack.h
@@ -3,31 +3,51 @@
 
 #include "IRewindDebuggerTrackCreator.h"
 #include "RewindDebuggerTrack.h"
-#include "Textures/SlateIcon.h"
 #include "SEventTimelineView.h"
-#include "IDetailsView.h"
-#include "VLogDetailsObject.h"
-#include "UObject/WeakObjectPtrTemplates.h"
+#include "Textures/SlateIcon.h"
+#include "Widgets/Views/SListView.h"
+
+struct FLogEntryItem;
 
 namespace RewindDebugger
 {
 
+/**
+ * Tracks representing all log lines of a given category 
+ */
 class FVisualLogCategoryTrack : public FRewindDebuggerTrack
 {
 public:
 	FVisualLogCategoryTrack(uint64 InObjectId, const FName& Category);
-	virtual ~FVisualLogCategoryTrack();
-	const FName GetCategory() const { return Category; }
+	FName GetCategory() const
+	{
+		return Category;
+	}
+
 private:
 	virtual bool UpdateInternal() override;
-	TSharedPtr<SWidget> GetDetailsViewInternal() override;
+	virtual TSharedPtr<SWidget> GetDetailsViewInternal() override;
 	virtual TSharedPtr<SWidget> GetTimelineViewInternal() override;
-	virtual FSlateIcon GetIconInternal() override { return Icon; }
-	virtual FText GetDisplayNameInternal() const override { return TrackName; }
-	virtual uint64 GetObjectIdInternal() const override { return ObjectId; }
-	virtual FName GetNameInternal() const override { return Category; }
 
-	UVLogDetailsObject* InitializeDetailsObject();
+	virtual FSlateIcon GetIconInternal() override
+	{
+		return Icon;
+	}
+
+	virtual FText GetDisplayNameInternal() const override
+	{
+		return TrackName;
+	}
+
+	virtual uint64 GetObjectIdInternal() const override
+	{
+		return ObjectId;
+	}
+
+	virtual FName GetNameInternal() const override
+	{
+		return Category;
+	}
 
 	TSharedPtr<SEventTimelineView::FTimelineEventData> GetEventData() const;
 	
@@ -38,8 +58,8 @@ private:
 	FName Category;
 	FSlateIcon Icon;
 	FText TrackName;
-	TSharedPtr<IDetailsView> DetailsView;
-	TWeakObjectPtr<UVLogDetailsObject> DetailsObjectWeakPtr;
+	TSharedPtr<SListView<TSharedPtr<FLogEntryItem>>> ListView;
+	TArray<TSharedPtr<FLogEntryItem>> LogEntries;
 
 	double PreviousScrubTime = 0.0f;
 	
@@ -51,13 +71,13 @@ private:
 	};
 };
 
-
-
-// Parent track 
+/**
+ * Parent track for category tracks
+ */
 class FVisualLogTrack : public FRewindDebuggerTrack
 {
 public:
-	FVisualLogTrack(uint64 InObjectId);
+	explicit FVisualLogTrack(uint64 InObjectId);
 
 private:
 	virtual bool UpdateInternal() override;
@@ -75,10 +95,11 @@ private:
 	TArray<TSharedPtr<FVisualLogCategoryTrack>> Children;
 };
 
-
+/**
+ * Track creator for VisualLogger recorded entries
+ */
 class FVisualLogTrackCreator : public IRewindDebuggerTrackCreator
 {
-private:
 	virtual FName GetTargetTypeNameInternal() const override;
 	virtual FName GetNameInternal() const override;
 	virtual void GetTrackTypesInternal(TArray<FRewindDebuggerTrackType>& Types) const override;
@@ -86,4 +107,4 @@ private:
 	virtual bool HasDebugInfoInternal(const FObjectId& InObjectId) const override;
 };
 
-}
\ No newline at end of file
+}
diff --git a/Engine/Plugins/Animation/GameplayInsights/Source/RewindDebuggerVLog/RewindDebuggerVLog.Build.cs b/Engine/Plugins/Animation/GameplayInsights/Source/RewindDebuggerVLog/RewindDebuggerVLog.Build.cs
--- a/Engine/Plugins/Animation/GameplayInsights/Source/RewindDebuggerVLog/RewindDebuggerVLog.Build.cs
+++ b/Engine/Plugins/Animation/GameplayInsights/Source/RewindDebuggerVLog/RewindDebuggerVLog.Build.cs
@@ -12,6 +12,7 @@ namespace UnrealBuildTool.Rules
 				"Core",
 				"CoreUObject",
 				"Engine",
+				"InputCore",
 				"UnrealEd",
 				"LogVisualizer",
 				"TraceAnalysis",
diff --git a/Engine/Source/Developer/LogVisualizer/Internal/LogVisualizerStyle.h b/Engine/Source/Developer/LogVisualizer/Internal/LogVisualizerStyle.h
new file mode 100644
--- /dev/null
+++ b/Engine/Source/Developer/LogVisualizer/Internal/LogVisualizerStyle.h
@@ -0,0 +1,28 @@
+// Copyright Epic Games, Inc. All Rights Reserved.
+#pragma once
+
+#include "Templates/SharedPointer.h"
+
+class ISlateStyle;
+class FSlateStyleSet;
+
+#define UE_API LOGVISUALIZER_API
+
+/**
+ * Implements the visual style of LogVisualizer.
+ */
+class FLogVisualizerStyle
+{
+public:
+	static void Initialize();
+	static void Shutdown();
+	UE_API static const ISlateStyle& Get();
+	UE_API static FName GetStyleSetName();
+
+private:
+	static TSharedRef<FSlateStyleSet> Create();
+
+	static TSharedPtr<FSlateStyleSet> StyleInstance;
+};
+
+#undef UE_API
diff --git a/Engine/Source/Developer/LogVisualizer/Internal/SVisualLoggerTableRow.h b/Engine/Source/Developer/LogVisualizer/Internal/SVisualLoggerTableRow.h
new file mode 100644
--- /dev/null
+++ b/Engine/Source/Developer/LogVisualizer/Internal/SVisualLoggerTableRow.h
@@ -0,0 +1,73 @@
+// Copyright Epic Games, Inc. All Rights Reserved.
+#pragma once
+
+#include "LogVisualizerStyle.h"
+#include "Templates/SharedPointer.h"
+#include "VisualLoggerTypes.h"
+#include "Widgets/Text/STextBlock.h"
+#include "Widgets/Views/SListView.h"
+#include "Widgets/Views/STableRow.h"
+
+class SVisualLoggerTableRow : public SMultiColumnTableRow<TSharedPtr<FLogEntryItem>>
+{
+public:
+	SLATE_BEGIN_ARGS(SVisualLoggerTableRow)
+	{
+	}
+		SLATE_ATTRIBUTE(FText, FilterText)
+	SLATE_END_ARGS()
+
+	static const FLazyName ColumnId_CategoryLabel;
+	static const FLazyName ColumnId_VerbosityLabel;
+	static const FLazyName ColumnId_LogLabel;
+
+	void Construct(const FArguments& InArgs, const TSharedRef<STableViewBase>& InOwnerTableView, const TSharedPtr<FLogEntryItem> InEntryItem)
+	{
+		Item = InEntryItem;
+		FilterText = InArgs._FilterText;
+		SMultiColumnTableRow::Construct(FSuperRowType::FArguments(), InOwnerTableView);
+	}
+
+	virtual TSharedRef<SWidget> GenerateWidgetForColumn(const FName& ColumnName) override
+	{
+		if (ColumnName == ColumnId_CategoryLabel)
+		{
+			return SNew(STextBlock)
+				.ColorAndOpacity(FSlateColor(Item->CategoryColor))
+				.Text(FText::FromString(Item->Category))
+				.HighlightText(this, &SVisualLoggerTableRow::GetFilterText);
+		}
+
+		 if (ColumnName == ColumnId_VerbosityLabel)
+		{
+			return SNew(STextBlock)
+				.ColorAndOpacity(FSlateColor(Item->Verbosity == ELogVerbosity::Error ? FLinearColor::Red : (Item->Verbosity == ELogVerbosity::Warning ? FLinearColor::Yellow : FLinearColor::Gray)))
+				.Text(FText::FromString(FString(TEXT("(")) + FString(::ToString(Item->Verbosity)) + FString(TEXT(")"))));
+		}
+
+		if (ColumnName == ColumnId_LogLabel)
+		{
+			return SNew(STextBlock)
+				.AutoWrapText(true)
+				.ColorAndOpacity(FSlateColor(Item->Verbosity == ELogVerbosity::Error ? FLinearColor::Red : (Item->Verbosity == ELogVerbosity::Warning ? FLinearColor::Yellow : FLinearColor::Gray)))
+				.Text(FText::FromString(Item->Line))
+				.HighlightText(this, &SVisualLoggerTableRow::GetFilterText)
+				.TextStyle(FLogVisualizerStyle::Get(), TEXT("TextLogs.Text"));
+		}
+
+		return SNullWidget::NullWidget;
+	}
+
+	FText GetFilterText() const
+	{
+		return FilterText.Get();
+	}
+
+private:
+	TAttribute<FText> FilterText;
+	TSharedPtr<FLogEntryItem> Item;
+};
+
+const FLazyName SVisualLoggerTableRow::ColumnId_CategoryLabel("Category");
+const FLazyName SVisualLoggerTableRow::ColumnId_VerbosityLabel("Verbosity");
+const FLazyName SVisualLoggerTableRow::ColumnId_LogLabel("Log");
diff --git a/Engine/Source/Developer/LogVisualizer/Internal/VisualLoggerTypes.h b/Engine/Source/Developer/LogVisualizer/Internal/VisualLoggerTypes.h
new file mode 100644
--- /dev/null
+++ b/Engine/Source/Developer/LogVisualizer/Internal/VisualLoggerTypes.h
@@ -0,0 +1,16 @@
+// Copyright Epic Games, Inc. All Rights Reserved.
+#pragma once
+
+#include "Math/Color.h"
+#include "Logging/LogVerbosity.h"
+#include "UObject/NameTypes.h"
+
+struct FLogEntryItem
+{
+	FString Category;
+	FLinearColor CategoryColor = FLinearColor::White;
+	ELogVerbosity::Type Verbosity = ELogVerbosity::Log;
+	FString Line;
+	int64 UserData = 0;
+	FName TagName = NAME_None;
+};
diff --git a/Engine/Source/Developer/LogVisualizer/Private/LogVisualizerStyle.cpp b/Engine/Source/Developer/LogVisualizer/Private/LogVisualizerStyle.cpp
--- a/Engine/Source/Developer/LogVisualizer/Private/LogVisualizerStyle.cpp
+++ b/Engine/Source/Developer/LogVisualizer/Private/LogVisualizerStyle.cpp
@@ -6,6 +6,7 @@
 #include "Styling/CoreStyle.h"
 #include "Styling/StarshipCoreStyle.h"
 #include "Styling/AppStyle.h"
+#include "Styling/SlateStyle.h"
 #include "Styling/SlateStyleMacros.h"
 
 // This is to fix the issue that SlateStyleMacros like IMAGE_BRUSH look for RootToContentDir but StyleSet->RootToContentDir is how this style is set up
diff --git a/Engine/Source/Developer/LogVisualizer/Private/LogVisualizerStyle.h b/Engine/Source/Developer/LogVisualizer/Private/LogVisualizerStyle.h
deleted file mode 100644
diff --git a/Engine/Source/Developer/LogVisualizer/Private/SVisualLoggerLogsList.cpp b/Engine/Source/Developer/LogVisualizer/Private/SVisualLoggerLogsList.cpp
--- a/Engine/Source/Developer/LogVisualizer/Private/SVisualLoggerLogsList.cpp
+++ b/Engine/Source/Developer/LogVisualizer/Private/SVisualLoggerLogsList.cpp
@@ -1,6 +1,7 @@
 // Copyright Epic Games, Inc. All Rights Reserved.
 
 #include "SVisualLoggerLogsList.h"
+#include "SVisualLoggerTableRow.h"
 #include "Misc/OutputDeviceHelper.h"
 #include "Widgets/Text/STextBlock.h"
 #include "LogVisualizerSettings.h"
@@ -12,16 +13,6 @@
 
 #define LOCTEXT_NAMESPACE "SVisualLoggerLogsList"
 
-struct FLogEntryItem
-{
-	FString Category;
-	FLinearColor CategoryColor;
-	ELogVerbosity::Type Verbosity;
-	FString Line;
-	int64 UserData;
-	FName TagName;
-};
-
 namespace ELogsSortMode
 {
 	enum Type
@@ -32,70 +23,6 @@ namespace ELogsSortMode
 	};
 }
 
-//----------------------------------------------------------------------//
-// SVisualLoggerTableRow
-//----------------------------------------------------------------------//
-typedef TSharedPtr<FLogEntryItem, ESPMode::ThreadSafe> FLogEntryItemPtr;
-
-class SVisualLoggerTableRow : public SMultiColumnTableRow<FLogEntryItemPtr>
-{
-public:
-	SLATE_BEGIN_ARGS(SVisualLoggerTableRow) { }
-	SLATE_END_ARGS()
-
-	static const FName ColumnId_CategoryLabel;
-	static const FName ColumnId_VerbosityLabel;
-	static const FName ColumnId_LogLabel;
-
-	void Construct(const FArguments& InArgs, const TSharedRef<STableViewBase>& InOwnerTableView, const FLogEntryItemPtr InEntryItem)
-	{
-		Item = InEntryItem;
-		SMultiColumnTableRow<FLogEntryItemPtr>::Construct(FSuperRowType::FArguments(), InOwnerTableView);
-	}
-
-	virtual TSharedRef<SWidget> GenerateWidgetForColumn(const FName& ColumnName) override
-	{
-		if (ColumnName == ColumnId_CategoryLabel)
-		{
-			return SNew(STextBlock)
-				.ColorAndOpacity(FSlateColor(Item->CategoryColor))
-				.Text(FText::FromString(Item->Category))
-				.HighlightText(this, &SVisualLoggerTableRow::GetFilterText);
-		}
-		else if (ColumnName == ColumnId_VerbosityLabel)
-		{
-			return SNew(STextBlock)
-				.ColorAndOpacity(FSlateColor(Item->Verbosity == ELogVerbosity::Error ? FLinearColor::Red : (Item->Verbosity == ELogVerbosity::Warning ? FLinearColor::Yellow : FLinearColor::Gray)))
-				.Text(FText::FromString(FString(TEXT("(")) + FString(::ToString(Item->Verbosity)) + FString(TEXT(")"))));
-		}
-		else if (ColumnName == ColumnId_LogLabel)
-		{
-			return SNew(STextBlock)
-				.AutoWrapText(true)
-				.ColorAndOpacity(FSlateColor(Item->Verbosity == ELogVerbosity::Error ? FLinearColor::Red : (Item->Verbosity == ELogVerbosity::Warning ? FLinearColor::Yellow : FLinearColor::Gray)))
-				.Text(FText::FromString(Item->Line))
-				.HighlightText(this, &SVisualLoggerTableRow::GetFilterText)
-				.TextStyle(FLogVisualizerStyle::Get(), TEXT("TextLogs.Text"));
-		}
-
-		return SNullWidget::NullWidget;
-	}
-
-	FText GetFilterText() const
-	{
-		static FText NoText;
-		const bool bSearchInsideLogs = GetDefault<ULogVisualizerSettings>()->bSearchInsideLogs;
-		return bSearchInsideLogs ? FText::FromString(FVisualLoggerFilters::Get().GetSearchString()) : NoText;
-	}
-
-private:
-	FLogEntryItemPtr Item;
-};
-
-const FName SVisualLoggerTableRow::ColumnId_CategoryLabel = FName("Category");
-const FName SVisualLoggerTableRow::ColumnId_VerbosityLabel = FName("Verbosity");
-const FName SVisualLoggerTableRow::ColumnId_LogLabel = FName("Log");
-
 //----------------------------------------------------------------------//
 // SVisualLoggerLogsList
 //----------------------------------------------------------------------//
@@ -166,7 +93,11 @@ void SVisualLoggerLogsList::OnFiltersChanged()
 
 TSharedRef<ITableRow> SVisualLoggerLogsList::LogEntryLinesGenerateRow(TSharedPtr<FLogEntryItem> Item, const TSharedRef<STableViewBase>& OwnerTable)
 {
-	return SNew(SVisualLoggerTableRow, OwnerTable, Item);
+	return SNew(SVisualLoggerTableRow, OwnerTable, Item)
+		.FilterText_Lambda([]
+			{
+				return GetDefault<ULogVisualizerSettings>()->bSearchInsideLogs ? FText::FromString(FVisualLoggerFilters::Get().GetSearchString()) : FText::GetEmpty();
+			});
 }
 
 void SVisualLoggerLogsList::OnFiltersSearchChanged(const FText& Filter)
