update anim next trace to trace object references as Ids

diff --git a/Engine/Plugins/Animation/GameplayInsights/Source/GameplayInsights/Private/ObjectAsTraceIdProxyArchiveReader.cpp b/Engine/Plugins/Animation/GameplayInsights/Source/GameplayInsights/Private/ObjectAsTraceIdProxyArchiveReader.cpp
new file mode 100644
--- /dev/null
+++ b/Engine/Plugins/Animation/GameplayInsights/Source/GameplayInsights/Private/ObjectAsTraceIdProxyArchiveReader.cpp
@@ -0,0 +1,105 @@
+// Copyright Epic Games, Inc. All Rights Reserved.
+
+#include "ObjectAsTraceIdProxyArchiveReader.h"
+
+#if WITH_ENGINE
+
+#include "Serialization/Archive.h"
+#include "UObject/UObjectGlobals.h"
+#include "UObject/Object.h"
+#include "Serialization/SerializedPropertyScope.h"
+#include "UObject/ObjectRedirector.h"
+#include "UObject/UnrealType.h"
+#include "Serialization/ObjectAndNameAsStringProxyArchive.h"
+#include "Serialization/ArchiveReplaceObjectRef.h"
+#include "Misc/PackageName.h"
+#include "ObjectTrace.h"
+#include "IGameplayProvider.h"
+#include "UObject/ConstructorHelpers.h"
+
+
+/*----------------------------------------------------------------------------
+	FObjectAsTraceIdProxyArchiveReader.
+----------------------------------------------------------------------------*/
+
+FObjectAsTraceIdProxyArchiveReader::FObjectAsTraceIdProxyArchiveReader( FArchive& InInnerArchive, const IGameplayProvider* InGameplayProvider)
+	: FObjectAsTraceIdProxyArchive(InInnerArchive),
+		GameplayProvider(InGameplayProvider)
+{
+}
+
+FObjectAsTraceIdProxyArchiveReader::~FObjectAsTraceIdProxyArchiveReader() = default;
+
+
+/**
+ * Serialize the given UObject* as a Traced Object Id
+ */
+FArchive& FObjectAsTraceIdProxyArchiveReader::operator<<(UObject*& Obj)
+{
+	if (IsLoading())
+	{
+		uint64 Id = 0;
+		InnerArchive << Id;
+
+		
+		if (const FClassInfo* ClassInfo = GameplayProvider->FindClassInfo(Id))
+		{
+			TSoftObjectPtr<UClass> Class;
+			Class = FSoftObjectPath(ClassInfo->PathName);
+
+			if (Class.LoadSynchronous())
+			{
+				Obj = Class.Get();
+			}
+		}
+		else if (const FObjectInfo* ObjectInfo = GameplayProvider->FindObjectInfo(Id))
+		{
+			// look up the object by fully qualified pathname
+			Obj = FindObject<UObject>(nullptr, ObjectInfo->PathName);
+			// If we couldn't find it, and we want to load it, do that
+			if(!Obj)
+			{
+				Obj = LoadObject<UObject>(nullptr, ObjectInfo->PathName);
+			}			
+		}
+		else
+		{
+			Obj = nullptr;
+		}
+	}
+	else 
+	{
+		Write(Obj);
+	}
+	return *this;
+}
+
+
+FArchive& FObjectAsTraceIdProxyArchiveReader::operator<<(FWeakObjectPtr& Obj)
+{
+	return FArchiveUObject::SerializeWeakObjectPtr(*this, Obj);
+}
+
+FArchive& FObjectAsTraceIdProxyArchiveReader::operator<<(FSoftObjectPtr& Value)
+{
+	if (IsLoading())
+	{
+		// Reset before serializing to clear the internal weak pointer. 
+		Value.ResetWeakPtr();
+	}
+	*this << Value.GetUniqueID();
+	return *this;
+}
+
+FArchive& FObjectAsTraceIdProxyArchiveReader::operator<<(FSoftObjectPath& Value)
+{
+	Value.SerializePath(*this);
+	return *this;
+}
+
+FArchive& FObjectAsTraceIdProxyArchiveReader::operator<<(FObjectPtr& Obj)
+{
+	return FArchiveUObject::SerializeObjectPtr(*this, Obj);
+}
+
+#endif
\ No newline at end of file
diff --git a/Engine/Plugins/Animation/GameplayInsights/Source/GameplayInsights/Public/ObjectAsTraceIdProxyArchiveReader.h b/Engine/Plugins/Animation/GameplayInsights/Source/GameplayInsights/Public/ObjectAsTraceIdProxyArchiveReader.h
new file mode 100644
--- /dev/null
+++ b/Engine/Plugins/Animation/GameplayInsights/Source/GameplayInsights/Public/ObjectAsTraceIdProxyArchiveReader.h
@@ -0,0 +1,43 @@
+// Copyright Epic Games, Inc. All Rights Reserved.
+
+#pragma once
+
+#include "CoreMinimal.h"
+
+#if WITH_ENGINE
+#include "ObjectAsTraceIdProxyArchive.h"
+
+class UObject;
+
+/**
+ * Subclass of FObjectAsTraceIdProxyArchive which is also able to read object references as trace ids and resolve them
+ *
+ * Objects will only resolve if they are found by name based on the traced name data for the object - this will work for Assets and Classes
+ * and may in some cases work for runtime UObject references if those UObjects still exist (eg during a rewind debugger recording in PIE)
+ * 
+ * @param InInnerArchive The actual FArchive object to serialize normal data types (FStrings, INTs, etc)
+ */
+
+GAMEPLAYINSIGHTS_API struct FObjectAsTraceIdProxyArchiveReader : public FObjectAsTraceIdProxyArchive
+{
+	/**
+	 * Creates and initializes a new instance.
+	 *
+	 * @param InInnerArchive - The inner archive to proxy.
+	 * @param bInLoadIfFindFails - Indicates whether to try and load a ref'd object if we don't find it
+	 */
+	GAMEPLAYINSIGHTS_API FObjectAsTraceIdProxyArchiveReader(FArchive& InInnerArchive, const class IGameplayProvider* InGameplayProvider);
+
+	GAMEPLAYINSIGHTS_API virtual ~FObjectAsTraceIdProxyArchiveReader();
+
+	virtual FArchive& operator<<(UObject*& Obj) override;
+	virtual FArchive& operator<<(FWeakObjectPtr& Obj) override;
+	virtual FArchive& operator<<(FSoftObjectPtr& Value) override;
+	virtual FArchive& operator<<(FSoftObjectPath& Value) override;
+	virtual FArchive& operator<<(FObjectPtr& Obj) override;
+
+	const class IGameplayProvider* GameplayProvider;
+};
+
+#endif
+
diff --git a/Engine/Plugins/Experimental/AnimNext/Source/AnimNext/Private/RewindDebugger/AnimNextTrace.cpp b/Engine/Plugins/Experimental/AnimNext/Source/AnimNext/Private/RewindDebugger/AnimNextTrace.cpp
--- a/Engine/Plugins/Experimental/AnimNext/Source/AnimNext/Private/RewindDebugger/AnimNextTrace.cpp
+++ b/Engine/Plugins/Experimental/AnimNext/Source/AnimNext/Private/RewindDebugger/AnimNextTrace.cpp
@@ -8,7 +8,7 @@
 #include "Misc/HashBuilder.h"
 #include "Module/AnimNextModuleInstance.h"
 #include "Serialization/MemoryWriter.h"
-#include "Serialization/ObjectAndNameAsStringProxyArchive.h"
+#include "ObjectAsTraceIdProxyArchive.h"
 #include "Serialization/ObjectWriter.h"
 
 #if ANIMNEXT_TRACE_ENABLED
@@ -204,7 +204,7 @@ void FAnimNextTrace::OutputAnimNextVariableSet(const TSharedRef<FUAFInstanceVari
 						PropertyDescriptions = PropertyDescs;	// need to copy since we can't serialize an array view
 
 						FMemoryWriter WriterArchive(ArchiveData);
-						FObjectAndNameAsStringProxyArchive Archive(WriterArchive, /*bInLoadIfFindFails*/true);
+						FObjectAsTraceIdProxyArchive Archive(WriterArchive);
 						Archive.UsingCustomVersion(CustomVersionGUID);
 						Archive.UsingCustomVersion(FPropertyBagCustomVersion::GUID);
 						Archive << PropertyDescriptions;
@@ -217,7 +217,7 @@ void FAnimNextTrace::OutputAnimNextVariableSet(const TSharedRef<FUAFInstanceVari
 
 				ArchiveData.SetNum(0, EAllowShrinking::No);
 				FMemoryWriter WriterArchive(ArchiveData);
-				FObjectAndNameAsStringProxyArchive Archive(WriterArchive, /*bInLoadIfFindFails*/true);
+				FObjectAsTraceIdProxyArchive Archive(WriterArchive);
 				UPropertyBag* PropertyBag = const_cast<UPropertyBag*>(InstancedPropertyBag.GetPropertyBagStruct());
 				if (PropertyBag != nullptr)
 				{
diff --git a/Engine/Plugins/Experimental/AnimNext/Source/AnimNextEditor/Private/RewindDebugger/AnimNextModuleTrack.cpp b/Engine/Plugins/Experimental/AnimNext/Source/AnimNextEditor/Private/RewindDebugger/AnimNextModuleTrack.cpp
--- a/Engine/Plugins/Experimental/AnimNext/Source/AnimNextEditor/Private/RewindDebugger/AnimNextModuleTrack.cpp
+++ b/Engine/Plugins/Experimental/AnimNext/Source/AnimNextEditor/Private/RewindDebugger/AnimNextModuleTrack.cpp
@@ -10,7 +10,7 @@
 #include "Subsystems/AssetEditorSubsystem.h"
 #include "Modules/ModuleManager.h"
 #include "PropertyEditorModule.h"
-#include "Serialization/ObjectAndNameAsStringProxyArchive.h"
+#include "ObjectAsTraceIdProxyArchiveReader.h"
 #include "StructUtils/PropertyBag.h"
 #include "Serialization/MemoryReader.h"
 
@@ -156,7 +156,7 @@ bool FAnimNextModuleTrack::UpdateInternal()
 					TraceServices::FFrame MarkerFrame;
 					if(FramesProvider.GetFrameFromTime(ETraceFrameType::TraceFrameType_Game, CurrentScrubTime, MarkerFrame))
 					{
-						Data->VariablesTimeline.EnumerateEvents(MarkerFrame.StartTime, MarkerFrame.EndTime, [AnimNextProvider, DetailsObject, this](double InStartTime, double InEndTime, uint32 InDepth, const FPropertyVariableData& VariableListData)
+						Data->VariablesTimeline.EnumerateEvents(MarkerFrame.StartTime, MarkerFrame.EndTime, [AnimNextProvider, DetailsObject, AnalysisSession, this](double InStartTime, double InEndTime, uint32 InDepth, const FPropertyVariableData& VariableListData)
 						{
 							// First look up property description
 							const FPropertyDescriptionData* DescriptionData = AnimNextProvider->GetPropertyDescriptionData(VariableListData.DescriptionHash);
@@ -166,10 +166,12 @@ bool FAnimNextModuleTrack::UpdateInternal()
 								return TraceServices::EEventEnumerate::Continue;
 							}
 							
+							const IGameplayProvider* GameplayProvider = AnalysisSession->ReadProvider<IGameplayProvider>("GameplayProvider");
+							
 							{
 								// Load the property descriptions
 								FMemoryReader Reader(DescriptionData->Data);
-								FObjectAndNameAsStringProxyArchive ReaderProxy(Reader, /*bInLoadIfFindFails*/true);
+								FObjectAsTraceIdProxyArchiveReader ReaderProxy(Reader, GameplayProvider);
 								ReaderProxy.UsingCustomVersion(AnimNext::FAnimNextTrace::CustomVersionGUID);
 								ReaderProxy << PropertyDescriptions;
 
@@ -180,7 +182,7 @@ bool FAnimNextModuleTrack::UpdateInternal()
 							{
 								// Load the property values
 								FMemoryReader Reader(VariableListData.ValueData);
-								FObjectAndNameAsStringProxyArchive ReaderProxy(Reader, /*bInLoadIfFindFails*/true);
+								FObjectAsTraceIdProxyArchiveReader ReaderProxy(Reader, GameplayProvider);
 
 								UPropertyBag* PropertyBag = const_cast<UPropertyBag*>(DetailsObject->Properties.GetPropertyBagStruct());
 								if (PropertyBag != nullptr)
diff --git a/Engine/Plugins/Experimental/AnimNextAnimGraph/Source/AnimNextAnimGraph/Private/Graph/TraceAnimNextGraphInstances.cpp b/Engine/Plugins/Experimental/AnimNextAnimGraph/Source/AnimNextAnimGraph/Private/Graph/TraceAnimNextGraphInstances.cpp
--- a/Engine/Plugins/Experimental/AnimNextAnimGraph/Source/AnimNextAnimGraph/Private/Graph/TraceAnimNextGraphInstances.cpp
+++ b/Engine/Plugins/Experimental/AnimNextAnimGraph/Source/AnimNextAnimGraph/Private/Graph/TraceAnimNextGraphInstances.cpp
@@ -7,8 +7,8 @@
 #include "Graph/AnimNextGraphInstance.h"
 #include "Module/AnimNextModuleInstance.h"
 #include "ObjectTrace.h"
+#include "ObjectAsTraceIdProxyArchive.h"
 #include "Serialization/MemoryWriter.h"
-#include "Serialization/ObjectAndNameAsStringProxyArchive.h"
 #include "TraitCore/ExecutionContext.h"
 #include "TraitCore/NodeInstance.h"
 #include "TraitInterfaces/IHierarchy.h"
@@ -104,7 +104,7 @@ namespace UE::AnimNext
 
 		TArray<uint8> ArchiveData;
 		FMemoryWriter WriterArchive(ArchiveData);
-		FObjectAndNameAsStringProxyArchive Archive(WriterArchive, /*bInLoadIfFindFails*/true);
+		FObjectAsTraceIdProxyArchive Archive(WriterArchive);
 		
 		static const FSerializableEvaluationProgram Defaults;
 		FSerializableEvaluationProgram SerializableProgram(Program);
diff --git a/Engine/Plugins/Experimental/AnimNextAnimGraph/Source/AnimNextAnimGraphEditor/Private/RewindDebugger/EvaluationProgramTrack.cpp b/Engine/Plugins/Experimental/AnimNextAnimGraph/Source/AnimNextAnimGraphEditor/Private/RewindDebugger/EvaluationProgramTrack.cpp
--- a/Engine/Plugins/Experimental/AnimNextAnimGraph/Source/AnimNextAnimGraphEditor/Private/RewindDebugger/EvaluationProgramTrack.cpp
+++ b/Engine/Plugins/Experimental/AnimNextAnimGraph/Source/AnimNextAnimGraphEditor/Private/RewindDebugger/EvaluationProgramTrack.cpp
@@ -9,7 +9,7 @@
 #include "Subsystems/AssetEditorSubsystem.h"
 #include "Modules/ModuleManager.h"
 #include "PropertyEditorModule.h"
-#include "Serialization/ObjectAndNameAsStringProxyArchive.h"
+#include "ObjectAsTraceIdProxyArchiveReader.h"
 #include "Serialization/ObjectReader.h"
 
 
@@ -126,15 +126,16 @@ bool FEvaluationProgramTrack::UpdateInternal()
 					PreviousScrubTime = CurrentScrubTime;
 					
 					const TraceServices::IFrameProvider& FramesProvider = TraceServices::ReadFrameProvider(*AnalysisSession);
+					const IGameplayProvider* GameplayProvider = AnalysisSession->ReadProvider<IGameplayProvider>("GameplayProvider");
 					TraceServices::FAnalysisSessionReadScope SessionReadScope(*AnalysisSession);
 
 					TraceServices::FFrame MarkerFrame;
 					if(FramesProvider.GetFrameFromTime(ETraceFrameType::TraceFrameType_Game, CurrentScrubTime, MarkerFrame))
 					{
-						Data->EvaluationProgramTimeline.EnumerateEvents(MarkerFrame.StartTime, MarkerFrame.EndTime, [DetailsObject](double InStartTime, double InEndTime, uint32 InDepth, const TArray<uint8>& VariableData)
+						Data->EvaluationProgramTimeline.EnumerateEvents(MarkerFrame.StartTime, MarkerFrame.EndTime, [DetailsObject, GameplayProvider](double InStartTime, double InEndTime, uint32 InDepth, const TArray<uint8>& VariableData)
 						{				
 							FMemoryReader Reader(VariableData);
-							FObjectAndNameAsStringProxyArchive Archive(Reader, /*bInLoadIfFindFails*/true);
+							FObjectAsTraceIdProxyArchiveReader Archive(Reader, GameplayProvider);
 
 							static const FSerializableEvaluationProgram Defaults;
 							FSerializableEvaluationProgram::StaticStruct()->SerializeItem(Archive, &DetailsObject->Program, &Defaults);
diff --git a/Engine/Plugins/Experimental/AnimNextAnimGraph/Source/AnimNextAnimGraphEditor/Private/RewindDebugger/SequenceInfoTrack.cpp b/Engine/Plugins/Experimental/AnimNextAnimGraph/Source/AnimNextAnimGraphEditor/Private/RewindDebugger/SequenceInfoTrack.cpp
--- a/Engine/Plugins/Experimental/AnimNextAnimGraph/Source/AnimNextAnimGraphEditor/Private/RewindDebugger/SequenceInfoTrack.cpp
+++ b/Engine/Plugins/Experimental/AnimNextAnimGraph/Source/AnimNextAnimGraphEditor/Private/RewindDebugger/SequenceInfoTrack.cpp
@@ -11,7 +11,7 @@
 #include "PropertyEditorModule.h"
 #include "Animation/AnimSequence.h"
 #include "EvaluationVM/Tasks/PushAnimSequenceKeyframe.h"
-#include "Serialization/ObjectAndNameAsStringProxyArchive.h"
+#include "ObjectAsTraceIdProxyArchiveReader.h"
 #include "Widgets/SCanvas.h"
 #include "Widgets/Text/STextBlock.h"
 #include "Widgets/Colors/SColorBlock.h"
@@ -307,18 +307,18 @@ bool FSequenceInfoTrack::UpdateInternal()
 					PreviousScrubTime = CurrentScrubTime;
 
 					const TraceServices::IFrameProvider& FramesProvider = TraceServices::ReadFrameProvider(*AnalysisSession);
+					const IGameplayProvider* GameplayProvider = AnalysisSession->ReadProvider<IGameplayProvider>("GameplayProvider");
 					TraceServices::FAnalysisSessionReadScope SessionReadScope(*AnalysisSession);
 
 					TraceServices::FFrame MarkerFrame;
 					if (FramesProvider.GetFrameFromTime(ETraceFrameType::TraceFrameType_Game, CurrentScrubTime, MarkerFrame))
 					{
 						Data->EvaluationProgramTimeline.EnumerateEvents(MarkerFrame.StartTime, MarkerFrame.EndTime,
-						[DetailsObject, this] ( double InStartTime, double InEndTime, uint32 InDepth, const TArray<uint8>& VariableData)
+						[DetailsObject, GameplayProvider, this] ( double InStartTime, double InEndTime, uint32 InDepth, const TArray<uint8>& VariableData)
 						{
 							// Why do this if panel is not selected?
 							FMemoryReader Reader(VariableData);
-							FObjectAndNameAsStringProxyArchive Archive(
-							Reader, /*bInLoadIfFindFails*/true);
+							FObjectAsTraceIdProxyArchiveReader Archive(Reader, GameplayProvider);
 
 							static const FSerializableEvaluationProgram Defaults;
 							FSerializableEvaluationProgram Program;
diff --git a/Engine/Source/Runtime/Engine/Private/ObjectAsTraceIdProxyArchive.cpp b/Engine/Source/Runtime/Engine/Private/ObjectAsTraceIdProxyArchive.cpp
new file mode 100644
--- /dev/null
+++ b/Engine/Source/Runtime/Engine/Private/ObjectAsTraceIdProxyArchive.cpp
@@ -0,0 +1,91 @@
+// Copyright Epic Games, Inc. All Rights Reserved.
+
+#include "ObjectAsTraceIdProxyArchive.h"
+#include "Serialization/Archive.h"
+#include "UObject/UObjectGlobals.h"
+#include "UObject/Object.h"
+#include "Serialization/SerializedPropertyScope.h"
+#include "UObject/ObjectRedirector.h"
+#include "UObject/UnrealType.h"
+#include "Serialization/ArchiveReplaceObjectRef.h"
+#include "Misc/PackageName.h"
+#include "ObjectTrace.h"
+
+/*----------------------------------------------------------------------------
+	FObjectAsTraceIdProxyArchive.
+----------------------------------------------------------------------------*/
+
+FObjectAsTraceIdProxyArchive::FObjectAsTraceIdProxyArchive( FArchive& InInnerArchive)
+	: FNameAsStringProxyArchive(InInnerArchive)
+{
+}
+
+FObjectAsTraceIdProxyArchive::~FObjectAsTraceIdProxyArchive() = default;
+
+/**
+ * Serialize the given UObject* as a Traced Object Id
+ */
+void FObjectAsTraceIdProxyArchive::Write(const UObject* Obj)
+{
+#if OBJECT_TRACE_ENABLED
+	if (Obj)
+	{
+		// first ensure the object has been traced 
+		if (const UClass* Class = Cast<UClass>(Obj))
+		{
+			TRACE_CLASS(Class);
+		}
+		TRACE_OBJECT(Obj);
+		// then serialize the Object as a Trace Object id;
+		uint64 Id = FObjectTrace::GetObjectId(Obj);
+		InnerArchive << Id;
+	}
+	else
+#endif
+	{
+		// for null pointer, output 0
+		uint64 zero = 0;
+		InnerArchive << zero;
+	}
+}
+
+FArchive& FObjectAsTraceIdProxyArchive::operator<<(UObject*& Obj)
+{
+	if (IsLoading())
+	{
+		// Loading isn't implemented for this class since we can't depend on GameplayProvider from here
+	}
+	else
+	{
+		Write(Obj);
+	}
+	return *this;
+}
+
+
+FArchive& FObjectAsTraceIdProxyArchive::operator<<(FWeakObjectPtr& Obj)
+{
+	return FArchiveUObject::SerializeWeakObjectPtr(*this, Obj);
+}
+
+FArchive& FObjectAsTraceIdProxyArchive::operator<<(FSoftObjectPtr& Value)
+{
+	if (IsLoading())
+	{
+		// Reset before serializing to clear the internal weak pointer. 
+		Value.ResetWeakPtr();
+	}
+	*this << Value.GetUniqueID();
+	return *this;
+}
+
+FArchive& FObjectAsTraceIdProxyArchive::operator<<(FSoftObjectPath& Value)
+{
+	Value.SerializePath(*this);
+	return *this;
+}
+
+FArchive& FObjectAsTraceIdProxyArchive::operator<<(FObjectPtr& Obj)
+{
+	return FArchiveUObject::SerializeObjectPtr(*this, Obj);
+}
\ No newline at end of file
diff --git a/Engine/Source/Runtime/Engine/Public/ObjectAsTraceIdProxyArchive.h b/Engine/Source/Runtime/Engine/Public/ObjectAsTraceIdProxyArchive.h
new file mode 100644
--- /dev/null
+++ b/Engine/Source/Runtime/Engine/Public/ObjectAsTraceIdProxyArchive.h
@@ -0,0 +1,43 @@
+// Copyright Epic Games, Inc. All Rights Reserved.
+
+#pragma once
+
+#include "CoreMinimal.h"
+#include "Serialization/NameAsStringProxyArchive.h"
+
+class UObject;
+
+/**
+ * Implements a proxy archive that serializes UObjects as Trace Ids with FObjectTrace
+ * This requires that a trace is active and the Object trace channel is enabled.
+ * This archive only supports writing, and not reading, since this archive needs to work in game runtime code, and ids can only be resolved into objects using editor code dependencies.
+ * To read data serialized by this archive, use FObjectAsTraceIdProxyArchiveReader from the GameplayInsights extension
+ *
+ * Expected use is:
+ *    FArchive* SomeAr = CreateAnAr();
+ *    ObjectAsTraceIdProxyArchive Ar(*SomeAr);
+ *    SomeObject->Serialize(Ar);
+ *    FinalizeAr(SomeAr);
+ * 
+ * @param InInnerArchive The actual FArchive object to serialize normal data types (FStrings, INTs, etc)
+ */
+
+ENGINE_API struct FObjectAsTraceIdProxyArchive : public FNameAsStringProxyArchive
+{
+	/**
+	 * Creates and initializes a new instance.
+	 *
+	 * @param InInnerArchive - The inner archive to proxy.
+	 */
+	ENGINE_API FObjectAsTraceIdProxyArchive(FArchive& InInnerArchive);
+
+	ENGINE_API virtual ~FObjectAsTraceIdProxyArchive() override;
+
+	virtual FArchive& operator<<(UObject*& Obj) override;
+	virtual FArchive& operator<<(FWeakObjectPtr& Obj) override;
+	virtual FArchive& operator<<(FSoftObjectPtr& Value) override;
+	virtual FArchive& operator<<(FSoftObjectPath& Value) override;
+	virtual FArchive& operator<<(FObjectPtr& Obj) override;
+
+	ENGINE_API void Write(const UObject* Object);
+};
\ No newline at end of file
