fix compiler warning

diff --git a/Engine/Plugins/Animation/GameplayInsights/Source/GameplayInsights/Public/ObjectAsTraceIdProxyArchiveReader.h b/Engine/Plugins/Animation/GameplayInsights/Source/GameplayInsights/Public/ObjectAsTraceIdProxyArchiveReader.h
--- a/Engine/Plugins/Animation/GameplayInsights/Source/GameplayInsights/Public/ObjectAsTraceIdProxyArchiveReader.h
+++ b/Engine/Plugins/Animation/GameplayInsights/Source/GameplayInsights/Public/ObjectAsTraceIdProxyArchiveReader.h
@@ -18,7 +18,7 @@ class UObject;
  * @param InInnerArchive The actual FArchive object to serialize normal data types (FStrings, INTs, etc)
  */
 
-GAMEPLAYINSIGHTS_API struct FObjectAsTraceIdProxyArchiveReader : public FObjectAsTraceIdProxyArchive
+struct FObjectAsTraceIdProxyArchiveReader : public FObjectAsTraceIdProxyArchive
 {
 	/**
 	 * Creates and initializes a new instance.
diff --git a/Engine/Source/Runtime/Engine/Public/ObjectAsTraceIdProxyArchive.h b/Engine/Source/Runtime/Engine/Public/ObjectAsTraceIdProxyArchive.h
--- a/Engine/Source/Runtime/Engine/Public/ObjectAsTraceIdProxyArchive.h
+++ b/Engine/Source/Runtime/Engine/Public/ObjectAsTraceIdProxyArchive.h
@@ -22,7 +22,7 @@ class UObject;
  * @param InInnerArchive The actual FArchive object to serialize normal data types (FStrings, INTs, etc)
  */
 
-ENGINE_API struct FObjectAsTraceIdProxyArchive : public FNameAsStringProxyArchive
+struct FObjectAsTraceIdProxyArchive : public FNameAsStringProxyArchive
 {
 	/**
 	 * Creates and initializes a new instance.
