Rewind debugger fix for IsObjectCurrentlyDebugged

diff --git a/Engine/Plugins/Animation/GameplayInsights/Source/RewindDebugger/Private/RewindDebugger.cpp b/Engine/Plugins/Animation/GameplayInsights/Source/RewindDebugger/Private/RewindDebugger.cpp
--- a/Engine/Plugins/Animation/GameplayInsights/Source/RewindDebugger/Private/RewindDebugger.cpp
+++ b/Engine/Plugins/Animation/GameplayInsights/Source/RewindDebugger/Private/RewindDebugger.cpp
@@ -1735,30 +1735,29 @@ TArray<TSharedPtr<FDebugObjectInfo>>& FRewindDebugger::GetDebuggedObjects()
 	return DebuggedObjects;
 }
 
-bool FRewindDebugger::IsObjectCurrentlyDebugged(uint64 ObjectId) const
+bool FRewindDebugger::IsObjectCurrentlyDebugged(uint64 InObjectId) const
 {
-	for (const TSharedPtr<RewindDebugger::FRewindDebuggerTrack>& Track : DebugTracks)
+	if (const TraceServices::IAnalysisSession* Session = GetAnalysisSession())
 	{
-		if (Track->GetUObjectId() == ObjectId)
+		TraceServices::FAnalysisSessionReadScope SessionReadScope(*Session);
+		if (const IGameplayProvider* GameplayProvider = Session->ReadProvider<IGameplayProvider>("GameplayProvider"))
 		{
-			return true;
-		}
-
-		bool Found = false;
-		Track->IterateSubTracks([ObjectId, &Found](const TSharedPtr<RewindDebugger::FRewindDebuggerTrack>& Child)
+			uint64 ObjectId = InObjectId;
+			while (ObjectId != RewindDebugger::FObjectId::InvalidId)
 			{
-				if (Child->GetUObjectId() == ObjectId)
+				for (const TSharedPtr<RewindDebugger::FRewindDebuggerTrack>& Track : DebugTracks)
 				{
-					Found = true;
-					// @Todo STDBG: want to stop iteration here
+					if (Track->GetUObjectId() == ObjectId)
+					{
+						return true;
+					}
 				}
-			});
 
-		if (Found)
-		{
-			return true;
+				ObjectId = GameplayProvider->GetObjectInfo(ObjectId).GetOuterUObjectId();
+			}
 		}
 	}
+
 	return false;
 }
 
