[RewindDebugger] deprecated IterateSubTracksInternal and replaced it by GetChildrenInternal that can be use to iterate subtracks and to visit all tracks recursively

diff --git a/Engine/Plugins/Animation/GameplayInsights/Source/GameplayInsights/Private/AnimCurvesTrack.cpp b/Engine/Plugins/Animation/GameplayInsights/Source/GameplayInsights/Private/AnimCurvesTrack.cpp
--- a/Engine/Plugins/Animation/GameplayInsights/Source/GameplayInsights/Private/AnimCurvesTrack.cpp
+++ b/Engine/Plugins/Animation/GameplayInsights/Source/GameplayInsights/Private/AnimCurvesTrack.cpp
@@ -27,18 +27,15 @@ TSharedPtr<SWidget> FAnimCurvesTrack::GetDetailsViewInternal()
 	
 }
 
-void FAnimCurvesTrack::IterateSubTracksInternal(TFunction<void(TSharedPtr<FRewindDebuggerTrack> SubTrack)> IteratorFunction)
+TConstArrayView<TSharedPtr<FRewindDebuggerTrack>> FAnimCurvesTrack::GetChildrenInternal(TArray<TSharedPtr<FRewindDebuggerTrack>>& OutTracks) const
 {
 	if (Children.Num() == 0)
 	{
-		IteratorFunction(ChildPlaceholder);
+		return MakeConstArrayView<TSharedPtr<FRewindDebuggerTrack>>(reinterpret_cast<const TSharedPtr<FRewindDebuggerTrack>*>(&ChildPlaceholder), 1);
 	}
-	
-	for(TSharedPtr<FAnimCurveTrack>& Track : Children)
-	{
-		IteratorFunction(Track);
-	}
-};
+
+	return MakeConstArrayView<TSharedPtr<FRewindDebuggerTrack>>(reinterpret_cast<const TSharedPtr<FRewindDebuggerTrack>*>(Children.GetData()), Children.Num());
+}
 
 bool FAnimCurvesTrack::UpdateInternal()
 {
diff --git a/Engine/Plugins/Animation/GameplayInsights/Source/GameplayInsights/Private/AnimCurvesTrack.h b/Engine/Plugins/Animation/GameplayInsights/Source/GameplayInsights/Private/AnimCurvesTrack.h
--- a/Engine/Plugins/Animation/GameplayInsights/Source/GameplayInsights/Private/AnimCurvesTrack.h
+++ b/Engine/Plugins/Animation/GameplayInsights/Source/GameplayInsights/Private/AnimCurvesTrack.h
@@ -58,7 +58,7 @@ private:
 	virtual FName GetNameInternal() const override { return "AnimationCurves"; }
 	virtual FText GetDisplayNameInternal() const override { return NSLOCTEXT("RewindDebugger", "AnimationCurvesTrackName", "Animation Curves"); }
 	virtual uint64 GetObjectIdInternal() const override { return ObjectId; }
-	virtual void IterateSubTracksInternal(TFunction<void(TSharedPtr<FRewindDebuggerTrack> SubTrack)> IteratorFunction) override;
+	virtual TConstArrayView<TSharedPtr<FRewindDebuggerTrack>> GetChildrenInternal(TArray<TSharedPtr<FRewindDebuggerTrack>>& OutTracks) const override;
 	
 	FSlateIcon Icon;
 	uint64 ObjectId;
diff --git a/Engine/Plugins/Animation/GameplayInsights/Source/GameplayInsights/Private/BlendWeightsTrack.cpp b/Engine/Plugins/Animation/GameplayInsights/Source/GameplayInsights/Private/BlendWeightsTrack.cpp
--- a/Engine/Plugins/Animation/GameplayInsights/Source/GameplayInsights/Private/BlendWeightsTrack.cpp
+++ b/Engine/Plugins/Animation/GameplayInsights/Source/GameplayInsights/Private/BlendWeightsTrack.cpp
@@ -32,13 +32,10 @@ FBlendWeightsTrack::FBlendWeightsTrack(uint64 InObjectId) : ObjectId(InObjectId)
 #endif
 }
 
-void FBlendWeightsTrack::IterateSubTracksInternal(TFunction<void(TSharedPtr<FRewindDebuggerTrack> SubTrack)> IteratorFunction)
+TConstArrayView<TSharedPtr<FRewindDebuggerTrack>> FBlendWeightsTrack::GetChildrenInternal(TArray<TSharedPtr<FRewindDebuggerTrack>>& OutTracks) const
 {
-	for(TSharedPtr<FBlendWeightTrack>& Track : Children)
-	{
-		IteratorFunction(Track);
-	}
-};
+	return MakeConstArrayView<TSharedPtr<FRewindDebuggerTrack>>(reinterpret_cast<const TSharedPtr<FRewindDebuggerTrack>*>(Children.GetData()), Children.Num());
+}
 
 bool FBlendWeightsTrack::UpdateInternal()
 {
@@ -133,14 +130,10 @@ FBlendWeightTrack::FBlendWeightTrack(uint64 InObjectId, uint64 InAssetId, uint32
 	SetIsExpanded(false);
 }
 
-
-void FBlendWeightTrack::IterateSubTracksInternal(TFunction<void(TSharedPtr<FRewindDebuggerTrack> SubTrack)> IteratorFunction)
+TConstArrayView<TSharedPtr<FRewindDebuggerTrack>> FBlendWeightTrack::GetChildrenInternal(TArray<TSharedPtr<FRewindDebuggerTrack>>& OutTracks) const
 {
-	for (TSharedPtr<FBlendWeightTrack>& Track : Children)
-	{
-		IteratorFunction(Track);
-	}
-};
+	return MakeConstArrayView<TSharedPtr<FRewindDebuggerTrack>>(reinterpret_cast<const TSharedPtr<FRewindDebuggerTrack>*>(Children.GetData()), Children.Num());
+}
 
 TSharedPtr<SCurveTimelineView::FTimelineCurveData> FBlendWeightTrack::GetCurveData() const
 {
diff --git a/Engine/Plugins/Animation/GameplayInsights/Source/GameplayInsights/Private/BlendWeightsTrack.h b/Engine/Plugins/Animation/GameplayInsights/Source/GameplayInsights/Private/BlendWeightsTrack.h
--- a/Engine/Plugins/Animation/GameplayInsights/Source/GameplayInsights/Private/BlendWeightsTrack.h
+++ b/Engine/Plugins/Animation/GameplayInsights/Source/GameplayInsights/Private/BlendWeightsTrack.h
@@ -35,8 +35,8 @@ private:
 	virtual bool UpdateInternal() override;
 	virtual TSharedPtr<SWidget> GetTimelineViewInternal() override;
 	virtual TSharedPtr<SWidget> GetDetailsViewInternal() override;
-	virtual void IterateSubTracksInternal(TFunction<void(TSharedPtr<FRewindDebuggerTrack> SubTrack)> IteratorFunction) override;
-	
+	virtual TConstArrayView<TSharedPtr<FRewindDebuggerTrack>> GetChildrenInternal(TArray<TSharedPtr<FRewindDebuggerTrack>>& OutTracks) const override;
+
 	virtual FSlateIcon GetIconInternal() override { return Icon; }
 	virtual FName GetNameInternal() const override { return "BlendWeight"; }
 	virtual FText GetDisplayNameInternal() const override { return CurveName; }
@@ -68,8 +68,8 @@ public:
 private:
 	virtual bool UpdateInternal() override;
 	virtual TSharedPtr<SWidget> GetDetailsViewInternal() override;
-	virtual void IterateSubTracksInternal(TFunction<void(TSharedPtr<FRewindDebuggerTrack> SubTrack)> IteratorFunction) override;
-	
+	virtual TConstArrayView<TSharedPtr<FRewindDebuggerTrack>> GetChildrenInternal(TArray<TSharedPtr<FRewindDebuggerTrack>>& OutTracks) const override;
+
 	virtual FName GetNameInternal() const override { return "BlendWeights"; }
 	virtual FText GetDisplayNameInternal() const override { return NSLOCTEXT("RewindDebugger", "Blend Weights Track Name", "Blend Weights"); }
 	virtual uint64 GetObjectIdInternal() const override { return ObjectId; }
diff --git a/Engine/Plugins/Animation/GameplayInsights/Source/GameplayInsights/Private/ExternalMorphTrack.cpp b/Engine/Plugins/Animation/GameplayInsights/Source/GameplayInsights/Private/ExternalMorphTrack.cpp
--- a/Engine/Plugins/Animation/GameplayInsights/Source/GameplayInsights/Private/ExternalMorphTrack.cpp
+++ b/Engine/Plugins/Animation/GameplayInsights/Source/GameplayInsights/Private/ExternalMorphTrack.cpp
@@ -18,13 +18,10 @@ namespace RewindDebugger
 		Icon = FSlateIcon("EditorStyle", "AnimGraph.Attribute.Curves.Icon", "AnimGraph.Attribute.Curves.Icon");
 	}
 
-	void FExternalMorphSetGroupTrack::IterateSubTracksInternal(TFunction<void(TSharedPtr<FRewindDebuggerTrack> SubTrack)> IteratorFunction)
+	TConstArrayView<TSharedPtr<FRewindDebuggerTrack>> FExternalMorphSetGroupTrack::GetChildrenInternal(TArray<TSharedPtr<FRewindDebuggerTrack>>& OutTracks) const
 	{
-		for (TSharedPtr<FExternalMorphSetTrack>& Track : Children)
-		{
-			IteratorFunction(Track);
-		}
-	};
+		return MakeConstArrayView<TSharedPtr<FRewindDebuggerTrack>>(reinterpret_cast<const TSharedPtr<FRewindDebuggerTrack>*>(Children.GetData()), Children.Num());
+	}
 
 	bool FExternalMorphSetGroupTrack::UpdateInternal()
 	{
@@ -102,13 +99,10 @@ namespace RewindDebugger
 	
 	}
 
-	void FExternalMorphSetTrack::IterateSubTracksInternal(TFunction<void(TSharedPtr<FRewindDebuggerTrack> SubTrack)> IteratorFunction)
+	TConstArrayView<TSharedPtr<FRewindDebuggerTrack>> FExternalMorphSetTrack::GetChildrenInternal(TArray<TSharedPtr<FRewindDebuggerTrack>>& OutTracks) const
 	{
-		for (TSharedPtr<FExternalMorphTrack>& Track : Children)
-		{
-			IteratorFunction(Track);
-		}
-	};
+		return MakeConstArrayView<TSharedPtr<FRewindDebuggerTrack>>(reinterpret_cast<const TSharedPtr<FRewindDebuggerTrack>*>(Children.GetData()), Children.Num());
+	}
 
 	bool FExternalMorphSetTrack::UpdateInternal()
 	{
diff --git a/Engine/Plugins/Animation/GameplayInsights/Source/GameplayInsights/Private/ExternalMorphTrack.h b/Engine/Plugins/Animation/GameplayInsights/Source/GameplayInsights/Private/ExternalMorphTrack.h
--- a/Engine/Plugins/Animation/GameplayInsights/Source/GameplayInsights/Private/ExternalMorphTrack.h
+++ b/Engine/Plugins/Animation/GameplayInsights/Source/GameplayInsights/Private/ExternalMorphTrack.h
@@ -61,7 +61,7 @@ namespace RewindDebugger
 		virtual FName GetNameInternal() const override						{ return "ExternalMorphSet"; }
 		virtual FText GetDisplayNameInternal() const override				{ return NSLOCTEXT("RewindDebugger", "ExternalMorphSetTrackName", "External Morph Set"); }
 		virtual uint64 GetObjectIdInternal() const override					{ return ObjectId; }
-		virtual void IterateSubTracksInternal(TFunction<void(TSharedPtr<FRewindDebuggerTrack> SubTrack)> IteratorFunction) override;
+		virtual TConstArrayView<TSharedPtr<FRewindDebuggerTrack>> GetChildrenInternal(TArray<TSharedPtr<FRewindDebuggerTrack>>& OutTracks) const override;
 		// ~END FRewindDebuggerTrack overrides.
 	
 	protected:
@@ -86,7 +86,7 @@ namespace RewindDebugger
 		virtual FName GetNameInternal() const override						{ return "ExternalMorphSetGroup"; }
 		virtual FText GetDisplayNameInternal() const override				{ return NSLOCTEXT("RewindDebugger", "ExternalMorphSetGroupTrackName", "External Morph Sets"); }
 		virtual uint64 GetObjectIdInternal() const override					{ return ObjectId; }
-		virtual void IterateSubTracksInternal(TFunction<void(TSharedPtr<FRewindDebuggerTrack> SubTrack)> IteratorFunction) override;
+		virtual TConstArrayView<TSharedPtr<FRewindDebuggerTrack>> GetChildrenInternal(TArray<TSharedPtr<FRewindDebuggerTrack>>& OutTracks) const override;
 		// ~END FRewindDebuggerTrack overrides.
 	
 	protected:
diff --git a/Engine/Plugins/Animation/GameplayInsights/Source/GameplayInsights/Private/InertializationsTrack.cpp b/Engine/Plugins/Animation/GameplayInsights/Source/GameplayInsights/Private/InertializationsTrack.cpp
--- a/Engine/Plugins/Animation/GameplayInsights/Source/GameplayInsights/Private/InertializationsTrack.cpp
+++ b/Engine/Plugins/Animation/GameplayInsights/Source/GameplayInsights/Private/InertializationsTrack.cpp
@@ -27,13 +27,10 @@ FInertializationsTrack::FInertializationsTrack(uint64 InObjectId) : ObjectId(InO
 	Icon = FSlateIcon("EditorStyle", "AnimGraph.Attribute.InertialBlending.Icon", "AnimGraph.Attribute.InertialBlending.Icon");
 }
 
-void FInertializationsTrack::IterateSubTracksInternal(TFunction<void(TSharedPtr<FRewindDebuggerTrack> SubTrack)> IteratorFunction)
+TConstArrayView<TSharedPtr<FRewindDebuggerTrack>> FInertializationsTrack::GetChildrenInternal(TArray<TSharedPtr<FRewindDebuggerTrack>>& OutTracks) const
 {
-	for(TSharedPtr<FInertializationTrack>& Track : Children)
-	{
-		IteratorFunction(Track);
-	}
-};
+	return MakeConstArrayView<TSharedPtr<FRewindDebuggerTrack>>(reinterpret_cast<const TSharedPtr<FRewindDebuggerTrack>*>(Children.GetData()), Children.Num());
+}
 
 bool FInertializationsTrack::UpdateInternal()
 {
diff --git a/Engine/Plugins/Animation/GameplayInsights/Source/GameplayInsights/Private/InertializationsTrack.h b/Engine/Plugins/Animation/GameplayInsights/Source/GameplayInsights/Private/InertializationsTrack.h
--- a/Engine/Plugins/Animation/GameplayInsights/Source/GameplayInsights/Private/InertializationsTrack.h
+++ b/Engine/Plugins/Animation/GameplayInsights/Source/GameplayInsights/Private/InertializationsTrack.h
@@ -49,8 +49,8 @@ public:
 private:
 	virtual bool UpdateInternal() override;
 	virtual TSharedPtr<SWidget> GetDetailsViewInternal() override;
-	virtual void IterateSubTracksInternal(TFunction<void(TSharedPtr<FRewindDebuggerTrack> SubTrack)> IteratorFunction) override;
-	
+	virtual TConstArrayView<TSharedPtr<FRewindDebuggerTrack>> GetChildrenInternal(TArray<TSharedPtr<FRewindDebuggerTrack>>& OutTracks) const override;
+
 	virtual FName GetNameInternal() const override { return "Inertializations"; }
 	virtual FText GetDisplayNameInternal() const override { return NSLOCTEXT("RewindDebugger", "Inertializations Track Name", "Inertializations"); }
 	virtual uint64 GetObjectIdInternal() const override { return ObjectId; }
diff --git a/Engine/Plugins/Animation/GameplayInsights/Source/GameplayInsights/Private/MontagesTrack.cpp b/Engine/Plugins/Animation/GameplayInsights/Source/GameplayInsights/Private/MontagesTrack.cpp
--- a/Engine/Plugins/Animation/GameplayInsights/Source/GameplayInsights/Private/MontagesTrack.cpp
+++ b/Engine/Plugins/Animation/GameplayInsights/Source/GameplayInsights/Private/MontagesTrack.cpp
@@ -29,13 +29,10 @@ FMontagesTrack::FMontagesTrack(uint64 InObjectId) : ObjectId(InObjectId)
 #endif
 }
 
-void FMontagesTrack::IterateSubTracksInternal(TFunction<void(TSharedPtr<FRewindDebuggerTrack> SubTrack)> IteratorFunction)
+TConstArrayView<TSharedPtr<FRewindDebuggerTrack>> FMontagesTrack::GetChildrenInternal(TArray<TSharedPtr<FRewindDebuggerTrack>>& OutTracks) const
 {
-	for(TSharedPtr<FMontageTrack>& Track : Children)
-	{
-		IteratorFunction(Track);
-	}
-};
+	return MakeConstArrayView<TSharedPtr<FRewindDebuggerTrack>>(reinterpret_cast<const TSharedPtr<FRewindDebuggerTrack>*>(Children.GetData()), Children.Num());
+}
 
 bool FMontagesTrack::UpdateInternal()
 {
@@ -118,14 +115,10 @@ FMontageTrack::FMontageTrack(uint64 InObjectId, uint64 InAssetId, FMontageTrack:
 	SetIsExpanded(false);
 }
 
-
-void FMontageTrack::IterateSubTracksInternal(TFunction<void(TSharedPtr<FRewindDebuggerTrack> SubTrack)> IteratorFunction)
+TConstArrayView<TSharedPtr<FRewindDebuggerTrack>> FMontageTrack::GetChildrenInternal(TArray<TSharedPtr<FRewindDebuggerTrack>>& OutTracks) const
 {
-	for (TSharedPtr<FMontageTrack>& Track : Children)
-	{
-		IteratorFunction(Track);
-	}
-};
+	return MakeConstArrayView<TSharedPtr<FRewindDebuggerTrack>>(reinterpret_cast<const TSharedPtr<FRewindDebuggerTrack>*>(Children.GetData()), Children.Num());
+}
 
 TSharedPtr<SCurveTimelineView::FTimelineCurveData> FMontageTrack::GetCurveData() const
 {
diff --git a/Engine/Plugins/Animation/GameplayInsights/Source/GameplayInsights/Private/MontagesTrack.h b/Engine/Plugins/Animation/GameplayInsights/Source/GameplayInsights/Private/MontagesTrack.h
--- a/Engine/Plugins/Animation/GameplayInsights/Source/GameplayInsights/Private/MontagesTrack.h
+++ b/Engine/Plugins/Animation/GameplayInsights/Source/GameplayInsights/Private/MontagesTrack.h
@@ -29,9 +29,9 @@ private:
 	virtual bool UpdateInternal() override;
 	virtual TSharedPtr<SWidget> GetTimelineViewInternal() override;
 	virtual TSharedPtr<SWidget> GetDetailsViewInternal() override;
-	virtual void IterateSubTracksInternal(TFunction<void(TSharedPtr<FRewindDebuggerTrack> SubTrack)> IteratorFunction) override;
-    virtual bool HandleDoubleClickInternal() override;
-	
+	virtual TConstArrayView<TSharedPtr<FRewindDebuggerTrack>> GetChildrenInternal(TArray<TSharedPtr<FRewindDebuggerTrack>>& OutTracks) const override;
+	virtual bool HandleDoubleClickInternal() override;
+
 	virtual FSlateIcon GetIconInternal() override { return Icon; }
 	virtual FName GetNameInternal() const override { return "Montage"; }
 	virtual FText GetDisplayNameInternal() const override { return CurveName; }
@@ -61,8 +61,8 @@ public:
 private:
 	virtual bool UpdateInternal() override;
 	virtual TSharedPtr<SWidget> GetDetailsViewInternal() override;
-	virtual void IterateSubTracksInternal(TFunction<void(TSharedPtr<FRewindDebuggerTrack> SubTrack)> IteratorFunction) override;
-	
+	virtual TConstArrayView<TSharedPtr<FRewindDebuggerTrack>> GetChildrenInternal(TArray<TSharedPtr<FRewindDebuggerTrack>>& OutTracks) const override;
+
 	virtual FName GetNameInternal() const override;
 	virtual FText GetDisplayNameInternal() const override { return NSLOCTEXT("RewindDebugger", "Montages Track Name", "Montages"); }
 	virtual uint64 GetObjectIdInternal() const override { return ObjectId; }
diff --git a/Engine/Plugins/Animation/GameplayInsights/Source/GameplayInsights/Private/NotifiesTrack.cpp b/Engine/Plugins/Animation/GameplayInsights/Source/GameplayInsights/Private/NotifiesTrack.cpp
--- a/Engine/Plugins/Animation/GameplayInsights/Source/GameplayInsights/Private/NotifiesTrack.cpp
+++ b/Engine/Plugins/Animation/GameplayInsights/Source/GameplayInsights/Private/NotifiesTrack.cpp
@@ -255,14 +255,11 @@ bool FNotifiesTrack::UpdateInternal()
 	
 	return bChanged;
 }
-	
-void FNotifiesTrack::IterateSubTracksInternal(TFunction<void(TSharedPtr<FRewindDebuggerTrack> SubTrack)> IteratorFunction)
+
+TConstArrayView<TSharedPtr<FRewindDebuggerTrack>> FNotifiesTrack::GetChildrenInternal(TArray<TSharedPtr<FRewindDebuggerTrack>>& OutTracks) const
 {
-	for(TSharedPtr<FNotifyTrack>& Track : Children)
-	{
-		IteratorFunction(Track);
-	}
-};
+	return MakeConstArrayView<TSharedPtr<FRewindDebuggerTrack>>(reinterpret_cast<const TSharedPtr<FRewindDebuggerTrack>*>(Children.GetData()), Children.Num());
+}
 
 bool FNotifiesTrackCreator::HasDebugInfoInternal(const FObjectId& InObjectId) const
 {
diff --git a/Engine/Plugins/Animation/GameplayInsights/Source/GameplayInsights/Private/NotifiesTrack.h b/Engine/Plugins/Animation/GameplayInsights/Source/GameplayInsights/Private/NotifiesTrack.h
--- a/Engine/Plugins/Animation/GameplayInsights/Source/GameplayInsights/Private/NotifiesTrack.h
+++ b/Engine/Plugins/Animation/GameplayInsights/Source/GameplayInsights/Private/NotifiesTrack.h
@@ -87,7 +87,7 @@ private:
 	virtual FName GetNameInternal() const override { return "Notifies"; }
 	virtual FText GetDisplayNameInternal() const override { return NSLOCTEXT("RewindDebugger", "NotifiesTrackName", "Notifies"); }
 	virtual uint64 GetObjectIdInternal() const override { return ObjectId; }
-	virtual void IterateSubTracksInternal(TFunction<void(TSharedPtr<FRewindDebuggerTrack> SubTrack)> IteratorFunction) override;
+	virtual TConstArrayView<TSharedPtr<FRewindDebuggerTrack>> GetChildrenInternal(TArray<TSharedPtr<FRewindDebuggerTrack>>& OutTracks) const override;
 
 	FSlateIcon Icon;
 	uint64 ObjectId;
diff --git a/Engine/Plugins/Animation/GameplayInsights/Source/GameplayInsights/Private/PoseWatchTrack.cpp b/Engine/Plugins/Animation/GameplayInsights/Source/GameplayInsights/Private/PoseWatchTrack.cpp
--- a/Engine/Plugins/Animation/GameplayInsights/Source/GameplayInsights/Private/PoseWatchTrack.cpp
+++ b/Engine/Plugins/Animation/GameplayInsights/Source/GameplayInsights/Private/PoseWatchTrack.cpp
@@ -233,12 +233,9 @@ TSharedPtr<SWidget> FPoseWatchTrack::GetTimelineViewInternal()
 	return TimelineView;
 }
 
-void FPoseWatchTrack::IterateSubTracksInternal(TFunction<void(TSharedPtr<FRewindDebuggerTrack> SubTrack)> IteratorFunction)
+TConstArrayView<TSharedPtr<FRewindDebuggerTrack>> FPoseWatchTrack::GetChildrenInternal(TArray<TSharedPtr<FRewindDebuggerTrack>>& OutTracks) const
 {
-	for(TSharedPtr<FPoseWatchCurveTrack>& Track : Children)
-	{
-		IteratorFunction(Track);
-	}
+	return MakeConstArrayView<TSharedPtr<FRewindDebuggerTrack>>(reinterpret_cast<const TSharedPtr<FRewindDebuggerTrack>*>(Children.GetData()), Children.Num());
 }
 
 FName FPoseWatchesTrackCreator::GetTargetTypeNameInternal() const
@@ -340,14 +337,11 @@ bool FPoseWatchesTrack::UpdateInternal()
 	
 	return bChanged;
 }
-	
-void FPoseWatchesTrack::IterateSubTracksInternal(TFunction<void(TSharedPtr<FRewindDebuggerTrack> SubTrack)> IteratorFunction)
+
+TConstArrayView<TSharedPtr<FRewindDebuggerTrack>> FPoseWatchesTrack::GetChildrenInternal(TArray<TSharedPtr<FRewindDebuggerTrack>>& OutTracks) const
 {
-	for(TSharedPtr<FPoseWatchTrack>& Track : Children)
-	{
-		IteratorFunction(Track);
-	}
-};
+	return MakeConstArrayView<TSharedPtr<FRewindDebuggerTrack>>(reinterpret_cast<const TSharedPtr<FRewindDebuggerTrack>*>(Children.GetData()), Children.Num());
+}
 
 bool FPoseWatchesTrackCreator::HasDebugInfoInternal(const FObjectId& InObjectId) const
 {
@@ -366,7 +360,6 @@ bool FPoseWatchesTrackCreator::HasDebugInfoInternal(const FObjectId& InObjectId)
 	return bHasData;
 }
 
-	
 }
 
 #undef LOCTEXT_NAMESPACE
diff --git a/Engine/Plugins/Animation/GameplayInsights/Source/GameplayInsights/Private/PoseWatchTrack.h b/Engine/Plugins/Animation/GameplayInsights/Source/GameplayInsights/Private/PoseWatchTrack.h
--- a/Engine/Plugins/Animation/GameplayInsights/Source/GameplayInsights/Private/PoseWatchTrack.h
+++ b/Engine/Plugins/Animation/GameplayInsights/Source/GameplayInsights/Private/PoseWatchTrack.h
@@ -45,8 +45,8 @@ private:
 	virtual FSlateIcon GetIconInternal() override { return Icon; }
 	virtual FText GetDisplayNameInternal() const override;
 	virtual uint64 GetObjectIdInternal() const override { return ObjectId; }
-	virtual void IterateSubTracksInternal(TFunction<void(TSharedPtr<FRewindDebuggerTrack> SubTrack)> IteratorFunction) override;
-		
+	virtual TConstArrayView<TSharedPtr<FRewindDebuggerTrack>> GetChildrenInternal(TArray<TSharedPtr<FRewindDebuggerTrack>>& OutTracks) const override;
+
 	TSharedPtr<SSegmentedTimelineView::FSegmentData> GetSegmentData() const;
 	
 	uint64 ObjectId;
@@ -70,7 +70,7 @@ public:
 
 private:
 	virtual bool UpdateInternal() override;
-	virtual void IterateSubTracksInternal(TFunction<void(TSharedPtr<FRewindDebuggerTrack> SubTrack)> IteratorFunction) override;
+	virtual TConstArrayView<TSharedPtr<FRewindDebuggerTrack>> GetChildrenInternal(TArray<TSharedPtr<FRewindDebuggerTrack>>& OutTracks) const override;
 
 	virtual FName GetNameInternal() const override { return "PoseWatches"; }
 	virtual FText GetDisplayNameInternal() const override { return NSLOCTEXT("RewindDebugger", "PoseWatchesTrackName", "Pose Watches"); }
diff --git a/Engine/Plugins/Animation/GameplayInsights/Source/GameplayInsights/Private/PropertiesTrack.cpp b/Engine/Plugins/Animation/GameplayInsights/Source/GameplayInsights/Private/PropertiesTrack.cpp
--- a/Engine/Plugins/Animation/GameplayInsights/Source/GameplayInsights/Private/PropertiesTrack.cpp
+++ b/Engine/Plugins/Animation/GameplayInsights/Source/GameplayInsights/Private/PropertiesTrack.cpp
@@ -252,16 +252,18 @@ namespace RewindDebugger
 		return ObjectId;
 	}
 
-	void FPropertiesTrack::IterateSubTracksInternal(TFunction<void(TSharedPtr<FRewindDebuggerTrack> SubTrack)> IteratorFunction)
+	TConstArrayView<TSharedPtr<FRewindDebuggerTrack>> FPropertiesTrack::GetChildrenInternal(TArray<TSharedPtr<FRewindDebuggerTrack>>& OutTracks) const
 	{
 		// Iterate through all the parent tracks
-		for (const TTuple<uint32, TSharedPtr<FPropertyTrack>> & PropertyWatch : ChildrenTracks)
+		for (const TTuple<uint32, TSharedPtr<FPropertyTrack>>& PropertyWatch : ChildrenTracks)
 		{
 			if (!PropertyWatch.Value->GetObjectProperty()->Property.ParentNameId)
 			{
-				IteratorFunction(PropertyWatch.Value);
+				OutTracks.Push(PropertyWatch.Value);
 			}
 		}
+
+		return {};
 	}
 
 	void FPropertiesTrack::OnPropertyWatched(uint64 InObjectId, uint32 InPropertyNameId)
diff --git a/Engine/Plugins/Animation/GameplayInsights/Source/GameplayInsights/Private/PropertiesTrack.h b/Engine/Plugins/Animation/GameplayInsights/Source/GameplayInsights/Private/PropertiesTrack.h
--- a/Engine/Plugins/Animation/GameplayInsights/Source/GameplayInsights/Private/PropertiesTrack.h
+++ b/Engine/Plugins/Animation/GameplayInsights/Source/GameplayInsights/Private/PropertiesTrack.h
@@ -42,7 +42,7 @@ namespace RewindDebugger
 		virtual FName GetNameInternal() const override;
 		virtual FText GetDisplayNameInternal() const override;
 		virtual uint64 GetObjectIdInternal() const override;
-		virtual void IterateSubTracksInternal(TFunction<void(TSharedPtr<FRewindDebuggerTrack> SubTrack)> IteratorFunction) override;
+		virtual TConstArrayView<TSharedPtr<FRewindDebuggerTrack>> GetChildrenInternal(TArray<TSharedPtr<FRewindDebuggerTrack>>& OutTracks) const override;
 		/** End IRewindDebuggerTrack interface */
 
 		/** Event handler used to create tracks for watched properties */
diff --git a/Engine/Plugins/Animation/GameplayInsights/Source/GameplayInsights/Private/PropertyTrack.cpp b/Engine/Plugins/Animation/GameplayInsights/Source/GameplayInsights/Private/PropertyTrack.cpp
--- a/Engine/Plugins/Animation/GameplayInsights/Source/GameplayInsights/Private/PropertyTrack.cpp
+++ b/Engine/Plugins/Animation/GameplayInsights/Source/GameplayInsights/Private/PropertyTrack.cpp
@@ -220,12 +220,9 @@ namespace RewindDebugger
 		return View;
 	}
 
-	void FPropertyTrack::IterateSubTracksInternal(TFunction<void(TSharedPtr<FRewindDebuggerTrack> SubTrack)> IteratorFunction)
+	TConstArrayView<TSharedPtr<FRewindDebuggerTrack>> FPropertyTrack::GetChildrenInternal(TArray<TSharedPtr<FRewindDebuggerTrack>>& OutTracks) const
 	{
-		for (TSharedPtr<FPropertyTrack>& Track : Children)
-		{
-			IteratorFunction(Track);
-		}
+		return MakeConstArrayView<TSharedPtr<FRewindDebuggerTrack>>(reinterpret_cast<const TSharedPtr<FRewindDebuggerTrack>*>(Children.GetData()), Children.Num());
 	}
 
 	FSlateIcon FPropertyTrack::GetIconInternal()
diff --git a/Engine/Plugins/Animation/GameplayInsights/Source/GameplayInsights/Private/PropertyTrack.h b/Engine/Plugins/Animation/GameplayInsights/Source/GameplayInsights/Private/PropertyTrack.h
--- a/Engine/Plugins/Animation/GameplayInsights/Source/GameplayInsights/Private/PropertyTrack.h
+++ b/Engine/Plugins/Animation/GameplayInsights/Source/GameplayInsights/Private/PropertyTrack.h
@@ -112,7 +112,7 @@ namespace RewindDebugger
 		/** Begin IRewindDebuggerTrack interface */
 		virtual bool UpdateInternal() override;
 		virtual TSharedPtr<SWidget> GetDetailsViewInternal() override;
-		virtual void IterateSubTracksInternal(TFunction<void(TSharedPtr<FRewindDebuggerTrack> SubTrack)> IteratorFunction) override;
+		virtual TConstArrayView<TSharedPtr<FRewindDebuggerTrack>> GetChildrenInternal(TArray<TSharedPtr<FRewindDebuggerTrack>>& OutTracks) const override;
 		virtual FSlateIcon GetIconInternal() override;
 		virtual FName GetNameInternal() const override;
 		virtual FText GetDisplayNameInternal() const override;
diff --git a/Engine/Plugins/Animation/GameplayInsights/Source/RewindDebugger/Private/RewindDebugger.cpp b/Engine/Plugins/Animation/GameplayInsights/Source/RewindDebugger/Private/RewindDebugger.cpp
--- a/Engine/Plugins/Animation/GameplayInsights/Source/RewindDebugger/Private/RewindDebugger.cpp
+++ b/Engine/Plugins/Animation/GameplayInsights/Source/RewindDebugger/Private/RewindDebugger.cpp
@@ -1096,7 +1096,7 @@ void FRewindDebugger::Tick(float DeltaTime)
 
 						const TraceServices::IFrameProvider& FrameProvider = TraceServices::ReadFrameProvider(*Session);
 						TraceServices::FFrame Frame;
-						if (FrameProvider.GetFrameFromTime(ETraceFrameType::TraceFrameType_Game, CurrentTraceTime, Frame))
+						if (FrameProvider.GetFrameFromTime(TraceFrameType_Game, CurrentTraceTime, Frame))
 						{
 							bool bNewActor = false;
 							if (!CandidateIds.Contains(RewindDebugger::FObjectId{TargetActorIdForMesh}))
@@ -1230,7 +1230,7 @@ void FRewindDebugger::RegisterTrackContextMenu()
 
 void FRewindDebugger::MakeOtherWorldsMenu(UToolMenu* Menu)
 {
-	const FRewindDebugger* RewindDebugger = FRewindDebugger::Instance();
+	const FRewindDebugger* RewindDebugger = Instance();
 
 	FToolMenuSection& Section = Menu->AddSection("Other Worlds", LOCTEXT("Other Worlds", "Other Worlds"));
 
@@ -1282,12 +1282,12 @@ void FRewindDebugger::MakeOtherWorldsMenu(UToolMenu* Menu)
 					FSlateIcon(),
 					FUIAction(FExecuteAction::CreateLambda([World = WorldInfo.Id]()
 						{
-							FRewindDebugger::Instance()->SetDisplayWorld(World);
+							Instance()->SetDisplayWorld(World);
 						}),
 						FCanExecuteAction(),
 						FIsActionChecked::CreateLambda([World = WorldInfo.Id]()
 							{
-								return FRewindDebugger::Instance()->DisplayWorldId == World;
+								return Instance()->DisplayWorldId == World;
 							})),
 					EUserInterfaceActionType::Check
 				);
@@ -1308,7 +1308,7 @@ void FRewindDebugger::SetDisplayWorld(uint64 WorldId)
 }
 void FRewindDebugger::MakeWorldsMenu(UToolMenu* Menu)
 {
-	const FRewindDebugger* RewindDebugger = FRewindDebugger::Instance();
+	const FRewindDebugger* RewindDebugger = Instance();
 
 	FToolMenuSection& ServerWorldsSection = Menu->AddSection("Server Worlds", LOCTEXT("Server", "Server"));
 	FToolMenuSection& GameWorldsSection = Menu->AddSection("Game Worlds", LOCTEXT("Game Worlds", "Game Worlds"));
@@ -1318,7 +1318,7 @@ void FRewindDebugger::MakeWorldsMenu(UToolMenu* Menu)
 		LOCTEXT("Other Worlds", "Other Worlds"),
 		LOCTEXT("Other Worlds Tooltip", "Additional worlds such as  Editor Preview worlds"),
 		FNewToolMenuChoice(
-			FNewToolMenuDelegate::CreateStatic(FRewindDebugger::MakeOtherWorldsMenu)
+			FNewToolMenuDelegate::CreateStatic(MakeOtherWorldsMenu)
 		));
 
 	if (const TraceServices::IAnalysisSession* Session = RewindDebugger->GetAnalysisSession())
@@ -1361,12 +1361,12 @@ void FRewindDebugger::MakeWorldsMenu(UToolMenu* Menu)
 					FSlateIcon(),
 					FUIAction(FExecuteAction::CreateLambda([World = WorldInfo.Id]()
 						{
-							FRewindDebugger::Instance()->SetDisplayWorld(World);
+							Instance()->SetDisplayWorld(World);
 						}),
 						FCanExecuteAction(),
 						FIsActionChecked::CreateLambda([World = WorldInfo.Id]()
 							{
-								return FRewindDebugger::Instance()->DisplayWorldId == World;
+								return Instance()->DisplayWorldId == World;
 							})),
 					EUserInterfaceActionType::Check
 				);
@@ -1602,7 +1602,7 @@ void FRewindDebugger::RegisterToolBar()
 		"Display World",
 		FUIAction(
 			FExecuteAction(),
-			FCanExecuteAction::CreateLambda([]() { return FRewindDebugger::Instance()->IsTraceFileLoaded(); })
+			FCanExecuteAction::CreateLambda([]() { return Instance()->IsTraceFileLoaded(); })
 		),
 		FNewToolMenuDelegate::CreateStatic(&FRewindDebugger::MakeWorldsMenu),
 		LOCTEXT("Display World", "Display World"),
@@ -1691,24 +1691,19 @@ TSharedPtr<RewindDebugger::FRewindDebuggerTrack> FRewindDebugger::GetSelectedTra
 
 void FRewindDebugger::SelectTrack(RewindDebugger::FObjectId ObjectId)
 {
-	for (const TSharedPtr<RewindDebugger::FRewindDebuggerTrack>& Track : DebugTracks)
+	using namespace RewindDebugger;
+	for (const TSharedPtr<FRewindDebuggerTrack>& Track : DebugTracks)
 	{
-		if (Track->GetAssociatedObjectId() == ObjectId)
-		{
-			TrackSelectionChanged(Track);
-			break;
-		}
-
-		if (Track->VisitSubTracks([this, ObjectId](const TSharedPtr<RewindDebugger::FRewindDebuggerTrack>& ChildTrack)
+		if (FRewindDebuggerTrack::Visit(Track, [this, ObjectId](const TSharedPtr<FRewindDebuggerTrack>& Track)
 			{
-				if (ChildTrack->GetAssociatedObjectId() == ObjectId)
+				if (Track->GetAssociatedObjectId() == ObjectId)
 				{
-					TrackSelectionChanged(ChildTrack);
-					return RewindDebugger::FRewindDebuggerTrack::EVisitResult::Break;
+					TrackSelectionChanged(Track);
+					return FRewindDebuggerTrack::EVisitResult::Break;
 				}
 
-				return RewindDebugger::FRewindDebuggerTrack::EVisitResult::Continue;
-			}) == RewindDebugger::FRewindDebuggerTrack::EVisitResult::Break)
+				return FRewindDebuggerTrack::EVisitResult::Continue;
+			}) == FRewindDebuggerTrack::EVisitResult::Break)
 		{
 			break;
 		}
@@ -1718,13 +1713,16 @@ void FRewindDebugger::SelectTrack(RewindDebugger::FObjectId ObjectId)
 // build a tree that's compatible with the public api from 5.0 for GetDebuggedObjects.
 void FRewindDebugger::RefreshDebuggedObjects(TArray<TSharedPtr<RewindDebugger::FRewindDebuggerTrack>>& InTracks, TArray<TSharedPtr<FDebugObjectInfo>>& OutObjects)
 {
+	using namespace RewindDebugger;
+
 	OutObjects.SetNum(0, EAllowShrinking::No);
-	for (const TSharedPtr<RewindDebugger::FRewindDebuggerTrack>& Track : InTracks)
+	for (const TSharedPtr<FRewindDebuggerTrack>& Track : InTracks)
 	{
-		const int Index = OutObjects.Add(MakeShared<FDebugObjectInfo>(Track->GetAssociatedObjectId(), Track->GetDisplayName().ToString()));
-		TArray<TSharedPtr<RewindDebugger::FRewindDebuggerTrack>> TrackChildren;
-		Track->IterateSubTracks([&TrackChildren](const TSharedPtr<RewindDebugger::FRewindDebuggerTrack>& Child) { TrackChildren.Add(Child); });
-		RefreshDebuggedObjects(TrackChildren, OutObjects[Index]->Children);
+		FRewindDebuggerTrack::Visit(Track, [&OutObjects](const TSharedPtr<FRewindDebuggerTrack>& Track)
+		{
+			OutObjects.Add(MakeShared<FDebugObjectInfo>(Track->GetAssociatedObjectId(), Track->GetDisplayName().ToString()));
+			return FRewindDebuggerTrack::EVisitResult::Continue;
+		});
 	}
 }
 
@@ -1736,24 +1734,20 @@ TArray<TSharedPtr<FDebugObjectInfo>>& FRewindDebugger::GetDebuggedObjects()
 
 bool FRewindDebugger::IsObjectCurrentlyDebugged(uint64 InObjectId) const
 {
-	if (const TraceServices::IAnalysisSession* Session = GetAnalysisSession())
+	using namespace RewindDebugger;
+	for (const TSharedPtr<FRewindDebuggerTrack>& Track : DebugTracks)
 	{
-		TraceServices::FAnalysisSessionReadScope SessionReadScope(*Session);
-		if (const IGameplayProvider* GameplayProvider = Session->ReadProvider<IGameplayProvider>("GameplayProvider"))
-		{
-			uint64 ObjectId = InObjectId;
-			while (ObjectId != RewindDebugger::FObjectId::InvalidId)
+		if (FRewindDebuggerTrack::Visit(Track, [InObjectId](TSharedPtr<FRewindDebuggerTrack> Track)
 			{
-				for (const TSharedPtr<RewindDebugger::FRewindDebuggerTrack>& Track : DebugTracks)
+				if (Track->GetUObjectId() == InObjectId)
 				{
-					if (Track->GetUObjectId() == ObjectId)
-					{
-						return true;
-					}
+					return FRewindDebuggerTrack::EVisitResult::Break;
 				}
 
-				ObjectId = GameplayProvider->GetObjectInfo(ObjectId).GetOuterUObjectId();
-			}
+				return FRewindDebuggerTrack::EVisitResult::Continue;
+			}) == FRewindDebuggerTrack::EVisitResult::Break)
+		{
+			return true;
 		}
 	}
 
diff --git a/Engine/Plugins/Animation/GameplayInsights/Source/RewindDebugger/Private/RewindDebuggerObjectTrack.cpp b/Engine/Plugins/Animation/GameplayInsights/Source/RewindDebugger/Private/RewindDebuggerObjectTrack.cpp
--- a/Engine/Plugins/Animation/GameplayInsights/Source/RewindDebugger/Private/RewindDebuggerObjectTrack.cpp
+++ b/Engine/Plugins/Animation/GameplayInsights/Source/RewindDebugger/Private/RewindDebuggerObjectTrack.cpp
@@ -154,44 +154,17 @@ bool FRewindDebuggerObjectTrack::HandleDoubleClickInternal()
 	return true;
 }
 
-void FRewindDebuggerObjectTrack::IterateSubTracksInternal(TFunction<void(TSharedPtr<FRewindDebuggerTrack> SubTrack)> IteratorFunction)
+TConstArrayView<TSharedPtr<FRewindDebuggerTrack>> FRewindDebuggerObjectTrack::GetChildrenInternal(TArray<TSharedPtr<FRewindDebuggerTrack>>& OutTracks) const
 {
-	for (FTrackCreatorAndTrack& TrackChild : ChildTrackCreators)
+	for (const FTrackCreatorAndTrack& Creator : ChildTrackCreators)
 	{
-		if (TrackChild.Track.IsValid())
+		if (Creator.Track.IsValid())
 		{
-			IteratorFunction(TrackChild.Track);
+			OutTracks.Push(Creator.Track);
 		}
 	}
 
-	for (TSharedPtr<FRewindDebuggerTrack>& Track : Children)
-	{
-		IteratorFunction(Track);
-	}
-}
-
-FRewindDebuggerTrack::EVisitResult FRewindDebuggerObjectTrack::VisitSubTracksInternal(TFunction<EVisitResult(TSharedPtr<FRewindDebuggerTrack> SubTrack)> IteratorFunction)
-{
-	for (FTrackCreatorAndTrack& TrackChild : ChildTrackCreators)
-	{
-		if (TrackChild.Track.IsValid())
-		{
-			if (IteratorFunction(TrackChild.Track) == EVisitResult::Break)
-			{
-				return EVisitResult::Break;
-			}
-		}
-	}
-
-	for (const TSharedPtr<FRewindDebuggerTrack>& Track : Children)
-	{
-		if (Track->VisitSubTracks(IteratorFunction) == EVisitResult::Break)
-		{
-			return EVisitResult::Break;
-		}
-	}
-
-	return EVisitResult::Continue;
+	return Children;
 }
 
 FName FRewindDebuggerObjectTrack::GetNameInternal() const
@@ -466,22 +439,24 @@ bool FRewindDebuggerObjectTrack::UpdateInternal()
 
 	{
 		TRACE_CPUPROFILER_EVENT_SCOPE(FRewindDebuggerObjectTrack::UpdateChilden);
-		for (auto& Child : Children)
+		for (const TSharedPtr<FRewindDebuggerTrack>& Child : Children)
 		{
 			if (Child->Update())
 			{
+				UE_LOG(LogRewindDebugger, Verbose, TEXT("List changed by: '%s'"), *Child->GetDisplayName().ToString());
 				bChanged = true;
 			}
 		}
 	}
 	{
 		TRACE_CPUPROFILER_EVENT_SCOPE(FRewindDebuggerObjectTrack::UpdateTrackChilden);
-		for (auto& TrackChild : ChildTrackCreators)
+		for (FTrackCreatorAndTrack& Creator : ChildTrackCreators)
 		{
-			if (TrackChild.Track.IsValid())
+			if (Creator.Track.IsValid())
 			{
-				if (TrackChild.Track->Update())
+				if (Creator.Track->Update())
 				{
+					UE_LOG(LogRewindDebugger, Verbose, TEXT("List changed by: '%s'"), *Creator.Track->GetDisplayName().ToString());
 					bChanged = true;
 				}
 			}
diff --git a/Engine/Plugins/Animation/GameplayInsights/Source/RewindDebugger/Private/RewindDebuggerObjectTrack.h b/Engine/Plugins/Animation/GameplayInsights/Source/RewindDebugger/Private/RewindDebuggerObjectTrack.h
--- a/Engine/Plugins/Animation/GameplayInsights/Source/RewindDebugger/Private/RewindDebuggerObjectTrack.h
+++ b/Engine/Plugins/Animation/GameplayInsights/Source/RewindDebugger/Private/RewindDebuggerObjectTrack.h
@@ -31,8 +31,7 @@ private:
 	virtual TSharedPtr<SWidget> GetDetailsViewInternal() override;
 	virtual TSharedPtr<SWidget> GetTimelineViewInternal() override;
 	virtual bool UpdateInternal() override;
-	virtual void IterateSubTracksInternal(TFunction<void(TSharedPtr<FRewindDebuggerTrack> SubTrack)> IteratorFunction) override;
-	virtual EVisitResult VisitSubTracksInternal(TFunction<EVisitResult(TSharedPtr<FRewindDebuggerTrack> SubTrack)> IteratorFunction) override;
+	virtual TConstArrayView<TSharedPtr<FRewindDebuggerTrack>> GetChildrenInternal(TArray<TSharedPtr<FRewindDebuggerTrack>>& OutTracks) const override;
 	virtual FName GetNameInternal() const override;
 	virtual FSlateIcon GetIconInternal() override;
 	virtual FText GetDisplayNameInternal() const override;
diff --git a/Engine/Plugins/Animation/GameplayInsights/Source/RewindDebuggerVLog/Private/VisualLogTrack.cpp b/Engine/Plugins/Animation/GameplayInsights/Source/RewindDebuggerVLog/Private/VisualLogTrack.cpp
--- a/Engine/Plugins/Animation/GameplayInsights/Source/RewindDebuggerVLog/Private/VisualLogTrack.cpp
+++ b/Engine/Plugins/Animation/GameplayInsights/Source/RewindDebuggerVLog/Private/VisualLogTrack.cpp
@@ -247,7 +247,10 @@ bool FVisualLogTrack::UpdateInternal()
 		const int32 TrackCount = UniqueTrackIds.Num();
 
 		if (Children.Num() != TrackCount)
+		{
+			UE_LOG(LogRewindDebugger, Verbose, TEXT("List changed because track count (%d) != current number of children (%d)"), TrackCount, Children.Num());
 			bChanged = true;
+		}
 
 		Children.SetNum(UniqueTrackIds.Num());
 		for (int i = 0; i < TrackCount; i++)
@@ -268,13 +271,10 @@ bool FVisualLogTrack::UpdateInternal()
 	return bChanged;
 }
 
-void FVisualLogTrack::IterateSubTracksInternal(TFunction<void(TSharedPtr<FRewindDebuggerTrack> SubTrack)> IteratorFunction)
+TConstArrayView<TSharedPtr<FRewindDebuggerTrack>> FVisualLogTrack::GetChildrenInternal(TArray<TSharedPtr<FRewindDebuggerTrack>>& OutTracks) const
 {
-	for (TSharedPtr<FVisualLogCategoryTrack>& Track : Children)
-	{
-		IteratorFunction(Track);
-	}
-};
+	return MakeConstArrayView<TSharedPtr<FRewindDebuggerTrack>>(reinterpret_cast<const TSharedPtr<FRewindDebuggerTrack>*>(Children.GetData()), Children.Num());
+}
 
 TSharedPtr<SWidget> FVisualLogTrack::GetDetailsViewInternal()
 {
diff --git a/Engine/Plugins/Animation/GameplayInsights/Source/RewindDebuggerVLog/Private/VisualLogTrack.h b/Engine/Plugins/Animation/GameplayInsights/Source/RewindDebuggerVLog/Private/VisualLogTrack.h
--- a/Engine/Plugins/Animation/GameplayInsights/Source/RewindDebuggerVLog/Private/VisualLogTrack.h
+++ b/Engine/Plugins/Animation/GameplayInsights/Source/RewindDebuggerVLog/Private/VisualLogTrack.h
@@ -87,7 +87,7 @@ private:
 	virtual FName GetNameInternal() const override { return "Visual Logging"; }
 	virtual FText GetDisplayNameInternal() const override { return NSLOCTEXT("RewindDebugger", "VisualLoggerTrackName", "Visual Logging"); }
 	virtual uint64 GetObjectIdInternal() const override { return ObjectId; }
-	virtual void IterateSubTracksInternal(TFunction<void(TSharedPtr<FRewindDebuggerTrack> SubTrack)> IteratorFunction) override;
+	virtual TConstArrayView<TSharedPtr<FRewindDebuggerTrack>> GetChildrenInternal(TArray<TSharedPtr<FRewindDebuggerTrack>>& OutTracks) const override;
 
 	FSlateIcon Icon;
 	uint64 ObjectId;
diff --git a/Engine/Plugins/Animation/PoseSearch/Source/Editor/Private/PoseSearchDebugger.cpp b/Engine/Plugins/Animation/PoseSearch/Source/Editor/Private/PoseSearchDebugger.cpp
--- a/Engine/Plugins/Animation/PoseSearch/Source/Editor/Private/PoseSearchDebugger.cpp
+++ b/Engine/Plugins/Animation/PoseSearch/Source/Editor/Private/PoseSearchDebugger.cpp
@@ -551,13 +551,10 @@ TSharedPtr<SWidget> FDebuggerTrack::GetDetailsViewInternal()
 	return FDebugger::Get()->GenerateInstance(ObjectId);
 }
 
-void FDebuggerTrack::IterateSubTracksInternal(TFunction<void(TSharedPtr<FRewindDebuggerTrack> SubTrack)> IteratorFunction)
+TConstArrayView<TSharedPtr<RewindDebugger::FRewindDebuggerTrack>> FDebuggerTrack::GetChildrenInternal(TArray<TSharedPtr<FRewindDebuggerTrack>>& OutTracks) const
 {
-	for(TSharedPtr<FSearchTrack>& SearchTrack : SearchTracks)
-	{
-		IteratorFunction(SearchTrack);
-	}
-};
+	return MakeConstArrayView<TSharedPtr<FRewindDebuggerTrack>>(reinterpret_cast<const TSharedPtr<FRewindDebuggerTrack>*>(SearchTracks.GetData()), SearchTracks.Num());
+}
 
 // FDebuggerTrackCreator
 ///////////////////////////////////////////////////
diff --git a/Engine/Plugins/Animation/PoseSearch/Source/Editor/Private/PoseSearchDebugger.h b/Engine/Plugins/Animation/PoseSearch/Source/Editor/Private/PoseSearchDebugger.h
--- a/Engine/Plugins/Animation/PoseSearch/Source/Editor/Private/PoseSearchDebugger.h
+++ b/Engine/Plugins/Animation/PoseSearch/Source/Editor/Private/PoseSearchDebugger.h
@@ -100,7 +100,7 @@ private:
 	virtual FName GetNameInternal() const override { return "PoseSearchDebugger"; }
 	virtual FText GetDisplayNameInternal() const override { return NSLOCTEXT("PoseSearchDebugger", "PoseSearchDebuggerTabTitle", "Pose Search"); }
 	virtual uint64 GetObjectIdInternal() const override { return ObjectId; }
-	virtual void IterateSubTracksInternal(TFunction<void(TSharedPtr<FRewindDebuggerTrack> SubTrack)> IteratorFunction) override;
+	virtual TConstArrayView<TSharedPtr<FRewindDebuggerTrack>> GetChildrenInternal(TArray<TSharedPtr<FRewindDebuggerTrack>>& OutTracks) const override;
 	virtual bool UpdateInternal() override;
 
 	uint64 ObjectId;
diff --git a/Engine/Plugins/Chooser/Source/ChooserEditor/Private/ChooserTrack.cpp b/Engine/Plugins/Chooser/Source/ChooserEditor/Private/ChooserTrack.cpp
--- a/Engine/Plugins/Chooser/Source/ChooserEditor/Private/ChooserTrack.cpp
+++ b/Engine/Plugins/Chooser/Source/ChooserEditor/Private/ChooserTrack.cpp
@@ -245,13 +245,10 @@ bool FChoosersTrack::UpdateInternal()
 	return bChanged;
 }
 
-void FChoosersTrack::IterateSubTracksInternal(TFunction<void(TSharedPtr<FRewindDebuggerTrack> SubTrack)> IteratorFunction)
+TConstArrayView<TSharedPtr<RewindDebugger::FRewindDebuggerTrack>> FChoosersTrack::GetChildrenInternal(TArray<TSharedPtr<FRewindDebuggerTrack>>& OutTracks) const
 {
-	for(TSharedPtr<FChooserTrack>& Track : Children)
-	{
-		IteratorFunction(Track);
-	}
-};
+	return MakeConstArrayView<TSharedPtr<FRewindDebuggerTrack>>(reinterpret_cast<const TSharedPtr<FRewindDebuggerTrack>*>(Children.GetData()), Children.Num());
+}
 
 bool FChoosersTrackCreator::HasDebugInfoInternal(const RewindDebugger::FObjectId& InObjectId) const
 {
diff --git a/Engine/Plugins/Chooser/Source/ChooserEditor/Private/ChooserTrack.h b/Engine/Plugins/Chooser/Source/ChooserEditor/Private/ChooserTrack.h
--- a/Engine/Plugins/Chooser/Source/ChooserEditor/Private/ChooserTrack.h
+++ b/Engine/Plugins/Chooser/Source/ChooserEditor/Private/ChooserTrack.h
@@ -55,7 +55,7 @@ private:
 	virtual FName GetNameInternal() const override { return "Choosers"; }
 	virtual FText GetDisplayNameInternal() const override { return NSLOCTEXT("RewindDebugger", "ChooserTrackName", "Chooser Evaluation"); }
 	virtual uint64 GetObjectIdInternal() const override { return ObjectId; }
-	virtual void IterateSubTracksInternal(TFunction<void(TSharedPtr<FRewindDebuggerTrack> SubTrack)> IteratorFunction) override;
+	virtual TConstArrayView<TSharedPtr<FRewindDebuggerTrack>> GetChildrenInternal(TArray<TSharedPtr<FRewindDebuggerTrack>>& OutTracks) const override;
 
 	FSlateIcon Icon;
 	uint64 ObjectId;
diff --git a/Engine/Plugins/Developer/FieldNotificationTrace/Source/FieldNotificationTraceEditor/Private/FieldNotificationTrack.cpp b/Engine/Plugins/Developer/FieldNotificationTrace/Source/FieldNotificationTraceEditor/Private/FieldNotificationTrack.cpp
--- a/Engine/Plugins/Developer/FieldNotificationTrace/Source/FieldNotificationTraceEditor/Private/FieldNotificationTrack.cpp
+++ b/Engine/Plugins/Developer/FieldNotificationTrace/Source/FieldNotificationTraceEditor/Private/FieldNotificationTrack.cpp
@@ -221,12 +221,9 @@ uint64 FObjectTrack::GetObjectIdInternal() const
 	return ObjectId;
 }
 
-void FObjectTrack::IterateSubTracksInternal(TFunction<void(TSharedPtr<FRewindDebuggerTrack> SubTrack)> IteratorFunction)
+TConstArrayView<TSharedPtr<RewindDebugger::FRewindDebuggerTrack>> FObjectTrack::GetChildrenInternal(TArray<TSharedPtr<FRewindDebuggerTrack>>& OutTracks) const
 {
-	for(TSharedPtr<FFieldNotifyTrack>& Track : Children)
-	{
-		IteratorFunction(Track);
-	}
+	return MakeConstArrayView<TSharedPtr<FRewindDebuggerTrack>>(reinterpret_cast<const TSharedPtr<FRewindDebuggerTrack>*>(Children.GetData()), Children.Num());
 }
 
 FName FTracksCreator::GetTargetTypeNameInternal() const
diff --git a/Engine/Plugins/Developer/FieldNotificationTrace/Source/FieldNotificationTraceEditor/Private/FieldNotificationTrack.h b/Engine/Plugins/Developer/FieldNotificationTrace/Source/FieldNotificationTraceEditor/Private/FieldNotificationTrack.h
--- a/Engine/Plugins/Developer/FieldNotificationTrace/Source/FieldNotificationTraceEditor/Private/FieldNotificationTrack.h
+++ b/Engine/Plugins/Developer/FieldNotificationTrace/Source/FieldNotificationTraceEditor/Private/FieldNotificationTrack.h
@@ -55,7 +55,7 @@ private:
 	virtual FName GetNameInternal() const override;
 	virtual FText GetDisplayNameInternal() const override;
 	virtual uint64 GetObjectIdInternal() const override;
-	virtual void IterateSubTracksInternal(TFunction<void(TSharedPtr<FRewindDebuggerTrack> SubTrack)> IteratorFunction) override;
+	virtual TConstArrayView<TSharedPtr<FRewindDebuggerTrack>> GetChildrenInternal(TArray<TSharedPtr<FRewindDebuggerTrack>>& OutTracks) const override;
 
 	FSlateIcon Icon;
 	uint64 ObjectId;
diff --git a/Engine/Plugins/Runtime/StateTree/Source/StateTreeEditorModule/Private/Debugger/StateTreeDebuggerTrack.cpp b/Engine/Plugins/Runtime/StateTree/Source/StateTreeEditorModule/Private/Debugger/StateTreeDebuggerTrack.cpp
--- a/Engine/Plugins/Runtime/StateTree/Source/StateTreeEditorModule/Private/Debugger/StateTreeDebuggerTrack.cpp
+++ b/Engine/Plugins/Runtime/StateTree/Source/StateTreeEditorModule/Private/Debugger/StateTreeDebuggerTrack.cpp
@@ -229,12 +229,9 @@ bool FStateTreeDebuggerOwnerTrack::UpdateInternal()
 	return bChanged;
 }
 
-void FStateTreeDebuggerOwnerTrack::IterateSubTracksInternal(TFunction<void(TSharedPtr<FRewindDebuggerTrack> SubTrack)> InIteratorFunction)
+TConstArrayView<TSharedPtr<RewindDebugger::FRewindDebuggerTrack>> FStateTreeDebuggerOwnerTrack::GetChildrenInternal(TArray<TSharedPtr<FRewindDebuggerTrack>>& OutTracks) const
 {
-	for (TSharedPtr<FStateTreeDebuggerInstanceTrack>& Track : SubTracks)
-	{
-		InIteratorFunction(Track);
-	}
+	return MakeConstArrayView<TSharedPtr<FRewindDebuggerTrack>>(reinterpret_cast<const TSharedPtr<FRewindDebuggerTrack>*>(SubTracks.GetData()), SubTracks.Num());
 }
 
 void FStateTreeDebuggerOwnerTrack::MarkAsStale(const double InStaleTime)
diff --git a/Engine/Plugins/Runtime/StateTree/Source/StateTreeEditorModule/Private/Debugger/StateTreeDebuggerTrack.h b/Engine/Plugins/Runtime/StateTree/Source/StateTreeEditorModule/Private/Debugger/StateTreeDebuggerTrack.h
--- a/Engine/Plugins/Runtime/StateTree/Source/StateTreeEditorModule/Private/Debugger/StateTreeDebuggerTrack.h
+++ b/Engine/Plugins/Runtime/StateTree/Source/StateTreeEditorModule/Private/Debugger/StateTreeDebuggerTrack.h
@@ -193,7 +193,7 @@ struct FStateTreeDebuggerOwnerTrack : FStateTreeDebuggerBaseTrack
 
 protected:
 	virtual bool UpdateInternal() override;
-	virtual void IterateSubTracksInternal(TFunction<void(TSharedPtr<FRewindDebuggerTrack> SubTrack)> InIteratorFunction) override;
+	virtual TConstArrayView<TSharedPtr<FRewindDebuggerTrack>> GetChildrenInternal(TArray<TSharedPtr<FRewindDebuggerTrack>>& OutTracks) const override;
 
 private:
 	TArray<TSharedPtr<FStateTreeDebuggerInstanceTrack>> SubTracks;
diff --git a/Engine/Plugins/Runtime/StateTree/Source/StateTreeEditorModule/Private/Debugger/StateTreeRewindDebuggerTrack.cpp b/Engine/Plugins/Runtime/StateTree/Source/StateTreeEditorModule/Private/Debugger/StateTreeRewindDebuggerTrack.cpp
--- a/Engine/Plugins/Runtime/StateTree/Source/StateTreeEditorModule/Private/Debugger/StateTreeRewindDebuggerTrack.cpp
+++ b/Engine/Plugins/Runtime/StateTree/Source/StateTreeEditorModule/Private/Debugger/StateTreeRewindDebuggerTrack.cpp
@@ -168,7 +168,6 @@ bool FRewindDebuggerTrack::UpdateInternal()
 				|| bActivateStatesChanged)
 			{
 				CompactTreeView->Refresh();
-				return true;
 			}
 		}
 	}
diff --git a/Engine/Source/Editor/RewindDebuggerInterface/Public/RewindDebuggerTrack.h b/Engine/Source/Editor/RewindDebuggerInterface/Public/RewindDebuggerTrack.h
--- a/Engine/Source/Editor/RewindDebuggerInterface/Public/RewindDebuggerTrack.h
+++ b/Engine/Source/Editor/RewindDebuggerInterface/Public/RewindDebuggerTrack.h
@@ -146,16 +146,64 @@ public:
 		return FObjectId(GetObjectIdInternal());
 	}
 
-	/** iterates over all sub-tracks of this track and call Iterator function */
-	EVisitResult VisitSubTracks(TFunction<EVisitResult(TSharedPtr<FRewindDebuggerTrack> SubTrack)> IteratorFunction)
+	/**
+	 * Visit the current track along with its child hierarchy recursively
+	 * to call provided function on each track until a call returns EVisitResult::Break
+	 */
+	static EVisitResult Visit(TSharedPtr<FRewindDebuggerTrack> Track, TFunction<EVisitResult(TSharedPtr<FRewindDebuggerTrack> SubTrack)> IteratorFunction)
 	{
-		return VisitSubTracksInternal(IteratorFunction);
+		if (IteratorFunction(Track) == EVisitResult::Break)
+		{
+			return EVisitResult::Break;
+		}
+
+		return Track->Visit(IteratorFunction);
+	}
+
+private:
+	/**
+	 * Visit the child hierarchy recursively to call provided function on each track until a call returns EVisitResult::Break
+	 */
+	EVisitResult Visit(TFunction<EVisitResult(TSharedPtr<FRewindDebuggerTrack> SubTrack)> IteratorFunction)
+	{
+		// GetChildrenInternal can provide children using the returned array view,
+		// the output parameter, or both.
+		TArray<TSharedPtr<FRewindDebuggerTrack>> Children;
+		for (const TSharedPtr<FRewindDebuggerTrack>& Track : GetChildrenInternal(Children))
+		{
+			if (Visit(Track, IteratorFunction) == EVisitResult::Break)
+			{
+				return EVisitResult::Break;
+			}
+		}
+
+		for (const TSharedPtr<FRewindDebuggerTrack>& Track : Children)
+		{
+			if (Visit(Track, IteratorFunction) == EVisitResult::Break)
+			{
+				return EVisitResult::Break;
+			}
+		}
+
+		return EVisitResult::Continue;
 	}
+public:
 
 	/** iterates over all sub-tracks of this track and call Iterator function */
-	void IterateSubTracks(TFunction<void(TSharedPtr<FRewindDebuggerTrack> SubTrack)> IteratorFunction)
+	void IterateSubTracks(TFunction<void(TSharedPtr<FRewindDebuggerTrack> SubTrack)> IteratorFunction) const
 	{
-		IterateSubTracksInternal(IteratorFunction);
+		// GetChildrenInternal can provide children using the returned array view,
+		// the output parameter, or both.
+		TArray<TSharedPtr<FRewindDebuggerTrack>> Children;
+		for (const TSharedPtr<FRewindDebuggerTrack>& Track : GetChildrenInternal(Children))
+		{
+			IteratorFunction(Track);
+		}
+
+		for (const TSharedPtr<FRewindDebuggerTrack>& Track : Children)
+		{
+			IteratorFunction(Track);
+		}
 	}
 
 	/** returns true for tracks that contain debug data (used for filtering out parts of the hierarchy with no useful information in them) */
@@ -248,24 +296,20 @@ private:
 		return true;
 	}
 
-	virtual void IterateSubTracksInternal(TFunction<void(TSharedPtr<FRewindDebuggerTrack> SubTrack)> IteratorFunction)
+	UE_DEPRECATED(5.7, "Implement and use GetChildrenInternal instead.")
+	virtual void IterateSubTracksInternal(TFunction<void(TSharedPtr<FRewindDebuggerTrack> SubTrack)> IteratorFunction) final
 	{
 	}
 
-	virtual EVisitResult VisitSubTracksInternal(TFunction<EVisitResult(TSharedPtr<FRewindDebuggerTrack> SubTrack)> IteratorFunction)
+	/**
+	 * Method that derive classes can override to provide children tracks for methods like Visit and IterateSubTracks
+	 * Caller of this method is expected to handle both the return value and the output parameter.
+	 * @param OutTracks Optional parameter that can be used when the derived class can't provide a const view of the children
+	 * @return An array view on the sub tracks
+	 */
+	virtual TConstArrayView<TSharedPtr<FRewindDebuggerTrack>> GetChildrenInternal(TArray<TSharedPtr<FRewindDebuggerTrack>>& OutTracks) const
 	{
-		// This is an opt-in optimization added later in the API so default implementation is provided using IterateSubTracks
-		// Not perfect since we keep iterating, but we stop calling the function after getting a result that indicates to stop.
-		EVisitResult Result = EVisitResult::Continue;
-		IterateSubTracks([IteratorFunction, &Result](TSharedPtr<FRewindDebuggerTrack> SubTrack)
-			{
-				if (Result == EVisitResult::Continue)
-				{
-					Result = IteratorFunction(SubTrack);
-				}
-			});
-
-		return EVisitResult::Continue;
+		return {};
 	}
 
 	virtual bool HandleDoubleClickInternal()
