[RewindDebugger] replaced some actor specific methods by more generic ones
Moved code to update root object position from RewindDebugger to the Animation extension since it all relies on the animation provider at the moment.

diff --git a/Engine/Plugins/Animation/GameplayInsights/Source/RewindDebugger/Private/RewindDebugger.cpp b/Engine/Plugins/Animation/GameplayInsights/Source/RewindDebugger/Private/RewindDebugger.cpp
--- a/Engine/Plugins/Animation/GameplayInsights/Source/RewindDebugger/Private/RewindDebugger.cpp
+++ b/Engine/Plugins/Animation/GameplayInsights/Source/RewindDebugger/Private/RewindDebugger.cpp
@@ -2,27 +2,19 @@
 
 #include "RewindDebugger.h"
 
-#include "Animation/AnimBlueprintGeneratedClass.h"
-#include "Animation/AnimTrace.h"
 #include "DesktopPlatformModule.h"
 #include "Editor.h"
 #include "Editor/UnrealEdEngine.h"
-#include "Engine/World.h"
-#include "EngineUtils.h"
-#include "GameFramework/Pawn.h"
+#include "Framework/Application/SlateApplication.h"
 #include "HAL/PlatformFileManager.h"
-#include "IAnimationProvider.h"
 #include "IDesktopPlatform.h"
-#include "IGameplayProvider.h"
 #include "Insights/IUnrealInsightsModule.h"
-#include "IRewindDebuggerDoubleClickHandler.h"
 #include "IRewindDebuggerExtension.h"
 #include "Kismet2/DebuggerCommands.h"
 #include "LevelEditor.h"
 #include "Misc/MessageDialog.h"
 #include "Modules/ModuleManager.h"
 #include "ObjectTrace.h"
-#include "ProfilingDebugging/TraceAuxiliary.h"
 #include "RewindDebuggerCommands.h"
 #include "RewindDebuggerModule.h"
 #include "RewindDebuggerObjectTrack.h"
@@ -30,13 +22,10 @@
 #include "RewindDebuggerRuntime/RewindDebuggerRuntime.h"
 #include "RewindDebuggerSettings.h"
 #include "RewindDebuggerTrackCreators.h"
-#include "SLevelViewport.h"
 #include "SModalSessionBrowser.h"
+#include "ToolMenuDelegates.h"
 #include "ToolMenus.h"
-#include "Trace/StoreClient.h"
-#include "TraceServices/Model/Frames.h"
 #include "UnrealEdGlobals.h"
-#include "UObject/UObjectIterator.h"
 #include "Widgets/Docking/SDockTab.h"
 #include "Widgets/Input/SNumericEntryBox.h"
 #include "Widgets/Layout/SSpacer.h"
@@ -101,7 +90,7 @@ FRewindDebugger::FRewindDebugger()
 	// Using ShutdownPIE instead of EndPIE to make sure all traces emitted during world EndPlay get processed before disabling channels
 	ShutdownPIEHandle = FEditorDelegates::ShutdownPIE.AddRaw(this, &FRewindDebugger::OnPIEStopped);
 
-	SelectedObjectName.OnPropertyChanged = SelectedObjectName.OnPropertyChanged.CreateLambda([this](FString Target)
+	RootObjectName.OnPropertyChanged = RootObjectName.OnPropertyChanged.CreateLambda([this](FString Target)
 		{
 			URewindDebuggerSettings& Settings = URewindDebuggerSettings::Get();
 			if (Settings.DebugTargetActor != Target)
@@ -113,7 +102,7 @@ FRewindDebugger::FRewindDebugger()
 
 			CandidateIds.SetNum(0);
 			GetTargetObjectIds(CandidateIds);
-			// make sure all the SubObjects of the target actor have been traced
+			// make sure all the SubObjects of the root object have been traced
 #if OBJECT_TRACE_ENABLED
 			for (const RewindDebugger::FObjectId& TargetObjectId : CandidateIds)
 			{
@@ -262,43 +251,22 @@ void FRewindDebugger::OnPIEStopped(bool bSimulating)
 	bDisplayWorldIdValid = false;
 }
 
-bool FRewindDebugger::GetTargetActorPosition(FVector& OutPosition) const
+bool FRewindDebugger::GetRootObjectPosition(FVector& OutPosition) const
 {
-	OutPosition = TargetActorPosition;
-	return bTargetActorPositionValid;
+	OutPosition = RootObjectPosition.Get(FVector::ZeroVector);
+	return RootObjectPosition.IsSet();
 }
 
-uint64 FRewindDebugger::GetTargetActorId() const
+void FRewindDebugger::SetRootObjectPosition(const TOptional<FVector>& InPosition)
 {
-	if (SelectedObjectName.Get() == "")
-	{
-		return 0;
-	}
-
-	uint64 TargetActorId = 0;
-
-	if (const TraceServices::IAnalysisSession* Session = GetAnalysisSession())
-	{
-		if (const IGameplayProvider* GameplayProvider = Session->ReadProvider<IGameplayProvider>("GameplayProvider"))
-		{
-			GameplayProvider->EnumerateObjects(CurrentTraceRange.GetLowerBoundValue(), CurrentTraceRange.GetUpperBoundValue(), [this, &TargetActorId](const FObjectInfo& InObjectInfo)
-				{
-					if (SelectedObjectName.Get() == InObjectInfo.Name)
-					{
-						TargetActorId = InObjectInfo.GetUObjectId();
-					}
-				});
-		}
-	}
-
-	return TargetActorId;
+	RootObjectPosition = InPosition;
 }
 
 void FRewindDebugger::GetTargetObjectIds(TArray<RewindDebugger::FObjectId>& OutTargetObjectIds) const
 {
 	OutTargetObjectIds.Empty(2);
 
-	if (SelectedObjectName.Get() == "")
+	if (RootObjectName.Get() == "")
 	{
 		return;
 	}
@@ -311,7 +279,7 @@ void FRewindDebugger::GetTargetObjectIds(TArray<RewindDebugger::FObjectId>& OutT
 		{
 			GameplayProvider->EnumerateObjects(CurrentTraceRange.GetLowerBoundValue(), CurrentTraceRange.GetUpperBoundValue(), [this, &OutTargetObjectIds](const FObjectInfo& InObjectInfo)
 				{
-					if (SelectedObjectName.Get() == InObjectInfo.Name)
+					if (RootObjectName.Get() == InObjectInfo.Name)
 					{
 						OutTargetObjectIds.Add(InObjectInfo.GetId());
 					}
@@ -319,7 +287,7 @@ void FRewindDebugger::GetTargetObjectIds(TArray<RewindDebugger::FObjectId>& OutT
 		}
 	}
 
-	// make sure all the SubObjects of the target actor have been traced
+	// make sure all the SubObjects of the root object have been traced
 #if OBJECT_TRACE_ENABLED
 	if (IsRecording())
 	{
@@ -344,7 +312,7 @@ void FRewindDebugger::RefreshDebugTracks()
 		GetTargetObjectIds(CandidateIds);
 	}
 
-	const FString SelectionName = SelectedObjectName.Get();
+	const FString SelectionName = RootObjectName.Get();
 
 	if (CandidateIds.Num() == 0 && !SelectionName.IsEmpty())
 	{
@@ -388,7 +356,7 @@ void FRewindDebugger::RefreshDebugTracks()
 
 			if (!FoundTrack)
 			{
-				DebugTracks.Add(MakeShared<RewindDebugger::FRewindDebuggerObjectTrack>(CandidateIdentifier, SelectedObjectName.Get(), true));
+				DebugTracks.Add(MakeShared<RewindDebugger::FRewindDebuggerObjectTrack>(CandidateIdentifier, RootObjectName.Get(), true));
 				bChanged = true;
 			}
 		}
@@ -435,7 +403,7 @@ void FRewindDebugger::OnClearRecording()
 	ClearTrace();
 	RecordingDuration.Set(0);
 	CandidateIds.Empty(2);
-	bTargetActorPositionValid = false;
+	RootObjectPosition.Reset();
 
 	IterateExtensions([this](IRewindDebuggerExtension* Extension)
 		{
@@ -944,23 +912,23 @@ void FRewindDebugger::GetScrubTimeInformation(double InDebugTime, FScrubTimeInfo
 						}
 
 						// Ensure there is not frames between start and end index
-						check(EndEventIndex == StartEventIndex + 1)
+						check(EndEventIndex == StartEventIndex + 1);
 
-							const FRecordingInfoMessage& Event = Recording->GetEvent(StartEventIndex);
+						const FRecordingInfoMessage& Event = Recording->GetEvent(StartEventIndex);
 						const FRecordingInfoMessage& NextEvent = Recording->GetEvent(EndEventIndex);
 
 						// Ensure debug time is between both frames time range
-						check(Event.ElapsedTime <= InDebugTime && NextEvent.ElapsedTime >= InDebugTime)
+						check(Event.ElapsedTime <= InDebugTime && NextEvent.ElapsedTime >= InDebugTime);
 
-							// Choose frame that is nearest to the debug time
-							if (InDebugTime - Event.ElapsedTime < NextEvent.ElapsedTime - InDebugTime)
-							{
-								ScrubFrameIndex = StartEventIndex;
-							}
-							else
-							{
-								ScrubFrameIndex = EndEventIndex;
-							}
+						// Choose frame that is nearest to the debug time
+						if (InDebugTime - Event.ElapsedTime < NextEvent.ElapsedTime - InDebugTime)
+						{
+							ScrubFrameIndex = StartEventIndex;
+						}
+						else
+						{
+							ScrubFrameIndex = EndEventIndex;
+						}
 					}
 				}
 
@@ -982,14 +950,20 @@ const TraceServices::IAnalysisSession* FRewindDebugger::GetAnalysisSession() con
 	return UnrealInsightsModule ? UnrealInsightsModule->GetAnalysisSession().Get() : nullptr;
 }
 
-const FObjectInfo* FRewindDebugger::FindOwningActorInfo(const IGameplayProvider* GameplayProvider, uint64 InObjectId) const
+uint64 FRewindDebugger::GetRootObjectId() const
 {
-	const FClassInfo* ActorClassInfo = GameplayProvider->FindClassInfo(*AActor::StaticClass()->GetPathName());
+	return DebuggedObjects.Num() ? DebuggedObjects[0]->GetUObjectId() : RewindDebugger::FObjectId::InvalidId;
+}
+
+const FObjectInfo* FRewindDebugger::FindTypedOuterInfo(TNotNull<const UStruct*> InType, TNotNull<const IGameplayProvider*> InGameplayProvider, const uint64 InObjectId) const
+{
+	const FClassInfo* TypeInfo = InGameplayProvider->FindClassInfo(*InType->GetPathName());
 
+	uint64 ObjectId(InObjectId);
 	while (true)
 	{
-		const FObjectInfo& ObjectInfo = GameplayProvider->GetObjectInfo(InObjectId);
-		if (GameplayProvider->IsSubClassOf(ObjectInfo.ClassId, ActorClassInfo->Id))
+		const FObjectInfo& ObjectInfo = InGameplayProvider->GetObjectInfo(ObjectId);
+		if (InGameplayProvider->IsSubClassOf(ObjectInfo.ClassId, TypeInfo->Id))
 		{
 			return &ObjectInfo;
 		}
@@ -999,7 +973,7 @@ const FObjectInfo* FRewindDebugger::FindOwningActorInfo(const IGameplayProvider*
 			return nullptr;
 		}
 
-		InObjectId = ObjectInfo.GetOuterUObjectId();
+		ObjectId = ObjectInfo.GetOuterUObjectId();
 	}
 }
 
@@ -1015,10 +989,7 @@ void FRewindDebugger::Tick(float DeltaTime)
 
 	if (const TraceServices::IAnalysisSession* Session = GetAnalysisSession())
 	{
-		const IAnimationProvider* AnimationProvider = Session->ReadProvider<IAnimationProvider>("AnimationProvider");
-		const IGameplayProvider* GameplayProvider = Session->ReadProvider<IGameplayProvider>("GameplayProvider");
-
-		if (AnimationProvider && GameplayProvider)
+		if (const IGameplayProvider* GameplayProvider = Session->ReadProvider<IGameplayProvider>("GameplayProvider"))
 		{
 			TraceServices::FAnalysisSessionReadScope SessionReadScope(*Session);
 
@@ -1066,7 +1037,9 @@ void FRewindDebugger::Tick(float DeltaTime)
 					SetCurrentScrubTime(RecordingDurationValue);
 					TrackCursorDelegate.ExecuteIfBound(false);
 				}
-				bTargetActorPositionValid = false;
+
+				// The debug position is only expected to be set outside of PIE
+				RootObjectPosition.Reset();
 			}
 			else
 			{
@@ -1087,65 +1060,8 @@ void FRewindDebugger::Tick(float DeltaTime)
 						}
 					}
 
-					SetCurrentScrubTime(CurrentScrubTime);// update trace time
-
-					const double CurrentTraceTime = TraceTime.Get();
-					if (CurrentTraceTime != PreviousTraceTime)
-					{
-						TRACE_CPUPROFILER_EVENT_SCOPE(FRewindDebugger::Tick_UpdateActorPosition);
-						PreviousTraceTime = CurrentTraceTime;
-
-						const TraceServices::IFrameProvider& FrameProvider = TraceServices::ReadFrameProvider(*Session);
-						TraceServices::FFrame Frame;
-						if (FrameProvider.GetFrameFromTime(TraceFrameType_Game, CurrentTraceTime, Frame))
-						{
-							bool bNewActor = false;
-							if (!CandidateIds.Contains(RewindDebugger::FObjectId{TargetActorIdForMesh}))
-							{
-								AnimationProvider->EnumerateSkeletalMeshPoseTimelines([this, &bNewActor, GameplayProvider](uint64 ObjectId, const IAnimationProvider::SkeletalMeshPoseTimeline& TimelineData)
-									{
-										// until we have actor transforms traced out, the first (from a non-server) skeletal mesh component transform on the target actor be used as the actor position
-
-										if (const FWorldInfo* WorldInfo = GameplayProvider->FindWorldInfoFromObject(ObjectId))
-										{
-											if (WorldInfo->NetMode != FWorldInfo::ENetMode::DedicatedServer)
-											{
-												if (const FObjectInfo* ActorInfo = FindOwningActorInfo(GameplayProvider, ObjectId))
-												{
-													if (CandidateIds.Contains(ActorInfo->GetId()))
-													{
-														bNewActor = true;
-														TargetActorIdForMesh = ActorInfo->GetUObjectId();
-														TargetActorMeshId = ObjectId;
-													}
-												}
-											}
-										}
-									});
-							}
-
-							AnimationProvider->ReadSkeletalMeshPoseTimeline(TargetActorMeshId, [this, &Frame, bNewActor](const IAnimationProvider::SkeletalMeshPoseTimeline& TimelineData, bool bHasCurves)
-								{
-									const FSkeletalMeshPoseMessage* PoseMessage = nullptr;
-
-									// Get last pose in frame
-									TimelineData.EnumerateEvents(Frame.StartTime, Frame.EndTime,
-										[&PoseMessage](double InStartTime, double InEndTime, uint32 InDepth, const FSkeletalMeshPoseMessage& InPoseMessage)
-										{
-											PoseMessage = &InPoseMessage;
-											return TraceServices::EEventEnumerate::Continue;
-										});
-
-									// Update position based on pose
-									if (PoseMessage)
-									{
-										// mark the target position as invalid for a frame when the actor changes, so it will be treated as a teleport by the camera system
-										bTargetActorPositionValid = !bNewActor;
-										TargetActorPosition = PoseMessage->ComponentToWorld.GetTranslation();
-									}
-								});
-						}
-					}
+					// update trace time
+					SetCurrentScrubTime(CurrentScrubTime);
 				}
 			}
 		}
diff --git a/Engine/Plugins/Animation/GameplayInsights/Source/RewindDebugger/Private/RewindDebugger.h b/Engine/Plugins/Animation/GameplayInsights/Source/RewindDebugger/Private/RewindDebugger.h
--- a/Engine/Plugins/Animation/GameplayInsights/Source/RewindDebugger/Private/RewindDebugger.h
+++ b/Engine/Plugins/Animation/GameplayInsights/Source/RewindDebugger/Private/RewindDebugger.h
@@ -57,8 +57,9 @@ public:
 	}
 
 	virtual const TraceServices::IAnalysisSession* GetAnalysisSession() const override;
-	virtual uint64 GetTargetActorId() const override;
-	virtual bool GetTargetActorPosition(FVector& OutPosition) const override;
+	virtual uint64 GetRootObjectId() const override;
+	virtual bool GetRootObjectPosition(FVector& OutPosition) const override;
+	virtual void SetRootObjectPosition(const TOptional<FVector>& InPosition) override;
 	virtual UWorld* GetWorldToVisualize() const override;
 	virtual bool IsRecording() const override;
 	virtual bool IsPIESimulating() const override
@@ -199,9 +200,15 @@ public:
 		return &RecordingDuration;
 	}
 
+	TBindableProperty<FString, BindingType_Out>* GetRootObjectNameProperty()
+	{
+		return &RootObjectName;
+	}
+
+	UE_DEPRECATED(5.7, "Use GetRootObjectNameProperty instead")
 	TBindableProperty<FString, BindingType_Out>* GetDebugTargetActorProperty()
 	{
-		return &SelectedObjectName;
+		return GetRootObjectNameProperty();
 	}
 
 	virtual void OpenDetailsPanel() override;
@@ -210,7 +217,7 @@ public:
 		bIsDetailsPanelOpen = bIsOpen;
 	}
 
-	virtual const FObjectInfo* FindOwningActorInfo(const IGameplayProvider* GameplayProvider, uint64 ObjectId) const override;
+	virtual const FObjectInfo* FindTypedOuterInfo(TNotNull<const UStruct*> InType, TNotNull<const IGameplayProvider*> InGameplayProvider, uint64 InObjectId) const override;
 
 	TArrayView<RewindDebugger::FRewindDebuggerTrackType> GetTrackTypes()
 	{
@@ -238,7 +245,7 @@ private:
 
 	TBindableProperty<double> TraceTime;
 	TBindableProperty<double> RecordingDuration;
-	TBindableProperty<FString, BindingType_Out> SelectedObjectName;
+	TBindableProperty<FString, BindingType_Out> RootObjectName;
 
 	enum class EControlState : int8
 	{
@@ -290,10 +297,7 @@ private:
 	mutable class IUnrealInsightsModule* UnrealInsightsModule = nullptr;
 	FTSTicker::FDelegateHandle TickerHandle;
 
-	bool bTargetActorPositionValid = false;
-	FVector TargetActorPosition;
-	uint64 TargetActorMeshId = 0;
-	uint64 TargetActorIdForMesh = 0;
+	TOptional<FVector> RootObjectPosition;
 
 	TArray<RewindDebugger::FRewindDebuggerTrackType> TrackTypes;
 
diff --git a/Engine/Plugins/Animation/GameplayInsights/Source/RewindDebugger/Private/RewindDebuggerAnimation.cpp b/Engine/Plugins/Animation/GameplayInsights/Source/RewindDebugger/Private/RewindDebuggerAnimation.cpp
--- a/Engine/Plugins/Animation/GameplayInsights/Source/RewindDebugger/Private/RewindDebuggerAnimation.cpp
+++ b/Engine/Plugins/Animation/GameplayInsights/Source/RewindDebugger/Private/RewindDebuggerAnimation.cpp
@@ -200,7 +200,7 @@ FRewindDebuggerAnimation::FSpawnedMeshComponentInfo* FRewindDebuggerAnimation::S
 
 		MeshComponentInfo->Actor = World->SpawnActor<AActor>(ActorSpawnParameters);
 
-		if(const FObjectInfo* ActorInfo = RewindDebugger->FindOwningActorInfo(GameplayProvider, ObjectId))
+		if (const FObjectInfo* ActorInfo = RewindDebugger->FindTypedOuterInfo<AActor>(GameplayProvider, ObjectId))
 		{
 			MeshComponentInfo->Actor->SetActorLabel(FString(TEXT("RewindDebugger: ") + FString(ActorInfo->Name)));
 		}
@@ -267,7 +267,7 @@ UAnimInstance* FRewindDebuggerAnimation::GetDebugAnimInstance(uint64 ObjectId)
 void FRewindDebuggerAnimation::Update(float DeltaTime, IRewindDebugger* RewindDebugger)
 {
 	check(RewindDebugger);
-	
+
 	if (RewindDebugger->IsPIESimulating() || RewindDebugger->GetRecordingDuration() == 0.0)
 	{
 		return;
@@ -276,20 +276,27 @@ void FRewindDebuggerAnimation::Update(float DeltaTime, IRewindDebugger* RewindDe
 	if (const TraceServices::IAnalysisSession* Session = RewindDebugger->GetAnalysisSession())
 	{
 		TraceServices::FAnalysisSessionReadScope SessionReadScope(*Session);
-		
-		double CurrentTraceTime = RewindDebugger->CurrentTraceTime();
+
+		const double CurrentTraceTime = RewindDebugger->CurrentTraceTime();
 		if (CurrentTraceTime != LastScrubTime)
 		{
 			const TraceServices::IFrameProvider& FrameProvider = TraceServices::ReadFrameProvider(*Session);
 			TraceServices::FFrame Frame;
-			if (FrameProvider.GetFrameFromTime(ETraceFrameType::TraceFrameType_Game, CurrentTraceTime, Frame))
+			if (FrameProvider.GetFrameFromTime(TraceFrameType_Game, CurrentTraceTime, Frame))
 			{
 				const IAnimationProvider* AnimationProvider = Session->ReadProvider<IAnimationProvider>("AnimationProvider");
 				const IGameplayProvider* GameplayProvider = Session->ReadProvider<IGameplayProvider>("GameplayProvider");
-				
+
 				if (AnimationProvider && GameplayProvider)
 				{
-					// UWorld* World = GetWorldToVisualize();
+					const TArray<TSharedPtr<FDebugObjectInfo>>& DebuggedObjects = RewindDebugger->GetDebuggedObjects();
+					bool bSearchForActor = !TargetActorIdForMesh.IsSet() || !DebuggedObjects.ContainsByPredicate(
+						[Id = TargetActorIdForMesh](const TSharedPtr<FDebugObjectInfo>& ObjectInfo)
+						{
+							return ObjectInfo->Id == Id;
+						});
+					bool bNewActor = false;
+					TargetActorPosition.Reset();
 
 					// update pose on all SkeletalMeshComponents:
 					// - enumerate all skeletal mesh pose timelines
@@ -297,8 +304,33 @@ void FRewindDebuggerAnimation::Update(float DeltaTime, IRewindDebugger* RewindDe
 					// - apply the recorded pose for the current Frame
 					{
 						TRACE_CPUPROFILER_EVENT_SCOPE(FRewindDebugger::Tick_UpdatePoses);
-						AnimationProvider->EnumerateSkeletalMeshPoseTimelines([this, RewindDebugger, &Frame, AnimationProvider, GameplayProvider](uint64 ObjectId, const IAnimationProvider::SkeletalMeshPoseTimeline& TimelineData)
+						AnimationProvider->EnumerateSkeletalMeshPoseTimelines(
+						[this, RewindDebugger, &Frame, AnimationProvider, GameplayProvider, &bSearchForActor, &DebuggedObjects, &bNewActor](uint64 ObjectId, const IAnimationProvider::SkeletalMeshPoseTimeline& TimelineData)
 						{
+							if (bSearchForActor)
+							{
+								// until we have actor transforms traced out, the first (from a non-server) skeletal mesh component transform on the target actor be used as the actor position
+								if (const FWorldInfo* WorldInfo = GameplayProvider->FindWorldInfoFromObject(ObjectId))
+								{
+									if (WorldInfo->NetMode != FWorldInfo::ENetMode::DedicatedServer)
+									{
+										if (const FObjectInfo* ActorInfo = RewindDebugger->FindTypedOuterInfo<AActor>(GameplayProvider, ObjectId))
+										{
+											if (DebuggedObjects.ContainsByPredicate([Id = ActorInfo->GetId()](const TSharedPtr<FDebugObjectInfo>& ObjectInfo)
+												{
+													return ObjectInfo->Id == Id;
+												}))
+											{
+												bSearchForActor = false;
+												bNewActor = true;
+												TargetActorIdForMesh = ActorInfo->GetId();
+												TargetActorMeshId = RewindDebugger::FObjectId{ ObjectId };
+											}
+										}
+									}
+								}
+							}
+
 							USkeletalMeshComponent* MeshComponent = nullptr;
 							bool bQueueForReset = false;
 							bool bLoadMesh = false;
@@ -359,7 +391,36 @@ void FRewindDebuggerAnimation::Update(float DeltaTime, IRewindDebugger* RewindDe
 						}
 						);
 					}
-        
+
+					{
+						TRACE_CPUPROFILER_EVENT_SCOPE(FRewindDebuggerAnimation::UpdateActorPosition);
+
+						// mark the target position as invalid for a frame when the actor changes, so it will be treated as a teleport by the camera system
+						if (!bNewActor)
+						{
+							AnimationProvider->ReadSkeletalMeshPoseTimeline(TargetActorMeshId.GetMainId(), [this, &Frame, bNewActor](const IAnimationProvider::SkeletalMeshPoseTimeline& TimelineData, bool bHasCurves)
+							{
+								const FSkeletalMeshPoseMessage* PoseMessage = nullptr;
+
+								// Get last pose in frame
+								TimelineData.EnumerateEvents(Frame.StartTime, Frame.EndTime,
+									[&PoseMessage](double InStartTime, double InEndTime, uint32 InDepth, const FSkeletalMeshPoseMessage& InPoseMessage)
+									{
+										PoseMessage = &InPoseMessage;
+										return TraceServices::EEventEnumerate::Continue;
+									});
+
+								// Update position based on pose
+								if (PoseMessage)
+								{
+									TargetActorPosition = PoseMessage->ComponentToWorld.GetTranslation();
+								}
+							});
+						}
+
+						RewindDebugger->SetRootObjectPosition(TargetActorPosition);
+					}
+
 					{
 						TRACE_CPUPROFILER_EVENT_SCOPE(FRewindDebugger::Tick_AnimBlueprintsDebug);
 						// Apply Animation Blueprint Debugging Data:
diff --git a/Engine/Plugins/Animation/GameplayInsights/Source/RewindDebugger/Private/RewindDebuggerAnimation.h b/Engine/Plugins/Animation/GameplayInsights/Source/RewindDebugger/Private/RewindDebuggerAnimation.h
--- a/Engine/Plugins/Animation/GameplayInsights/Source/RewindDebugger/Private/RewindDebuggerAnimation.h
+++ b/Engine/Plugins/Animation/GameplayInsights/Source/RewindDebugger/Private/RewindDebuggerAnimation.h
@@ -1,7 +1,6 @@
 // Copyright Epic Games, Inc. All Rights Reserved.
 #pragma once
 
-#include "CoreMinimal.h"
 #include "IAnimationProvider.h"
 #include "IGameplayProvider.h"
 #include "IRewindDebuggerExtension.h"
@@ -80,5 +79,9 @@ private:
 	TMap<uint64, FMeshComponentResetData> MeshComponentsToReset;
 	double LastScrubTime = -1;
 
+	RewindDebugger::FObjectId TargetActorMeshId;
+	RewindDebugger::FObjectId TargetActorIdForMesh;
+	TOptional<FVector> TargetActorPosition;
+
 	static FRewindDebuggerAnimation* Instance;
 };
diff --git a/Engine/Plugins/Animation/GameplayInsights/Source/RewindDebugger/Private/RewindDebuggerCamera.cpp b/Engine/Plugins/Animation/GameplayInsights/Source/RewindDebugger/Private/RewindDebuggerCamera.cpp
--- a/Engine/Plugins/Animation/GameplayInsights/Source/RewindDebugger/Private/RewindDebuggerCamera.cpp
+++ b/Engine/Plugins/Animation/GameplayInsights/Source/RewindDebugger/Private/RewindDebuggerCamera.cpp
@@ -120,8 +120,7 @@ void FRewindDebuggerCamera::Update(float DeltaTime, IRewindDebugger* RewindDebug
 				FLevelEditorViewportClient& LevelViewportClient = LevelViewport->GetLevelViewportClient();
 
 				FVector TargetActorPosition;
-				bool bTargetActorPositionValid = RewindDebugger->GetTargetActorPosition(TargetActorPosition);
-
+				const bool bTargetActorPositionValid = RewindDebugger->GetRootObjectPosition(TargetActorPosition);
 
 				if (CameraMode() == ERewindDebuggerCameraMode::FollowTargetActor)
 				{
diff --git a/Engine/Plugins/Animation/GameplayInsights/Source/RewindDebugger/Private/RewindDebuggerModule.cpp b/Engine/Plugins/Animation/GameplayInsights/Source/RewindDebugger/Private/RewindDebuggerModule.cpp
--- a/Engine/Plugins/Animation/GameplayInsights/Source/RewindDebugger/Private/RewindDebuggerModule.cpp
+++ b/Engine/Plugins/Animation/GameplayInsights/Source/RewindDebugger/Private/RewindDebuggerModule.cpp
@@ -198,7 +198,7 @@ TSharedRef<SDockTab> FRewindDebuggerModule::SpawnRewindDebuggerTab(const FSpawnT
 	}
 	
 	RewindDebuggerWidget = SNew(SRewindDebugger, CommandList.ToSharedRef(), MajorTab, SpawnTabArgs.GetOwnerWindow())
-								.DebuggedObjectName({ DebuggerInstance->GetDebugTargetActorProperty(), URewindDebuggerSettings::Get().DebugTargetActor})
+								.DebuggedObjectName({ DebuggerInstance->GetRootObjectNameProperty(), URewindDebuggerSettings::Get().DebugTargetActor})
 								.RecordingDuration(DebuggerInstance->GetRecordingDurationProperty())
 								.Tracks(&DebuggerInstance->GetTracks())
 								.TraceTime(DebuggerInstance->GetTraceTimeProperty())
diff --git a/Engine/Plugins/Animation/GameplayInsights/Source/RewindDebugger/Private/RewindDebuggerSettings.h b/Engine/Plugins/Animation/GameplayInsights/Source/RewindDebugger/Private/RewindDebuggerSettings.h
--- a/Engine/Plugins/Animation/GameplayInsights/Source/RewindDebugger/Private/RewindDebuggerSettings.h
+++ b/Engine/Plugins/Animation/GameplayInsights/Source/RewindDebugger/Private/RewindDebuggerSettings.h
@@ -9,7 +9,7 @@
 class URewindDebuggerExtensionSettings;
 
 UENUM()
-enum class ERewindDebuggerCameraMode
+enum class ERewindDebuggerCameraMode : uint8
 {
 	Replay UMETA(Tooltip="Replay Recorded Camera"),
 	FollowTargetActor UMETA(Tooltip="Follow Target Actor"),
diff --git a/Engine/Plugins/Chooser/Source/ChooserEditor/Private/ChooserTrack.cpp b/Engine/Plugins/Chooser/Source/ChooserEditor/Private/ChooserTrack.cpp
--- a/Engine/Plugins/Chooser/Source/ChooserEditor/Private/ChooserTrack.cpp
+++ b/Engine/Plugins/Chooser/Source/ChooserEditor/Private/ChooserTrack.cpp
@@ -5,6 +5,7 @@
 #include "Chooser.h"
 #include "ChooserProvider.h"
 #include "Editor.h"
+#include "GameFramework/Actor.h"
 #include "IGameplayProvider.h"
 #include "IRewindDebugger.h"
 #include "Subsystems/AssetEditorSubsystem.h"
@@ -114,7 +115,7 @@ bool FChooserTrack::HandleDoubleClickInternal()
 			
 			const FObjectInfo& OwnerObjectInfo  = GameplayProvider->GetObjectInfo(ObjectId);
 			FString DebugName(OwnerObjectInfo.Name);
-			if (const FObjectInfo* ActorInfo = RewindDebugger->FindOwningActorInfo(GameplayProvider, ObjectId))
+			if (const FObjectInfo* ActorInfo = RewindDebugger->FindTypedOuterInfo<AActor>(GameplayProvider, ObjectId))
 			{
 				DebugName += " in " + FString(ActorInfo->Name);
 			}
diff --git a/Engine/Plugins/Chooser/Source/ChooserEditor/Private/RewindDebuggerChooser.cpp b/Engine/Plugins/Chooser/Source/ChooserEditor/Private/RewindDebuggerChooser.cpp
--- a/Engine/Plugins/Chooser/Source/ChooserEditor/Private/RewindDebuggerChooser.cpp
+++ b/Engine/Plugins/Chooser/Source/ChooserEditor/Private/RewindDebuggerChooser.cpp
@@ -4,6 +4,7 @@
 
 #include "Chooser.h"
 #include "ChooserProvider.h"
+#include "GameFramework/Actor.h"
 #include "IGameplayProvider.h"
 
 FRewindDebuggerChooser::FRewindDebuggerChooser()
@@ -37,7 +38,7 @@ void FRewindDebuggerChooser::Update(float DeltaTime, IRewindDebugger* RewindDebu
 					const FObjectInfo& ContextObjectInfo = GameplayProvider->GetObjectInfo(OwnerId);
 					FString DebugName = ContextObjectInfo.Name;
 
-					if (const FObjectInfo* ActorInfo = RewindDebugger->FindOwningActorInfo(GameplayProvider, OwnerId))
+					if (const FObjectInfo* ActorInfo = RewindDebugger->FindTypedOuterInfo<AActor>(GameplayProvider, OwnerId))
 					{
 						DebugName += " in " + FString(ActorInfo->Name);
 					}
diff --git a/Engine/Plugins/Experimental/Animation/TrajectoryTools/Source/TrajectoryTools/Private/TrajectoryRewindDebuggerExtension.cpp b/Engine/Plugins/Experimental/Animation/TrajectoryTools/Source/TrajectoryTools/Private/TrajectoryRewindDebuggerExtension.cpp
--- a/Engine/Plugins/Experimental/Animation/TrajectoryTools/Source/TrajectoryTools/Private/TrajectoryRewindDebuggerExtension.cpp
+++ b/Engine/Plugins/Experimental/Animation/TrajectoryTools/Source/TrajectoryTools/Private/TrajectoryRewindDebuggerExtension.cpp
@@ -369,7 +369,7 @@ void FRewindDebuggerTrajectory::BuildTrajectoryOwnersList(const IRewindDebugger*
 			const bool bIsObjectASkelMesh = InObjectInfo.ClassId == SkelMeshCmpClassInfo->Id || InGameplayProvider->IsSubClassOf(InObjectInfo.ClassId, SkelMeshCmpClassInfo->Id);
 			if (bIsObjectASkelMesh)
 			{
-				const bool bHasOwningActor = InRewindDebugger->FindOwningActorInfo(InGameplayProvider, InObjectInfo.GetUObjectId()) != nullptr;
+				const bool bHasOwningActor = InRewindDebugger->FindTypedOuterInfo<AActor>(InGameplayProvider, InObjectInfo.GetUObjectId()) != nullptr;
 				if (bHasOwningActor)
 				{
 					OutObjectInfos.Push(InObjectInfo);
diff --git a/Engine/Source/Editor/RewindDebuggerInterface/Public/IRewindDebugger.h b/Engine/Source/Editor/RewindDebuggerInterface/Public/IRewindDebugger.h
--- a/Engine/Source/Editor/RewindDebuggerInterface/Public/IRewindDebugger.h
+++ b/Engine/Source/Editor/RewindDebuggerInterface/Public/IRewindDebugger.h
@@ -96,10 +96,16 @@ public:
 	/** @return the current analysis session */
 	virtual const TraceServices::IAnalysisSession* GetAnalysisSession() const = 0;
 
-	/** @return insights id for the selected target actor */
-	virtual uint64 GetTargetActorId() const = 0;
+	/** @return insights id for the root object currently debugged */
+	virtual uint64 GetRootObjectId() const = 0;
 
-	/** @return list of all components of the selected target actor (with the actor as the first element in the list)*/
+	UE_DEPRECATED(5.7, "Use GetRootObjectId instead")
+	virtual uint64 GetTargetActorId() const final
+	{
+		return 0; //GetDebuggedObjectId();
+	}
+
+	/** @return list of all objects associated to the active tracks */
 	virtual TArray<TSharedPtr<FDebugObjectInfo>>& GetDebuggedObjects() = 0;
 
 	UE_DEPRECATED(5.7, "Use GetDebuggedObjects instead")
@@ -133,8 +139,17 @@ public:
 	/** Try selecting the track associated to a given object */
 	virtual void SelectTrack(RewindDebugger::FObjectId ObjectId) = 0;
 
-	/** @return position of the selected target actor (returns true if position is valid) */
-	virtual bool GetTargetActorPosition(FVector& OutPosition) const = 0;
+	/** @return position of the root object being debugged (returns true if position is valid) */
+	virtual bool GetRootObjectPosition(FVector& OutPosition) const = 0;
+
+	/** Stores position of the root object being debugged */
+	virtual void SetRootObjectPosition(const TOptional<FVector>& InPosition) = 0;
+
+	UE_DEPRECATED(5.7, "Use GetRootObjectPosition instead")
+	virtual bool GetTargetActorPosition(FVector& OutPosition) const final
+	{
+		return GetRootObjectPosition(OutPosition);
+	}
 
 	/** @return the world that the debugger is replaying in */
 	virtual UWorld* GetWorldToVisualize() const = 0;
@@ -154,8 +169,21 @@ public:
 	/** Opens the Rewind Debugger details panel tab */
 	virtual void OpenDetailsPanel() = 0;
 
-	/** @return the object information of the first object inheriting from AActor in the outer chain of ObjectId */
-	virtual const FObjectInfo* FindOwningActorInfo(const IGameplayProvider* GameplayProvider, uint64 ObjectId) const = 0;
+	/** @return the object information of the first object inheriting from the specified type in the outer chain of ObjectId */
+	template <typename T>
+	const FObjectInfo* FindTypedOuterInfo(TNotNull<const IGameplayProvider*> GameplayProvider, uint64 ObjectId) const
+	{
+		return FindTypedOuterInfo(T::StaticClass(), GameplayProvider, ObjectId);
+	}
+
+	/** @return the object information of the first object inheriting from the specified type in the outer chain of ObjectId */
+	virtual const FObjectInfo* FindTypedOuterInfo(TNotNull<const UStruct*> Type, TNotNull<const IGameplayProvider*> GameplayProvider, uint64 ObjectId) const = 0;
+
+	UE_DEPRECATED(5.7, "Use FindTypedOuterInfo instead")
+	virtual const FObjectInfo* FindOwningActorInfo(const IGameplayProvider* GameplayProvider, uint64 ObjectId) const final
+	{
+		return nullptr;
+	}
 
 	virtual bool ShouldDisplayWorld(uint64 WorldId) = 0;
 
