Rework recent change to use engine monospace font in Rewind Debugger VLog to use a smaller font size while still using the engine font

diff --git a/Engine/Plugins/Animation/GameplayInsights/Source/RewindDebuggerVLog/Private/RewindDebuggerVLog.cpp b/Engine/Plugins/Animation/GameplayInsights/Source/RewindDebuggerVLog/Private/RewindDebuggerVLog.cpp
--- a/Engine/Plugins/Animation/GameplayInsights/Source/RewindDebuggerVLog/Private/RewindDebuggerVLog.cpp
+++ b/Engine/Plugins/Animation/GameplayInsights/Source/RewindDebuggerVLog/Private/RewindDebuggerVLog.cpp
@@ -93,6 +93,12 @@ void FRewindDebuggerVLog::Initialize()
 	TraceDevice.ImmediateRenderDelegate.BindRaw(this, &FRewindDebuggerVLog::ImmediateRender);
 
 	DelegateHandle = UDebugDrawService::Register(TEXT("VirtualTextureResidency")/*TEXT("VisLog")*/, FDebugDrawDelegate::CreateRaw(this, &FRewindDebuggerVLog::OnShowDebugInfo));
+
+	MonospaceFont = TStrongObjectPtr<UFont>(NewObject<UFont>(GetTransientPackage(), NAME_None, RF_Transient));
+	MonospaceFont->FontCacheType = EFontCacheType::Runtime;
+	MonospaceFont->RuntimeFontSource = ERuntimeFontSource::CoreStyleDefault;
+	MonospaceFont->LegacyFontName = FName("Mono");
+	MonospaceFont->LegacyFontSize = 9;
 }
 
 bool MatchCategoryFilters(const FName& CategoryName, ELogVerbosity::Type Verbosity)
@@ -114,7 +120,7 @@ void FRewindDebuggerVLog::RenderLogEntry(const FVisualLogEntry& Entry, UCanvas*
 	else
 	{
 		UWorld* World = IRewindDebugger::Instance()->GetWorldToVisualize();
-		FVisualLogEntryRenderer::RenderLogEntry(World, Entry, &MatchCategoryFilters, Canvas, GEngine->GetMediumFont(), GEngine->GetMonospaceFont(), ScreenTextY);
+		FVisualLogEntryRenderer::RenderLogEntry(World, Entry, &MatchCategoryFilters, Canvas, GEngine->GetMediumFont(), MonospaceFont.Get(), ScreenTextY);
 
 	}
 }
diff --git a/Engine/Plugins/Animation/GameplayInsights/Source/RewindDebuggerVLog/Private/RewindDebuggerVLog.h b/Engine/Plugins/Animation/GameplayInsights/Source/RewindDebuggerVLog/Private/RewindDebuggerVLog.h
--- a/Engine/Plugins/Animation/GameplayInsights/Source/RewindDebuggerVLog/Private/RewindDebuggerVLog.h
+++ b/Engine/Plugins/Animation/GameplayInsights/Source/RewindDebuggerVLog/Private/RewindDebuggerVLog.h
@@ -47,6 +47,7 @@ private:
 	int32 ScreenTextY = 0;
 
 	FDelegateHandle DelegateHandle;
+	TStrongObjectPtr<UFont> MonospaceFont;
 
 	TArray<uint64> DebuggedObjectIds;
 	TArray<FVisualLogEntry> ImmediateRenderQueue;
