Fix non unity build

diff --git a/Engine/Plugins/Animation/GameplayInsights/Source/RewindDebuggerVLog/Private/RewindDebuggerVLog.h b/Engine/Plugins/Animation/GameplayInsights/Source/RewindDebuggerVLog/Private/RewindDebuggerVLog.h
--- a/Engine/Plugins/Animation/GameplayInsights/Source/RewindDebuggerVLog/Private/RewindDebuggerVLog.h
+++ b/Engine/Plugins/Animation/GameplayInsights/Source/RewindDebuggerVLog/Private/RewindDebuggerVLog.h
@@ -5,8 +5,8 @@
 #include "IRewindDebuggerExtension.h"
 #include "VLogRenderingActor.h"
 #include "ToolMenu.h"
-
-class UFont;
+// Needed for TStrongObjecPtr which requires complete type
+#include "Engine/Font.h"
 
 // Rewind debugger extension for Visual Logger support
 
