Fix typo in comment, ty BugHawk

diff --git a/Engine/Plugins/Animation/GameplayInsights/Source/RewindDebuggerVLog/Private/RewindDebuggerVLog.h b/Engine/Plugins/Animation/GameplayInsights/Source/RewindDebuggerVLog/Private/RewindDebuggerVLog.h
--- a/Engine/Plugins/Animation/GameplayInsights/Source/RewindDebuggerVLog/Private/RewindDebuggerVLog.h
+++ b/Engine/Plugins/Animation/GameplayInsights/Source/RewindDebuggerVLog/Private/RewindDebuggerVLog.h
@@ -5,7 +5,7 @@
 #include "IRewindDebuggerExtension.h"
 #include "VLogRenderingActor.h"
 #include "ToolMenu.h"
-// Needed for TStrongObjecPtr which requires complete type
+// Needed for TStrongObjectPtr which requires complete type
 #include "Engine/Font.h"
 
 // Rewind debugger extension for Visual Logger support
