[RewindDebugger] added missing shape type

diff --git a/Engine/Plugins/Animation/GameplayInsights/Source/RewindDebuggerVLog/Private/VisualLogEntryRenderer.cpp b/Engine/Plugins/Animation/GameplayInsights/Source/RewindDebuggerVLog/Private/VisualLogEntryRenderer.cpp
--- a/Engine/Plugins/Animation/GameplayInsights/Source/RewindDebuggerVLog/Private/VisualLogEntryRenderer.cpp
+++ b/Engine/Plugins/Animation/GameplayInsights/Source/RewindDebuggerVLog/Private/VisualLogEntryRenderer.cpp
@@ -399,6 +399,7 @@ void FVisualLogEntryRenderer::RenderLogEntry(class UWorld* World, const FVisualL
 		break;
 		
 		case EVisualLoggerShapeElement::Circle:
+		case EVisualLoggerShapeElement::WireCircle:
 		{
 			const bool bDrawLabel = (ElementToDraw->Description.IsEmpty() == false);
 				
