Add null check in RewindDebuggerAnimatio::ApplyPoseToMesh

diff --git a/Engine/Plugins/Animation/GameplayInsights/Source/RewindDebugger/Private/RewindDebuggerAnimation.cpp b/Engine/Plugins/Animation/GameplayInsights/Source/RewindDebugger/Private/RewindDebuggerAnimation.cpp
--- a/Engine/Plugins/Animation/GameplayInsights/Source/RewindDebugger/Private/RewindDebuggerAnimation.cpp
+++ b/Engine/Plugins/Animation/GameplayInsights/Source/RewindDebugger/Private/RewindDebuggerAnimation.cpp
@@ -132,27 +132,29 @@ void FRewindDebuggerAnimation::ApplyPoseToMesh(const IAnimationProvider* Animati
 		{
 			if (bApplyMesh)
 			{
-				const FObjectInfo* SkeletalMeshObjectInfo = GameplayProvider->FindObjectInfo(PoseMessage->MeshId);
-				USkeletalMesh* SkeletalMesh = TSoftObjectPtr<USkeletalMesh>(FSoftObjectPath(SkeletalMeshObjectInfo->PathName)).LoadSynchronous();
-										
-				if (SkeletalMesh == nullptr)
+				if (const FObjectInfo* SkeletalMeshObjectInfo = GameplayProvider->FindObjectInfo(PoseMessage->MeshId))
 				{
-					// if the skeletal mesh asset was not found, try the skeleton asset preview mesh as a fallback
-					if (SkeletalMeshInfo->SkeletonId != 0)
+					USkeletalMesh* SkeletalMesh = TSoftObjectPtr<USkeletalMesh>(FSoftObjectPath(SkeletalMeshObjectInfo->PathName)).LoadSynchronous();
+										
+					if (SkeletalMesh == nullptr)
 					{
-						const FObjectInfo& SkeletonInfo = GameplayProvider->GetObjectInfo(SkeletalMeshInfo->SkeletonId);
-						
-						USkeleton* Skeleton = TSoftObjectPtr<USkeleton>(FSoftObjectPath(SkeletonInfo.PathName)).LoadSynchronous();
-						if (Skeleton)
+						// if the skeletal mesh asset was not found, try the skeleton asset preview mesh as a fallback
+						if (SkeletalMeshInfo->SkeletonId != 0)
 						{
-							SkeletalMesh = Skeleton->GetPreviewMesh(true);
+							const FObjectInfo& SkeletonInfo = GameplayProvider->GetObjectInfo(SkeletalMeshInfo->SkeletonId);
+						
+							USkeleton* Skeleton = TSoftObjectPtr<USkeleton>(FSoftObjectPath(SkeletonInfo.PathName)).LoadSynchronous();
+							if (Skeleton)
+							{
+								SkeletalMesh = Skeleton->GetPreviewMesh(true);
+							}
 						}
 					}
-				}
 				
-				if(SkeletalMesh)
-				{
-					MeshComponent->SetSkeletalMesh(SkeletalMesh);
+					if(SkeletalMesh)
+					{
+						MeshComponent->SetSkeletalMesh(SkeletalMesh);
+					}
 				}
 			}
 			
