[RewindDebugger] reverted the defer recording on PIE start since we miss a lot of initialization events during BeginPlay and the first frame.
The original issue that this was addressing can no longer be reproduced and if we manage to reproduce we'll work on a different fix.

diff --git a/Engine/Plugins/Animation/GameplayInsights/Source/RewindDebugger/Private/RewindDebugger.cpp b/Engine/Plugins/Animation/GameplayInsights/Source/RewindDebugger/Private/RewindDebugger.cpp
--- a/Engine/Plugins/Animation/GameplayInsights/Source/RewindDebugger/Private/RewindDebugger.cpp
+++ b/Engine/Plugins/Animation/GameplayInsights/Source/RewindDebugger/Private/RewindDebugger.cpp
@@ -169,7 +169,7 @@ void FRewindDebugger::OnPIEStarted(bool bSimulating)
 
 	if (ShouldAutoRecordOnPIE())
 	{
-		bQueueStartRecording = true;
+		StartRecording();
 	}
 }
 
@@ -976,12 +976,6 @@ void FRewindDebugger::Tick(float DeltaTime)
 {
 	TRACE_CPUPROFILER_EVENT_SCOPE(FRewindDebugger::Tick);
 
-	if (bQueueStartRecording)
-	{
-		StartRecording();
-		bQueueStartRecording = false;
-	}
-
 	if (const TraceServices::IAnalysisSession* Session = GetAnalysisSession())
 	{
 		if (const IGameplayProvider* GameplayProvider = Session->ReadProvider<IGameplayProvider>("GameplayProvider"))
diff --git a/Engine/Plugins/Animation/GameplayInsights/Source/RewindDebugger/Private/RewindDebugger.h b/Engine/Plugins/Animation/GameplayInsights/Source/RewindDebugger/Private/RewindDebugger.h
--- a/Engine/Plugins/Animation/GameplayInsights/Source/RewindDebugger/Private/RewindDebugger.h
+++ b/Engine/Plugins/Animation/GameplayInsights/Source/RewindDebugger/Private/RewindDebugger.h
@@ -254,7 +254,6 @@ private:
 	FOnTrackListChanged TrackListChangedDelegate;
 	FOnTrackCursor TrackCursorDelegate;
 
-	bool bQueueStartRecording = false;
 	bool bTraceJustConnected = false;
 	bool bPIEStarted = false;
 	bool bPIESimulating = false;
