[RewindDebugger] do not attempt to create new track from external request SetObjectToDebug if there is no active analysis session

diff --git a/Engine/Plugins/Animation/GameplayInsights/Source/RewindDebugger/Private/RewindDebugger.cpp b/Engine/Plugins/Animation/GameplayInsights/Source/RewindDebugger/Private/RewindDebugger.cpp
--- a/Engine/Plugins/Animation/GameplayInsights/Source/RewindDebugger/Private/RewindDebugger.cpp
+++ b/Engine/Plugins/Animation/GameplayInsights/Source/RewindDebugger/Private/RewindDebugger.cpp
@@ -1596,6 +1596,12 @@ void FRewindDebugger::SetObjectToDebug(const RewindDebugger::FObjectId ObjectId)
 		return;
 	}
 
+	if (GetAnalysisSession() == nullptr)
+	{
+		UE_LOG(LogRewindDebugger, Log, TEXT("Unable to set the object to debug since there is no active session"));
+		return;
+	}
+
 	CandidateIds.Reset(1);
 	CandidateIds.Push(ObjectId);
 
