[RewindDebugger] allow track to override the default behavior to step backward/forward commands and their tooltip.

diff --git a/Engine/Plugins/Animation/GameplayInsights/Source/RewindDebugger/Private/RewindDebugger.cpp b/Engine/Plugins/Animation/GameplayInsights/Source/RewindDebugger/Private/RewindDebugger.cpp
--- a/Engine/Plugins/Animation/GameplayInsights/Source/RewindDebugger/Private/RewindDebugger.cpp
+++ b/Engine/Plugins/Animation/GameplayInsights/Source/RewindDebugger/Private/RewindDebugger.cpp
@@ -735,11 +735,33 @@ void FRewindDebugger::Step(const int32 InNumberOfFrames)
 
 void FRewindDebugger::StepForward()
 {
+	if (SelectedTrack)
+	{
+		const TOptional<double> NewScrubTime = SelectedTrack->GetStepFrameTime(RewindDebugger::EStepMode::Forward, CurrentScrubTime);
+		if (NewScrubTime.IsSet())
+		{
+			SetCurrentScrubTime(NewScrubTime.GetValue());
+			TrackCursorDelegate.ExecuteIfBound(false);
+			return;
+		}
+	}
+
 	Step(1);
 }
 
 void FRewindDebugger::StepBackward()
 {
+	if (SelectedTrack)
+	{
+		const TOptional<double> NewScrubTime = SelectedTrack->GetStepFrameTime(RewindDebugger::EStepMode::Backward, CurrentScrubTime);
+		if (NewScrubTime.IsSet())
+		{
+			SetCurrentScrubTime(NewScrubTime.GetValue());
+			TrackCursorDelegate.ExecuteIfBound(false);
+			return;
+		}
+	}
+
 	Step(-1);
 }
 
@@ -1307,7 +1329,16 @@ void FRewindDebugger::RegisterToolBar()
 	Section.AddEntry(FToolMenuEntry::InitToolBarButton(
 		Commands.PreviousFrame,
 		FText(),
-		TAttribute<FText>(),
+		TAttribute<FText>::CreateLambda([]()
+			{
+				FText Override;
+				if (Instance()->SelectedTrack.IsValid())
+				{
+					Override = Instance()->SelectedTrack->GetStepCommandTooltip(RewindDebugger::EStepMode::Backward);
+				}
+
+				return Override.IsEmpty() ? FRewindDebuggerCommands::Get().PreviousFrame->GetDescription() : Override;
+			}),
 		FSlateIcon("RewindDebuggerStyle", "RewindDebugger.PreviousFrame.small")));
 
 	Section.AddEntry(FToolMenuEntry::InitToolBarButton(
@@ -1452,7 +1483,16 @@ void FRewindDebugger::RegisterToolBar()
 	Section.AddEntry(FToolMenuEntry::InitToolBarButton(
 		Commands.NextFrame,
 		FText(),
-		TAttribute<FText>(),
+		TAttribute<FText>::CreateLambda([]()
+			{
+				FText Override;
+				if (Instance()->SelectedTrack.IsValid())
+				{
+					Override = Instance()->SelectedTrack->GetStepCommandTooltip(RewindDebugger::EStepMode::Forward);
+				}
+
+				return Override.IsEmpty() ? FRewindDebuggerCommands::Get().NextFrame->GetDescription() : Override;
+			}),
 		FSlateIcon("RewindDebuggerStyle", "RewindDebugger.NextFrame.small")));
 
 	Section.AddEntry(FToolMenuEntry::InitToolBarButton(
diff --git a/Engine/Plugins/Animation/GameplayInsights/Source/RewindDebugger/Private/RewindDebuggerObjectTrack.cpp b/Engine/Plugins/Animation/GameplayInsights/Source/RewindDebugger/Private/RewindDebuggerObjectTrack.cpp
--- a/Engine/Plugins/Animation/GameplayInsights/Source/RewindDebugger/Private/RewindDebuggerObjectTrack.cpp
+++ b/Engine/Plugins/Animation/GameplayInsights/Source/RewindDebugger/Private/RewindDebuggerObjectTrack.cpp
@@ -154,6 +154,28 @@ bool FRewindDebuggerObjectTrack::HandleDoubleClickInternal()
 	return true;
 }
 
+FText FRewindDebuggerObjectTrack::GetStepCommandTooltipInternal(const EStepMode StepMode) const
+{
+	// Allow primary child track to override step tooltip
+	if (PrimaryChildTrack.IsValid())
+	{
+		return PrimaryChildTrack->GetStepCommandTooltip(StepMode);
+	}
+
+	return {};
+}
+
+TOptional<double> FRewindDebuggerObjectTrack::GetStepFrameTimeInternal(const EStepMode StepMode, const double CurrentScrubTime) const
+{
+	// Allow primary child track to override step frame time
+	if (PrimaryChildTrack.IsValid())
+	{
+		return PrimaryChildTrack->GetStepFrameTime(StepMode, CurrentScrubTime);
+	}
+
+	return {};
+}
+
 TConstArrayView<TSharedPtr<FRewindDebuggerTrack>> FRewindDebuggerObjectTrack::GetChildrenInternal(TArray<TSharedPtr<FRewindDebuggerTrack>>& OutTracks) const
 {
 	for (const FTrackCreatorAndTrack& Creator : ChildTrackCreators)
diff --git a/Engine/Plugins/Animation/GameplayInsights/Source/RewindDebugger/Private/RewindDebuggerObjectTrack.h b/Engine/Plugins/Animation/GameplayInsights/Source/RewindDebugger/Private/RewindDebuggerObjectTrack.h
--- a/Engine/Plugins/Animation/GameplayInsights/Source/RewindDebugger/Private/RewindDebuggerObjectTrack.h
+++ b/Engine/Plugins/Animation/GameplayInsights/Source/RewindDebugger/Private/RewindDebuggerObjectTrack.h
@@ -38,6 +38,8 @@ private:
 	virtual uint64 GetObjectIdInternal() const override;
 	virtual bool HasDebugDataInternal() const override;
 	virtual bool HandleDoubleClickInternal() override;
+	virtual FText GetStepCommandTooltipInternal(EStepMode) const override;
+	virtual TOptional<double> GetStepFrameTimeInternal(EStepMode, double) const override;
 
 	mutable FText DisplayName;
 	FString ObjectName;
diff --git a/Engine/Plugins/Runtime/StateTree/Source/StateTreeEditorModule/Private/Debugger/StateTreeRewindDebuggerTrack.cpp b/Engine/Plugins/Runtime/StateTree/Source/StateTreeEditorModule/Private/Debugger/StateTreeRewindDebuggerTrack.cpp
--- a/Engine/Plugins/Runtime/StateTree/Source/StateTreeEditorModule/Private/Debugger/StateTreeRewindDebuggerTrack.cpp
+++ b/Engine/Plugins/Runtime/StateTree/Source/StateTreeEditorModule/Private/Debugger/StateTreeRewindDebuggerTrack.cpp
@@ -238,6 +238,54 @@ bool FRewindDebuggerTrack::HandleDoubleClickInternal()
 	return false;
 }
 
+FText FRewindDebuggerTrack::GetStepCommandTooltipInternal(const RewindDebugger::EStepMode StepMode) const
+{
+	switch (StepMode)
+	{
+	case RewindDebugger::EStepMode::Forward:
+		return LOCTEXT("StepForwardTooltip", "Step to next event in StateTree track");
+	case RewindDebugger::EStepMode::Backward:
+		return LOCTEXT("StepBackwardTooltip", "Step to previous event in StateTree track");
+	}
+
+	return {};
+}
+
+TOptional<double> FRewindDebuggerTrack::GetStepFrameTimeInternal(const RewindDebugger::EStepMode StepMode, const double CurrentScrubTime) const
+{
+	FScrubState ScrubState(&EventCollection);
+	ScrubState.SetScrubTime(CurrentScrubTime, /*ForceRefresh*/true);
+	if (ScrubState.IsInBounds())
+	{
+		const FFrameSpan& Span = EventCollection.FrameSpans[ScrubState.GetFrameSpanIndex()];
+
+		if (StepMode == RewindDebugger::EStepMode::Backward
+			&& CurrentScrubTime > Span.GetWorldTimeStart())
+		{
+			return Span.GetWorldTimeStart();
+		}
+
+		if (StepMode == RewindDebugger::EStepMode::Forward
+			&& CurrentScrubTime < Span.GetWorldTimeEnd())
+		{
+			return Span.GetWorldTimeEnd();
+		}
+	}
+
+	if (StepMode == RewindDebugger::EStepMode::Backward &&
+		ScrubState.HasPreviousFrame())
+	{
+		return ScrubState.GotoPreviousFrame();
+	}
+
+	if (StepMode == RewindDebugger::EStepMode::Forward
+		&& ScrubState.HasNextFrame())
+	{
+		return ScrubState.GotoNextFrame();
+	}
+
+	return {};
+}
 
 } // namespace UE::StateTreeDebugger
 #undef LOCTEXT_NAMESPACE
\ No newline at end of file
diff --git a/Engine/Plugins/Runtime/StateTree/Source/StateTreeEditorModule/Private/Debugger/StateTreeRewindDebuggerTrack.h b/Engine/Plugins/Runtime/StateTree/Source/StateTreeEditorModule/Private/Debugger/StateTreeRewindDebuggerTrack.h
--- a/Engine/Plugins/Runtime/StateTree/Source/StateTreeEditorModule/Private/Debugger/StateTreeRewindDebuggerTrack.h
+++ b/Engine/Plugins/Runtime/StateTree/Source/StateTreeEditorModule/Private/Debugger/StateTreeRewindDebuggerTrack.h
@@ -54,6 +54,8 @@ private:
 	virtual FName GetNameInternal() const override;
 	virtual uint64 GetObjectIdInternal() const override;
 	virtual bool HandleDoubleClickInternal() override;
+	virtual FText GetStepCommandTooltipInternal(RewindDebugger::EStepMode) const override;
+	virtual TOptional<double> GetStepFrameTimeInternal(RewindDebugger::EStepMode, double) const override;
 	//~ End RewindDebugger::FRewindDebuggerTrack interface
 
 	FInstanceEventCollection EventCollection;
diff --git a/Engine/Source/Editor/RewindDebuggerInterface/Public/RewindDebuggerTrack.h b/Engine/Source/Editor/RewindDebuggerInterface/Public/RewindDebuggerTrack.h
--- a/Engine/Source/Editor/RewindDebuggerInterface/Public/RewindDebuggerTrack.h
+++ b/Engine/Source/Editor/RewindDebuggerInterface/Public/RewindDebuggerTrack.h
@@ -27,6 +27,12 @@ enum class ETrackHiddenFlags : uint8
 };
 ENUM_CLASS_FLAGS(ETrackHiddenFlags);
 
+enum class EStepMode : uint8
+{
+	Forward,
+	Backward
+};
+
 class FRewindDebuggerTrack
 {
 public:
@@ -146,6 +152,18 @@ public:
 		return FObjectId(GetObjectIdInternal());
 	}
 
+	/** tracks can customize the tooltip for the step command */
+	FText GetStepCommandTooltip(const EStepMode StepMode) const
+	{
+		return GetStepCommandTooltipInternal(StepMode);
+	}
+
+	/** tracks can return a different time for the step action; default the main recording frame time is used */
+	TOptional<double> GetStepFrameTime(const EStepMode StepMode, const double CurrentScrubTime) const
+	{
+		return GetStepFrameTimeInternal(StepMode, CurrentScrubTime);
+	}
+
 	/**
 	 * Visit the current track along with its child hierarchy recursively
 	 * to call provided function on each track until a call returns EVisitResult::Break
@@ -296,6 +314,16 @@ private:
 		return true;
 	}
 
+	virtual FText GetStepCommandTooltipInternal(EStepMode) const
+	{
+		return {};
+	}
+
+	virtual TOptional<double> GetStepFrameTimeInternal(EStepMode, double) const
+	{
+		return {};
+	}
+
 	UE_DEPRECATED(5.7, "Implement and use GetChildrenInternal instead.")
 	virtual void IterateSubTracksInternal(TFunction<void(TSharedPtr<FRewindDebuggerTrack> SubTrack)> IteratorFunction) final
 	{
