[RewindDebugger] moved Step override implementations inside generic method to test if it can be used and pause if required

diff --git a/Engine/Plugins/Animation/GameplayInsights/Source/RewindDebugger/Private/RewindDebugger.cpp b/Engine/Plugins/Animation/GameplayInsights/Source/RewindDebugger/Private/RewindDebugger.cpp
--- a/Engine/Plugins/Animation/GameplayInsights/Source/RewindDebugger/Private/RewindDebugger.cpp
+++ b/Engine/Plugins/Animation/GameplayInsights/Source/RewindDebugger/Private/RewindDebugger.cpp
@@ -709,6 +709,23 @@ void FRewindDebugger::Step(const int32 InNumberOfFrames)
 	{
 		Pause();
 
+		if (FMath::Abs(InNumberOfFrames) == 1
+			&& SelectedTrack)
+		{
+			const TOptional<double> NewScrubTime = SelectedTrack->GetStepFrameTime(
+				InNumberOfFrames == 1
+					? RewindDebugger::EStepMode::Forward
+					: RewindDebugger::EStepMode::Backward
+				, CurrentScrubTime);
+
+			if (NewScrubTime.IsSet())
+			{
+				SetCurrentScrubTime(NewScrubTime.GetValue());
+				TrackCursorDelegate.ExecuteIfBound(false);
+				return;
+			}
+		}
+
 		if (const TraceServices::IAnalysisSession* Session = GetAnalysisSession())
 		{
 			TraceServices::FAnalysisSessionReadScope SessionReadScope(*Session);
@@ -735,37 +752,14 @@ void FRewindDebugger::Step(const int32 InNumberOfFrames)
 
 void FRewindDebugger::StepForward()
 {
-	if (SelectedTrack)
-	{
-		const TOptional<double> NewScrubTime = SelectedTrack->GetStepFrameTime(RewindDebugger::EStepMode::Forward, CurrentScrubTime);
-		if (NewScrubTime.IsSet())
-		{
-			SetCurrentScrubTime(NewScrubTime.GetValue());
-			TrackCursorDelegate.ExecuteIfBound(false);
-			return;
-		}
-	}
-
 	Step(1);
 }
 
 void FRewindDebugger::StepBackward()
 {
-	if (SelectedTrack)
-	{
-		const TOptional<double> NewScrubTime = SelectedTrack->GetStepFrameTime(RewindDebugger::EStepMode::Backward, CurrentScrubTime);
-		if (NewScrubTime.IsSet())
-		{
-			SetCurrentScrubTime(NewScrubTime.GetValue());
-			TrackCursorDelegate.ExecuteIfBound(false);
-			return;
-		}
-	}
-
 	Step(-1);
 }
 
-
 void FRewindDebugger::ScrubToTime(double ScrubTime, bool bIsScrubbing)
 {
 	if (CanScrub())
