[RewindDebugger] fixed primary track not assigned immediately on track creation

#rnx

diff --git a/Engine/Plugins/Animation/GameplayInsights/Source/RewindDebugger/Private/RewindDebuggerObjectTrack.cpp b/Engine/Plugins/Animation/GameplayInsights/Source/RewindDebugger/Private/RewindDebuggerObjectTrack.cpp
--- a/Engine/Plugins/Animation/GameplayInsights/Source/RewindDebugger/Private/RewindDebuggerObjectTrack.cpp
+++ b/Engine/Plugins/Animation/GameplayInsights/Source/RewindDebugger/Private/RewindDebuggerObjectTrack.cpp
@@ -328,14 +328,6 @@ bool FRewindDebuggerObjectTrack::UpdateInternal()
 
 		if (CreatorAndTrack.Track.IsValid())
 		{
-			if (CreatorAndTrack.Creator->IsCreatingPrimaryChildTrack())
-			{
-				ensureMsgf(!PrimaryChildTrack.IsValid(), TEXT("Multiple primary tracks is not supported yet."));
-				PrimaryChildTrack = CreatorAndTrack.Track;
-				// Hide the track since its data will be exposed by the current Object track
-				PrimaryChildTrack->SetHiddenFlag(ETrackHiddenFlags::HiddenByCode);
-			}
-
 			if (!bHasDebugInfo)
 			{
 				bChanged = true;
@@ -350,6 +342,15 @@ bool FRewindDebuggerObjectTrack::UpdateInternal()
 				CreatorAndTrack.Track = CreatorAndTrack.Creator->CreateTrack(ObjectId);
 			}
 		}
+
+		if (CreatorAndTrack.Track.IsValid()
+			&& CreatorAndTrack.Creator->IsCreatingPrimaryChildTrack())
+		{
+			ensureMsgf(!PrimaryChildTrack.IsValid(), TEXT("Multiple primary tracks is not supported yet."));
+			PrimaryChildTrack = CreatorAndTrack.Track;
+			// Hide the track since its data will be exposed by the current Object track
+			PrimaryChildTrack->SetHiddenFlag(ETrackHiddenFlags::HiddenByCode);
+		}
 	}
 
 	// Fallback code path to add views with no track implementation
