Rewind Debugger: Add fake parent relationship for level sequence actors

#rnx

diff --git a/Engine/Plugins/Animation/GameplayInsights/Source/RewindDebugger/Private/RewindDebuggerObjectTrack.cpp b/Engine/Plugins/Animation/GameplayInsights/Source/RewindDebugger/Private/RewindDebuggerObjectTrack.cpp
--- a/Engine/Plugins/Animation/GameplayInsights/Source/RewindDebugger/Private/RewindDebuggerObjectTrack.cpp
+++ b/Engine/Plugins/Animation/GameplayInsights/Source/RewindDebugger/Private/RewindDebuggerObjectTrack.cpp
@@ -11,6 +11,8 @@
 #include "RewindDebuggerViewCreators.h"
 #include "SEventTimelineView.h"
 #include "Styling/SlateIconFinder.h"
+#include "LevelSequence.h"
+#include "LevelSequenceActor.h"
 
 #define LOCTEXT_NAMESPACE "RewindDebuggerObjectTrack"
 
@@ -392,28 +394,68 @@ bool FRewindDebuggerObjectTrack::UpdateInternal()
 
 		TRange<double> ViewRange = RewindDebugger->GetCurrentViewRange();
 
-		GameplayProvider->EnumerateSubobjects(ObjectId, [this, &FoundObjects, &bChanged, &ViewRange, GameplayProvider](const FObjectId& SubObjectId)
+		auto AddChildObject = [this, &FoundObjects, &bChanged, &ViewRange, GameplayProvider](const FObjectId& SubObjectId)
+		{
+			const TRange<double> Lifetime = GameplayProvider->GetObjectRecordingLifetime(SubObjectId);
+			const TRange<double> Overlap = TRange<double>::Intersection(Lifetime, ViewRange);
+			// only display the track if the lifetime of the object and the view range overlap
+			if (!Overlap.IsEmpty())
 			{
-				const TRange<double> Lifetime = GameplayProvider->GetObjectRecordingLifetime(SubObjectId);
-				const TRange<double> Overlap = TRange<double>::Intersection(Lifetime, ViewRange);
-				// only display the track if the lifetime of the object and the view range overlap
-				if (!Overlap.IsEmpty())
+				const int32 FoundIndex = Children.FindLastByPredicate([SubObjectId](const TSharedPtr<FRewindDebuggerTrack>& Track)
+					{
+						return Track->GetAssociatedObjectId() == SubObjectId;
+					});
+
+				if (FoundIndex == INDEX_NONE)
 				{
-					const int32 FoundIndex = Children.FindLastByPredicate([SubObjectId](const TSharedPtr<FRewindDebuggerTrack>& Track)
-						{
-							return Track->GetAssociatedObjectId() == SubObjectId;
-						});
+					const FObjectInfo& ObjectInfo = GameplayProvider->GetObjectInfo(SubObjectId);
+					Children.Add(MakeShared<FRewindDebuggerObjectTrack>(SubObjectId, ObjectInfo.Name));
+					bChanged = true;
+				}
 
-					if (FoundIndex == INDEX_NONE)
+				FoundObjects.Add(FObjectId{SubObjectId});
+			}
+		};
+
+		GameplayProvider->EnumerateSubobjects(ObjectId, [this, AddChildObject](const FObjectId& SubObjectId)
+		{
+			AddChildObject(SubObjectId);
+		});
+
+		// for now we'll hard code the level sequence special case.
+		// the level sequence is not part of the same world - and it is not using an outer
+		// relationship to the level sequence actor. We introduce a fake parenting relationship
+		// here to allow inspection of the level sequence under the level sequence actor.
+		if (const FObjectInfo* ObjectInfo = GameplayProvider->FindObjectInfo(ObjectId))
+		{
+			if (const FClassInfo* LevelSequenceActorClassInfo = GameplayProvider->FindClassInfo(*ALevelSequenceActor::StaticClass()->GetPathName()))
+			{
+				const uint64 LevelSequenceActorClassId = LevelSequenceActorClassInfo->Id;
+				bool bIsLevelSequenceActor = false;
+	
+				uint64 ClassId = ObjectInfo->ClassId;
+				while (ClassId != FObjectId::InvalidId)
+				{
+					if (ClassId == LevelSequenceActorClassId)
 					{
-						const FObjectInfo& ObjectInfo = GameplayProvider->GetObjectInfo(SubObjectId);
-						Children.Add(MakeShared<FRewindDebuggerObjectTrack>(SubObjectId, ObjectInfo.Name));
-						bChanged = true;
+						bIsLevelSequenceActor = true;
+						break;
 					}
+					const FClassInfo& ClassInfo = GameplayProvider->GetClassInfo(ClassId);
+					ClassId = ClassInfo.SuperId;
+				}
 
-					FoundObjects.Add(FObjectId{SubObjectId});
+				if (bIsLevelSequenceActor)
+				{
+					const ALevelSequenceActor* LevelSequenceActor = Cast<ALevelSequenceActor>(FObjectTrace::GetObjectFromId(ObjectId.GetMainId()));
+					if (LevelSequenceActor)
+					{
+						const FObjectId LevelSequenceObjectId(FObjectTrace::GetObjectId(LevelSequenceActor->LevelSequenceAsset));
+						AddChildObject(LevelSequenceObjectId);
+					}
 				}
-			});
+			}
+		}
 	}
 
 	// add controller and it's component hierarchy if one is attached
diff --git a/Engine/Plugins/Animation/GameplayInsights/Source/RewindDebugger/RewindDebugger.Build.cs b/Engine/Plugins/Animation/GameplayInsights/Source/RewindDebugger/RewindDebugger.Build.cs
--- a/Engine/Plugins/Animation/GameplayInsights/Source/RewindDebugger/RewindDebugger.Build.cs
+++ b/Engine/Plugins/Animation/GameplayInsights/Source/RewindDebugger/RewindDebugger.Build.cs
@@ -38,7 +38,8 @@ namespace UnrealBuildTool.Rules
 				"RewindDebuggerRuntime",
 				"ToolWidgets",
 				"MainFrame",
-				"DeveloperSettings",
+				"DeveloperSettings", 
+				"LevelSequence", "LevelSequence",
 			});
 		}
 	}
