Rewind Debugger: Follow up for 45526455 - add missing ifdef

#rnx

diff --git a/Engine/Plugins/Animation/GameplayInsights/Source/RewindDebugger/Private/RewindDebuggerObjectTrack.cpp b/Engine/Plugins/Animation/GameplayInsights/Source/RewindDebugger/Private/RewindDebuggerObjectTrack.cpp
--- a/Engine/Plugins/Animation/GameplayInsights/Source/RewindDebugger/Private/RewindDebuggerObjectTrack.cpp
+++ b/Engine/Plugins/Animation/GameplayInsights/Source/RewindDebugger/Private/RewindDebuggerObjectTrack.cpp
@@ -422,6 +422,7 @@ bool FRewindDebuggerObjectTrack::UpdateInternal()
 			AddChildObject(SubObjectId);
 		});
 
+#if OBJECT_TRACE_ENABLED
 		// for now we'll hard code the level sequence special case.
 		// the level sequence is not part of the same world - and it is not using an outer
 		// relationship to the level sequence actor. We introduce a fake parenting relationship
@@ -456,6 +457,7 @@ bool FRewindDebuggerObjectTrack::UpdateInternal()
 				}
 			}
 		}
+#endif
 	}
 
 	// add controller and it's component hierarchy if one is attached
