[RewindDebugger] Restore previous VLOG recording state when stopping traces

diff --git a/Engine/Plugins/Animation/GameplayInsights/Source/RewindDebuggerVLogRuntime/Private/RewindDebuggerVLogRuntime.h b/Engine/Plugins/Animation/GameplayInsights/Source/RewindDebuggerVLogRuntime/Private/RewindDebuggerVLogRuntime.h
--- a/Engine/Plugins/Animation/GameplayInsights/Source/RewindDebuggerVLogRuntime/Private/RewindDebuggerVLogRuntime.h
+++ b/Engine/Plugins/Animation/GameplayInsights/Source/RewindDebuggerVLogRuntime/Private/RewindDebuggerVLogRuntime.h
@@ -2,30 +2,39 @@
 
 #pragma once
 
+#if UE_ENABLE_INCLUDE_ORDER_DEPRECATED_IN_5_8
 #include "CoreMinimal.h"
+#endif // UE_ENABLE_INCLUDE_ORDER_DEPRECATED_IN_5_8
+
+#include "EngineDefines.h"
 #include "RewindDebuggerRuntimeInterface/IRewindDebuggerRuntimeExtension.h"
 #include "VisualLogger/VisualLogger.h"
 
 namespace RewindDebugger
 {
-
-	class FRewindDebuggerVLogRuntime : public IRewindDebuggerRuntimeExtension 
+	class FRewindDebuggerVLogRuntime : public IRewindDebuggerRuntimeExtension
 	{
+#if ENABLE_VISUAL_LOG
 	public:
 		virtual void RecordingStarted() override
 		{
-#if ENABLE_VISUAL_LOG	
 			// start recording visual logger data
+			bWasRecording = FVisualLogger::Get().IsRecording();
 			FVisualLogger::Get().SetIsRecordingToTrace(true);
-#endif
 		}
-		
+
 		virtual void RecordingStopped() override
 		{
-#if ENABLE_VISUAL_LOG	
 			// stop recording visual logger data
 			FVisualLogger::Get().SetIsRecordingToTrace(false);
-#endif
+			if (!bWasRecording)
+			{
+				FVisualLogger::Get().SetIsRecording(false);
+			}
 		}
+
+	private:
+		bool bWasRecording = false;
+#endif // ENABLE_VISUAL
 	};
 }
