Fixed Rewind Debugger Property Track's not display child properties info on details panel.
Fixed NotifyType not being properly shown on traced notifies in Rewind Debugger.

diff --git a/Engine/Plugins/Animation/GameplayInsights/Source/GameplayInsights/Private/SNotifiesView.cpp b/Engine/Plugins/Animation/GameplayInsights/Source/GameplayInsights/Private/SNotifiesView.cpp
--- a/Engine/Plugins/Animation/GameplayInsights/Source/GameplayInsights/Private/SNotifiesView.cpp
+++ b/Engine/Plugins/Animation/GameplayInsights/Source/GameplayInsights/Private/SNotifiesView.cpp
@@ -102,7 +102,7 @@ void SNotifiesView::GetVariantsAtFrame(const TraceServices::FFrame& InFrame, TAr
 			InTimeline.EnumerateEvents(InFrame.StartTime, InFrame.EndTime, ProcessEvent);
 		});
 
-		bForEvent = false;
+		bForEvent = true;
 		AnimationProvider->EnumerateNotifyStateTimelines(ObjectId, [&InFrame, &ProcessEvent](uint32 InNotifyNameId, const FAnimationProvider::AnimNotifyTimeline& InTimeline)
 		{
 			InTimeline.EnumerateEvents(InFrame.StartTime, InFrame.EndTime, ProcessEvent);
diff --git a/Engine/Plugins/Animation/GameplayInsights/Source/GameplayInsights/Private/SPropertyTrackView.cpp b/Engine/Plugins/Animation/GameplayInsights/Source/GameplayInsights/Private/SPropertyTrackView.cpp
--- a/Engine/Plugins/Animation/GameplayInsights/Source/GameplayInsights/Private/SPropertyTrackView.cpp
+++ b/Engine/Plugins/Animation/GameplayInsights/Source/GameplayInsights/Private/SPropertyTrackView.cpp
@@ -6,21 +6,59 @@
 #include "PropertyHelpers.h"
 #include "PropertyTrack.h"
 #include "VariantTreeNode.h"
+#include "PropertyWatchManager.h"
+#include "Framework/MultiBox/MultiBoxBuilder.h"
+
+#define LOCTEXT_NAMESPACE "SPropertyTrackView"
 
 void SPropertyTrackView::GetVariantsAtFrame(const TraceServices::FFrame& InFrame,TArray<TSharedRef<FVariantTreeNode>>& OutVariants) const
 {
-	if (TSharedPtr<RewindDebugger::FPropertyTrack> PropertyTrack = Track.Pin())
+	FMemMark Mark(FMemStack::Get());
+	
+	if (const TSharedPtr<RewindDebugger::FPropertyTrack> PropertyTrack = Track.Pin())
 	{
 		TraceServices::FAnalysisSessionReadScope SessionReadScope(*AnalysisSession);
-		const TSharedRef<FVariantTreeNode> Header = OutVariants.Add_GetRef(FVariantTreeNode::MakeHeader(NSLOCTEXT("SPropertyTrackView", "Properties", "Properties"), INDEX_NONE));
-	
+		TSharedPtr<FVariantTreeNode> Header = OutVariants.Add_GetRef(FVariantTreeNode::MakeHeader(FText::FromName(PropertyTrack->GetPropertyName()), INDEX_NONE));
+		
 		if (const FGameplayProvider* GameplayProvider = AnalysisSession->ReadProvider<FGameplayProvider>(FGameplayProvider::ProviderName))
 		{
-			PropertyTrack->ReadObjectPropertyValueAtFrame(InFrame, *GameplayProvider, [GameplayProvider, &PropertyTrack, &Header](const FObjectPropertyValue& InValue, uint32 InValueIndex, const FObjectPropertiesMessage& InMessage)
+			TArray<TSharedPtr<FVariantTreeNode>, TMemStackAllocator<>> PropertyVariants;
+			
+			PropertyTrack->ReadObjectPropertyValueAtFrame(InFrame, *GameplayProvider, [GameplayProvider, &PropertyTrack, &Header, &PropertyVariants, &OutVariants](const FObjectPropertyValue& InValue, uint32 InValueIndex, const FObjectPropertiesMessage& InMessage)
 			{
-				GameplayProvider->ReadObjectPropertiesStorage(PropertyTrack->GetUObjectId(), InMessage, [GameplayProvider, InValueIndex, &PropertyTrack, &Header](const TConstArrayView<FObjectPropertyValue> & InStorage)
+				GameplayProvider->ReadObjectPropertiesStorage(PropertyTrack->GetUObjectId(), InMessage, [GameplayProvider, InValueIndex, &PropertyTrack, &Header, &PropertyVariants, &OutVariants](const TConstArrayView<FObjectPropertyValue> & InStorage)
 				{
-					Header->AddChild(FObjectPropertyHelpers::GetVariantNodeFromProperty(InValueIndex, *GameplayProvider, InStorage));
+					// Handle case where out property has children.
+					for (int32 i = InValueIndex; i < InStorage.Num(); ++i)
+					{
+						PropertyVariants.Add(FObjectPropertyHelpers::GetVariantNodeFromProperty(i, *GameplayProvider, InStorage));
+						
+						// note assumes that order is parent->child in the properties array
+						if (InStorage[i].ParentId != INDEX_NONE && uint32(InStorage[i].ParentId) >= InValueIndex && InStorage[i].ParentId < i)
+						{
+							if (InStorage[i].ParentId != INDEX_NONE && InStorage[i].ParentId != InValueIndex)
+							{
+								PropertyVariants[InStorage[i].ParentId - InValueIndex]->AddChild(PropertyVariants.Last().ToSharedRef());
+							}
+							else
+							{
+								Header->AddChild(PropertyVariants.Last().ToSharedRef());
+							}
+						}
+						else if (i != InValueIndex)
+						{
+							break;
+						}
+					}
+
+					// Single property case.
+					if (Header->GetChildren().IsEmpty())
+					{
+						// No need for a header.
+						OutVariants.Pop();
+						
+						OutVariants.Add(PropertyVariants[0].ToSharedRef());
+					}
 				});
 			});
 		}
@@ -36,7 +74,39 @@ FName SPropertyTrackView::GetName() const
 	return Track.IsValid() ? Track.Pin()->GetName() : FName(TEXT("SPropertyTrackView"));
 }
 
+void SPropertyTrackView::BuildContextMenu(FMenuBuilder& MenuBuilder)
+{
+	SPropertiesDebugViewBase::BuildContextMenu(MenuBuilder);
+
+	MenuBuilder.BeginSection(GetName(), LOCTEXT("DetailsLabel", "Property"));
+	MenuBuilder.AddMenuEntry(
+		LOCTEXT("AddPropertyWatch", "Add Property Watch"),
+		LOCTEXT("AddPropertyTooltip", "Adds a watch track for the given property"),
+		{},
+		FExecuteAction::CreateLambda([this]()
+		{
+			FPropertyWatchManager * PropertyWatchManager = FPropertyWatchManager::Instance();
+			check(PropertyWatchManager)
+
+			if (SelectedPropertyId)
+			{
+				if (!PropertyWatchManager->WatchProperty(ObjectId, SelectedPropertyId))
+				{
+					UE_LOG(LogRewindDebugger, Warning, TEXT("Failed to watch property with id {%d}"), SelectedPropertyId);
+				}
+			}
+			else
+			{
+				UE_LOG(LogRewindDebugger, Warning, TEXT("No property was selected or an invalid property was used while attempting to add a property watch track."))
+			}
+		}));
+	MenuBuilder.EndSection();
+}
+
+
 void SPropertyTrackView::SetTrack(const TWeakPtr<RewindDebugger::FPropertyTrack>& InTrack)
 {
 	Track = InTrack;
 }
+
+#undef LOCTEXT_NAMESPACE
diff --git a/Engine/Plugins/Animation/GameplayInsights/Source/GameplayInsights/Private/SPropertyTrackView.h b/Engine/Plugins/Animation/GameplayInsights/Source/GameplayInsights/Private/SPropertyTrackView.h
--- a/Engine/Plugins/Animation/GameplayInsights/Source/GameplayInsights/Private/SPropertyTrackView.h
+++ b/Engine/Plugins/Animation/GameplayInsights/Source/GameplayInsights/Private/SPropertyTrackView.h
@@ -24,6 +24,7 @@ public:
 	/** Begin SPropertiesDebugViewBase interface */
 	virtual void GetVariantsAtFrame(const TraceServices::FFrame& InFrame, TArray<TSharedRef<FVariantTreeNode>>& OutVariants) const override;
 	virtual FName GetName() const override;
+	virtual void BuildContextMenu(FMenuBuilder& MenuBuilder) override;
 	/** End SPropertiesDebugViewBase interface */
 
 	/** Update the track from which information will be queried and displayed */
