visual logger timeline allocation size

#Synced-CL 49083535

diff --git a/Engine/Plugins/Animation/GameplayInsights/Source/RewindDebuggerVLog/Private/VisualLoggerProvider.cpp b/Engine/Plugins/Animation/GameplayInsights/Source/RewindDebuggerVLog/Private/VisualLoggerProvider.cpp
--- a/Engine/Plugins/Animation/GameplayInsights/Source/RewindDebuggerVLog/Private/VisualLoggerProvider.cpp
+++ b/Engine/Plugins/Animation/GameplayInsights/Source/RewindDebuggerVLog/Private/VisualLoggerProvider.cpp
@@ -42,7 +42,7 @@ void  FVisualLoggerProvider::AppendVisualLogEntry(uint64 InObjectId, double InTi
 {
 	Session.WriteAccessCheck();
 
-	TSharedPtr<TraceServices::TPointTimeline<FVisualLogEntry>> Timeline;
+	TSharedPtr<TraceServices::TPointTimeline<FVisualLogEntry, FVLogTimelineSettings>> Timeline;
 	uint32* IndexPtr = ObjectIdToLogEntryTimelines.Find(InObjectId);
 	if(IndexPtr != nullptr)
 	{
@@ -50,7 +50,7 @@ void  FVisualLoggerProvider::AppendVisualLogEntry(uint64 InObjectId, double InTi
 	}
 	else
 	{
-		Timeline = MakeShared<TraceServices::TPointTimeline<FVisualLogEntry>>(Session.GetLinearAllocator());
+		Timeline = MakeShared<TraceServices::TPointTimeline<FVisualLogEntry, FVLogTimelineSettings>>(Session.GetLinearAllocator());
 		ObjectIdToLogEntryTimelines.Add(InObjectId, LogEntryTimelines.Num());
 		LogEntryTimelines.Add(Timeline.ToSharedRef());
 	}
diff --git a/Engine/Plugins/Animation/GameplayInsights/Source/RewindDebuggerVLog/Private/VisualLoggerProvider.h b/Engine/Plugins/Animation/GameplayInsights/Source/RewindDebuggerVLog/Private/VisualLoggerProvider.h
--- a/Engine/Plugins/Animation/GameplayInsights/Source/RewindDebuggerVLog/Private/VisualLoggerProvider.h
+++ b/Engine/Plugins/Animation/GameplayInsights/Source/RewindDebuggerVLog/Private/VisualLoggerProvider.h
@@ -10,6 +10,14 @@
 
 namespace TraceServices { class IAnalysisSession; }
 
+struct FVLogTimelineSettings
+{
+	enum
+	{
+		EventsPerPage = 4096
+	};
+};
+
 class FVisualLoggerProvider : public IVisualLoggerProvider
 {
 public:
@@ -29,6 +37,6 @@ private:
 	TraceServices::IAnalysisSession& Session;
 
 	TMap<uint64, uint32> ObjectIdToLogEntryTimelines;
-	TArray<TSharedRef<TraceServices::TPointTimeline<FVisualLogEntry>>> LogEntryTimelines;
+	TArray<TSharedRef<TraceServices::TPointTimeline<FVisualLogEntry, FVLogTimelineSettings>>> LogEntryTimelines;
 	TArray<FName> Categories;
 };
