[RewindDebugger] added support for remote sessions provided by TracedBasedDebugger/Analysis modules
- updated toolbar to use RecordingControls
- removed SaveTrace command since analyzed sessions are saved to disk by default
- converted VLog filters from checkbox to toggles to allow users to set multiple ones without closing the menu
- GameplayProvider no longer filters out Server worlds

#Synced-CL 49372183

diff --git a/Engine/Plugins/Animation/GameplayInsights/Source/GameplayInsights/Private/GameplayProvider.cpp b/Engine/Plugins/Animation/GameplayInsights/Source/GameplayInsights/Private/GameplayProvider.cpp
--- a/Engine/Plugins/Animation/GameplayInsights/Source/GameplayInsights/Private/GameplayProvider.cpp
+++ b/Engine/Plugins/Animation/GameplayInsights/Source/GameplayInsights/Private/GameplayProvider.cpp
@@ -770,9 +770,10 @@ void FGameplayProvider::AppendRecordingInfo(uint64 InWorldId, double InProfileTi
 	{
 		if (WorldInfos.IsValidIndex(*WorldIndex))
 		{
-			// Try to only keep track of the client world for now.
 			// We should put each world into its own timeline in the future, to support frame stepping in sync with server frames
-			if (WorldInfos[*WorldIndex].NetMode == FWorldInfo::ENetMode::Standalone || WorldInfos[*WorldIndex].NetMode == FWorldInfo::ENetMode::Client)
+			if (WorldInfos[*WorldIndex].NetMode == FWorldInfo::ENetMode::Standalone
+				|| WorldInfos[*WorldIndex].NetMode == FWorldInfo::ENetMode::Client
+				|| WorldInfos[*WorldIndex].NetMode == FWorldInfo::ENetMode::DedicatedServer)
 			{
 				FRecordingInfoMessage NewRecordingInfo;
 				NewRecordingInfo.WorldId = InWorldId;
diff --git a/Engine/Plugins/Animation/GameplayInsights/Source/RewindDebugger/Private/RewindDebugger.cpp b/Engine/Plugins/Animation/GameplayInsights/Source/RewindDebugger/Private/RewindDebugger.cpp
--- a/Engine/Plugins/Animation/GameplayInsights/Source/RewindDebugger/Private/RewindDebugger.cpp
+++ b/Engine/Plugins/Animation/GameplayInsights/Source/RewindDebugger/Private/RewindDebugger.cpp
@@ -1,4 +1,4 @@
-// Copyright Epic Games, Inc. All Rights Reserved.
+﻿// Copyright Epic Games, Inc. All Rights Reserved.
 
 #include "RewindDebugger.h"
 
@@ -8,7 +8,6 @@
 #include "Framework/Application/SlateApplication.h"
 #include "HAL/PlatformFileManager.h"
 #include "IDesktopPlatform.h"
-#include "Insights/IUnrealInsightsModule.h"
 #include "IRewindDebuggerExtension.h"
 #include "Kismet2/DebuggerCommands.h"
 #include "LevelEditor.h"
@@ -16,15 +15,19 @@
 #include "Modules/ModuleManager.h"
 #include "ObjectTrace.h"
 #include "RewindDebuggerCommands.h"
+#include "RewindDebuggerEngineEditorBridge.h"
 #include "RewindDebuggerModule.h"
 #include "RewindDebuggerObjectTrack.h"
 #include "RewindDebuggerPlaceholderTrack.h"
+#include "RewindDebuggerRecordingControls.h"
 #include "RewindDebuggerRuntime/RewindDebuggerRuntime.h"
 #include "RewindDebuggerSettings.h"
 #include "RewindDebuggerTrackCreators.h"
+#include "SessionInfo.h"
 #include "SModalSessionBrowser.h"
 #include "ToolMenuDelegates.h"
 #include "ToolMenus.h"
+#include "TraceSessionsManager.h"
 #include "UnrealEdGlobals.h"
 #include "Widgets/Docking/SDockTab.h"
 #include "Widgets/Input/SNumericEntryBox.h"
@@ -57,27 +60,31 @@ static void TraceSubobjects(const UObject* OuterObject)
 
 FRewindDebugger::FRewindDebugger()
 {
-	if (RewindDebugger::FRewindDebuggerRuntime::Instance() == nullptr)
+	using namespace RewindDebugger;
+
+#if WITH_TRACE_BASED_DEBUGGERS
+	using namespace UE::TraceBasedDebuggers;
+
+	if (FRewindDebuggerRuntime::Instance() == nullptr)
 	{
-		RewindDebugger::FRewindDebuggerRuntime::Initialize();
+		FRewindDebuggerRuntime::Initialize();
 	}
 
-	if (RewindDebugger::FRewindDebuggerRuntime* Runtime = RewindDebugger::FRewindDebuggerRuntime::Instance())
+	if (FRewindDebuggerRuntime* Runtime = FRewindDebuggerRuntime::Instance())
 	{
-		Runtime->ClearRecording.AddRaw(this, &FRewindDebugger::OnClearRecording);
-		Runtime->RecordingStarted.AddRaw(this, &FRewindDebugger::OnRecordingStarted);
-		Runtime->RecordingStopped.AddRaw(this, &FRewindDebugger::OnRecordingStopped);
+		ClearRecordingHandle = Runtime->ClearRecording.AddRaw(this, &FRewindDebugger::OnClearRecording);
+		RecordingStartedHandle = Runtime->RecordingStarted.AddRaw(this, &FRewindDebugger::OnRecordingStarted);
+		RecordingStoppedHandle = Runtime->RecordingStopped.AddRaw(this, &FRewindDebugger::OnRecordingStopped);
 	}
+#endif // WITH_TRACE_BASED_DEBUGGERS
 
-	RewindDebugger::FRewindDebuggerTrackCreators::EnumerateCreators([this](const RewindDebugger::IRewindDebuggerTrackCreator* Creator)
+	FRewindDebuggerTrackCreators::EnumerateCreators([this](const IRewindDebuggerTrackCreator* Creator)
 		{
 			Creator->GetTrackTypes(TrackTypes);
 		});
 
 	RecordingDuration.Set(0);
 
-	UnrealInsightsModule = &FModuleManager::LoadModuleChecked<IUnrealInsightsModule>("TraceInsights");
-
 	PreBeginPIEHandle = FEditorDelegates::PreBeginPIE.AddRaw(this, &FRewindDebugger::OnPIEStarted);
 	PausePIEHandle = FEditorDelegates::PausePIE.AddRaw(this, &FRewindDebugger::OnPIEPaused);
 	ResumePIEHandle = FEditorDelegates::ResumePIE.AddRaw(this, &FRewindDebugger::OnPIEResumed);
@@ -99,7 +106,7 @@ FRewindDebugger::FRewindDebugger()
 			GetTargetObjectIds(CandidateIds);
 			// make sure all the SubObjects of the root object have been traced
 #if OBJECT_TRACE_ENABLED
-			for (const RewindDebugger::FObjectId& TargetObjectId : CandidateIds)
+			for (const FObjectId& TargetObjectId : CandidateIds)
 			{
 				if (const UObject* TargetObject = FObjectTrace::GetObjectFromId(TargetObjectId.GetMainId()))
 				{
@@ -119,6 +126,20 @@ FRewindDebugger::FRewindDebugger()
 
 			return true;
 		});
+
+#if WITH_TRACE_BASED_DEBUGGERS
+	using namespace UE::RewindDebugger;
+	if (const TSharedPtr<FRemoteSessionsManager> RemoteSessionManager = FRewindDebuggerEngineEditorBridge::Get().GetSessionsManager())
+	{
+		// Create remote sessions recording controls and register them to the main toolbar
+		RecordingControls = MakeShared<FDebuggerRecordingControls>(FRewindDebuggerModule::MainStatusBarName
+			, RemoteSessionManager.ToSharedRef()
+			, DebuggerGuid
+			, LogRewindDebugger);
+		RecordingControls->Initialize();
+		RecordingControls->AddToMenu(FRewindDebuggerModule::MainToolBarName);
+	}
+#endif // WITH_TRACE_BASED_DEBUGGERS
 }
 
 FRewindDebugger::~FRewindDebugger()
@@ -131,10 +152,14 @@ FRewindDebugger::~FRewindDebugger()
 
 	FTSTicker::RemoveTicker(TickerHandle);
 
+#if WITH_TRACE_BASED_DEBUGGERS
 	if (RewindDebugger::FRewindDebuggerRuntime* Runtime = RewindDebugger::FRewindDebuggerRuntime::Instance())
 	{
-		Runtime->RecordingStarted.RemoveAll(this);
+		Runtime->ClearRecording.Remove(ClearRecordingHandle);
+		Runtime->RecordingStarted.Remove(RecordingStartedHandle);
+		Runtime->RecordingStopped.Remove(RecordingStoppedHandle);
 	}
+#endif // WITH_TRACE_BASED_DEBUGGERS
 }
 
 void FRewindDebugger::Initialize()
@@ -225,7 +250,6 @@ void FRewindDebugger::OnPIESingleStepped(bool bSimulating)
 #endif // OBJECT_TRACE_ENABLED
 }
 
-
 void FRewindDebugger::OnPIEStopped(bool bSimulating)
 {
 	if (IsRecording() && bPIESimulating)
@@ -371,16 +395,16 @@ void FRewindDebugger::RefreshDebugTracks()
 	}
 }
 
-void FRewindDebugger::OnConnection()
-{
-	// queue up some operations to happen on the game thread next tick
-	bTraceJustConnected = true;
-	FTraceAuxiliary::OnConnection.RemoveAll(this);
-}
-
 bool FRewindDebugger::CanStartRecording() const
 {
-	return !IsRecording() && bPIESimulating;
+	// Relying on the recording controls since they rely on the
+	// sessions manager which also handles active sessions from other debuggers.
+	if (RecordingControls)
+	{
+		return RecordingControls->CanStartRecording();
+	}
+
+	return false;
 }
 
 void FRewindDebugger::StartRecording() const
@@ -390,12 +414,31 @@ void FRewindDebugger::StartRecording() const
 		return;
 	}
 
-	if (RewindDebugger::FRewindDebuggerRuntime* Runtime = RewindDebugger::FRewindDebuggerRuntime::Instance())
+	if (RecordingControls)
+	{
+		RecordingControls->StartRecording();
+	}
+}
+
+void FRewindDebugger::StopRecording() const
+{
+	if (RecordingControls)
 	{
-		Runtime->StartRecording();
+		RecordingControls->StopRecording();
 	}
 }
 
+#if WITH_TRACE_BASED_DEBUGGERS
+void FRewindDebugger::StopLocalRecording() const
+{
+	RewindDebugger::FRewindDebuggerRuntime* Runtime = RewindDebugger::FRewindDebuggerRuntime::Instance();
+	if (Runtime != nullptr && Runtime->IsRecording())
+	{
+		Runtime->StopRecording();
+	}
+}
+#endif // WITH_TRACE_BASED_DEBUGGERS
+
 void FRewindDebugger::OnClearRecording()
 {
 	ClearTrace();
@@ -412,13 +455,18 @@ void FRewindDebugger::OnClearRecording()
 
 void FRewindDebugger::OnRecordingStarted()
 {
-	IterateExtensions([this](IRewindDebuggerExtension* Extension)
+	// Previous implementation calling the following code was executed on the game thread
+	// after calling FTraceAuxiliary::Start, but we are now waiting for the connection
+	// to be ready (i.e. FTraceAuxiliary::OnConnection) which comes from a worker thread.
+	// Extensions code in 'RecordingStarted' might require to run on the Game thread so this
+	// is why we defer the callbacks.
+	ExecuteOnGameThread(UE_SOURCE_LOCATION, [this]
 		{
-			Extension->RecordingStarted(this);
-		}
-	);
-
-	UnrealInsightsModule->StartAnalysisForLastLiveSession(5.0);
+			IterateExtensions([this](IRewindDebuggerExtension* Extension)
+				{
+					Extension->RecordingStarted(this);
+				});
+		});
 }
 
 void FRewindDebugger::OnRecordingStopped()
@@ -432,18 +480,30 @@ void FRewindDebugger::OnRecordingStopped()
 
 bool FRewindDebugger::CanOpenTrace() const
 {
-	return !bPIEStarted;
+	return !bPIEStarted && !IsRecording();
 }
 
 
 void FRewindDebugger::OpenTrace(const FString& FilePath)
 {
-	ClearTrace();
+	ClearTraceAndActiveSession();
 
 	bDisplayWorldIdValid = false;
 
-	IUnrealInsightsModule& TraceInsightsModule = FModuleManager::LoadModuleChecked<IUnrealInsightsModule>("TraceInsights");
-	TraceInsightsModule.StartAnalysisForTraceFile(*FilePath);
+#if UE_WITH_TRACE_BASED_DEBUGGERS_ANALYSIS
+	using namespace UE::TraceBasedDebuggers;
+	if (const TSharedPtr<FTraceSessionsManager> TraceSessionsManager = FRewindDebuggerModule::Get().GetTraceSessionsManager())
+	{
+		FTraceSessionDescriptor NewSessionFromFileDescriptor;
+		NewSessionFromFileDescriptor.SessionName = TraceSessionsManager->LoadTraceFile(FilePath);
+		NewSessionFromFileDescriptor.bIsLiveSession = false;
+
+		if (NewSessionFromFileDescriptor.IsValid())
+		{
+			OpenSession(NewSessionFromFileDescriptor);
+		}
+	}
+#endif // UE_WITH_TRACE_BASED_DEBUGGERS_ANALYSIS
 
 	// todo: optionally open the map the trace file was recorded in
 }
@@ -481,19 +541,17 @@ void FRewindDebugger::OpenTrace()
 
 void FRewindDebugger::AttachToSession()
 {
-	ClearTrace();
 	const TSharedRef<SModalSessionBrowser> SessionBrowserModal = SNew(SModalSessionBrowser);
 
 	if (SessionBrowserModal->ShowModal() != EAppReturnType::Cancel)
 	{
+		ClearTrace();
+
 		bool bSuccess = false;
 		const SModalSessionBrowser::FTraceSessionInfo SessionInfo = SessionBrowserModal->GetSelectedTraceInfo();
 		if (SessionInfo.bIsValid)
 		{
-			const FString SessionAddress = SessionBrowserModal->GetSelectedTraceStoreAddress();
-			IUnrealInsightsModule& TraceInsightsModule = FModuleManager::LoadModuleChecked<IUnrealInsightsModule>("TraceInsights");
-			TraceInsightsModule.StartAnalysisForTrace(SessionInfo.TraceID);
-			bSuccess = TraceInsightsModule.GetAnalysisSession().IsValid();
+			bSuccess = ConnectToLiveSession(SessionInfo.TraceID, SessionBrowserModal->GetSelectedTraceStoreAddress());
 		}
 
 		if (!bSuccess)
@@ -510,7 +568,10 @@ bool FRewindDebugger::CanClearTrace() const
 
 void FRewindDebugger::ClearTrace()
 {
-	StopRecording();
+#if WITH_TRACE_BASED_DEBUGGERS
+	StopLocalRecording();
+#endif
+
 	RecordingDuration.Set(0);
 
 	CandidateIds.Empty();
@@ -528,64 +589,15 @@ void FRewindDebugger::ClearTrace()
 		}
 	);
 
-	IUnrealInsightsModule& TraceInsightsModule = FModuleManager::LoadModuleChecked<IUnrealInsightsModule>("TraceInsights");
-	// only way I can find to clear the session is trying to load a name that doesn't exist.
-	TraceInsightsModule.StartAnalysisForTraceFile(TEXT("0"));
-
 	RefreshDebugTracks();
 }
 
-bool FRewindDebugger::CanSaveTrace() const
-{
-	const TraceServices::IAnalysisSession* Session = GetAnalysisSession();
-	return Session != nullptr && Session->IsAnalysisComplete();
-}
-
-void FRewindDebugger::SaveTrace(FString FileName)
-{
-	if (const TraceServices::IAnalysisSession* Session = GetAnalysisSession())
-	{
-		if (Session->IsAnalysisComplete())
-		{
-			const FString SourceFileName = Session->GetName();
-
-			FPlatformFileManager& FileManager = FPlatformFileManager::Get();
-			IPlatformFile& PlatformFile = FileManager.GetPlatformFile();
-
-			PlatformFile.CopyFile(*FileName, *SourceFileName);
-		}
-
-	}
-}
-
-void FRewindDebugger::SaveTrace()
+void FRewindDebugger::ClearTraceAndActiveSession()
 {
-	const FString FolderPath = "";
-
-	TArray<FString> OutFilenames;
-	if (IDesktopPlatform* DesktopPlatform = FDesktopPlatformModule::Get())
-	{
-		FString ExtensionStr;
-		ExtensionStr += TEXT("Rewind Debugger Recording |*.utrace|");
-
-		DesktopPlatform->SaveFileDialog(
-			FSlateApplication::Get().FindBestParentWindowHandleForDialogs(nullptr),
-			LOCTEXT("SaveDialogTitle", "Save Rewind Debugger Recording").ToString(),
-			FolderPath,
-			TEXT(""),
-			*ExtensionStr,
-			EFileDialogFlags::None,
-			OutFilenames
-		);
-	}
-
-	if (OutFilenames.Num() > 0)
-	{
-		if (OutFilenames[0].EndsWith(TEXT(".utrace")))
-		{
-			SaveTrace(OutFilenames[0]);
-		}
-	}
+#if UE_WITH_TRACE_BASED_DEBUGGERS_ANALYSIS
+	CloseActiveTraceSession();
+#endif
+	ClearTrace();
 }
 
 bool FRewindDebugger::ShouldAutoRecordOnPIE() const
@@ -614,12 +626,9 @@ void FRewindDebugger::SetShouldAutoEject(bool value)
 	RewindDebuggerSettings.SaveConfig();
 }
 
-void FRewindDebugger::StopRecording()
+bool FRewindDebugger::CanUsePlaybackControls() const
 {
-	if (RewindDebugger::FRewindDebuggerRuntime* Runtime = RewindDebugger::FRewindDebuggerRuntime::Instance())
-	{
-		Runtime->StopRecording();
-	}
+	return !bPIESimulating && !IsRecording();
 }
 
 bool FRewindDebugger::CanPause() const
@@ -642,12 +651,12 @@ void FRewindDebugger::Pause()
 
 bool FRewindDebugger::IsPlaying() const
 {
-	return ControlState == EControlState::Play && !bPIESimulating;
+	return ControlState == EControlState::Play && CanUsePlaybackControls();
 }
 
 bool FRewindDebugger::CanPlay() const
 {
-	return ControlState != EControlState::Play && !bPIESimulating && RecordingDuration.Get() > 0;
+	return ControlState != EControlState::Play && CanUsePlaybackControls() && RecordingDuration.Get() > 0;
 }
 
 void FRewindDebugger::Play()
@@ -665,7 +674,7 @@ void FRewindDebugger::Play()
 
 bool FRewindDebugger::CanPlayReverse() const
 {
-	return ControlState != EControlState::PlayReverse && !bPIESimulating && RecordingDuration.Get() > 0;
+	return ControlState != EControlState::PlayReverse && CanUsePlaybackControls() && RecordingDuration.Get() > 0;
 }
 
 void FRewindDebugger::PlayReverse()
@@ -683,7 +692,7 @@ void FRewindDebugger::PlayReverse()
 
 bool FRewindDebugger::CanScrub() const
 {
-	return !bPIESimulating && RecordingDuration.Get() > 0;
+	return CanUsePlaybackControls() && RecordingDuration.Get() > 0;
 }
 
 void FRewindDebugger::ScrubToStart()
@@ -790,11 +799,7 @@ UWorld* FRewindDebugger::GetWorldToVisualize() const
 
 bool FRewindDebugger::IsRecording() const
 {
-	if (const RewindDebugger::FRewindDebuggerRuntime* Runtime = RewindDebugger::FRewindDebuggerRuntime::Instance())
-	{
-		return Runtime->IsRecording();
-	}
-	return false;
+	return RecordingControls ? RecordingControls->IsRecording() : false;
 }
 
 bool FRewindDebugger::IsTraceFileLoaded() const
@@ -956,12 +961,22 @@ void FRewindDebugger::GetScrubTimeInformation(double InDebugTime, FScrubTimeInfo
 
 const TraceServices::IAnalysisSession* FRewindDebugger::GetAnalysisSession() const
 {
-	if (UnrealInsightsModule == nullptr)
+	return GetAnalysisSessionAsShared().Get();
+}
+
+TSharedPtr<const TraceServices::IAnalysisSession> FRewindDebugger::GetAnalysisSessionAsShared() const
+{
+	TSharedPtr<const TraceServices::IAnalysisSession> Session = nullptr;
+
+#if UE_WITH_TRACE_BASED_DEBUGGERS_ANALYSIS
+	using namespace UE::TraceBasedDebuggers;
+	if (const TSharedPtr<FTraceSessionsManager> TraceManager = FRewindDebuggerModule::Get().GetTraceSessionsManager())
 	{
-		UnrealInsightsModule = &FModuleManager::LoadModuleChecked<IUnrealInsightsModule>("TraceInsights");
+		Session = TraceManager->GetSession(CurrentSessionDescriptor.SessionName);
 	}
+#endif // UE_WITH_TRACE_BASED_DEBUGGERS_ANALYSIS
 
-	return UnrealInsightsModule ? UnrealInsightsModule->GetAnalysisSession().Get() : nullptr;
+	return Session;
 }
 
 uint64 FRewindDebugger::GetRootObjectId() const
@@ -991,6 +1006,131 @@ const FObjectInfo* FRewindDebugger::FindTypedOuterInfo(TNotNull<const UStruct*>
 	}
 }
 
+#if UE_WITH_TRACE_BASED_DEBUGGERS_ANALYSIS
+
+bool FRewindDebugger::ConnectToLiveSession(const uint32 SessionID, const FStringView InSessionAddress)
+{
+	CloseActiveTraceSession();
+
+	using namespace UE::TraceBasedDebuggers;
+	if (const TSharedPtr<FTraceSessionsManager> TraceSessionsManager = FRewindDebuggerModule::Get().GetTraceSessionsManager())
+	{
+		FTraceSessionDescriptor NewSessionDescriptor;
+		NewSessionDescriptor.bIsLiveSession = true;
+		NewSessionDescriptor.SessionName = TraceSessionsManager->ConnectToLiveSession(InSessionAddress, SessionID);
+
+		if (NewSessionDescriptor.IsValid())
+		{
+			OpenSession(NewSessionDescriptor);
+			return true;
+		}
+	}
+
+	return false;
+}
+
+bool FRewindDebugger::ConnectToLiveSession_Direct(const FGuid& RemoteSessionID)
+{
+	CloseActiveTraceSession();
+
+	using namespace UE::TraceBasedDebuggers;
+	if (const TSharedPtr<FTraceSessionsManager> TraceSessionsManager = FRewindDebuggerModule::Get().GetTraceSessionsManager())
+	{
+		FTraceSessionDescriptor NewSessionDescriptor;
+		NewSessionDescriptor.RemoteSessionID = RemoteSessionID;
+		NewSessionDescriptor.bIsLiveSession = true;
+		NewSessionDescriptor.SessionName = TraceSessionsManager->ConnectToLiveSession_Direct(RemoteSessionID, NewSessionDescriptor.SessionPort);
+
+		if (NewSessionDescriptor.IsValid())
+		{
+			OpenSession(NewSessionDescriptor);
+			return true;
+		}
+	}
+
+	return false;
+}
+
+bool FRewindDebugger::ConnectToLiveSession_Relay(const FGuid& RemoteSessionID)
+{
+	CloseActiveTraceSession();
+
+	using namespace UE::TraceBasedDebuggers;
+	if (const TSharedPtr<FTraceSessionsManager> TraceSessionsManager = FRewindDebuggerModule::Get().GetTraceSessionsManager())
+	{
+		FTraceSessionDescriptor NewSessionDescriptor;
+		NewSessionDescriptor.bIsLiveSession = true;
+		NewSessionDescriptor.RemoteSessionID = RemoteSessionID;
+		NewSessionDescriptor.SessionName = TraceSessionsManager->ConnectToLiveSession_Relay(RemoteSessionID);
+
+		if (NewSessionDescriptor.IsValid())
+		{
+			OpenSession(NewSessionDescriptor);
+			return true;
+		}
+	}
+
+	return false;
+}
+
+void FRewindDebugger::OpenSession(const UE::TraceBasedDebuggers::FTraceSessionDescriptor& SessionDescriptor)
+{
+	ensureMsgf(!CurrentSessionDescriptor.IsValid(), TEXT("Expecting current session to be closed before opening a new one"));
+	CurrentSessionDescriptor = SessionDescriptor;
+}
+
+const UE::TraceBasedDebuggers::FTraceSessionDescriptor& FRewindDebugger::GetSessionDescriptor(const FGuid& RemoteSessionID) const
+{
+	if (CurrentSessionDescriptor.IsValid()
+		&& CurrentSessionDescriptor.RemoteSessionID == RemoteSessionID)
+	{
+		return CurrentSessionDescriptor;
+	}
+
+	static UE::TraceBasedDebuggers::FTraceSessionDescriptor InvalidSessionDescriptor;
+	return InvalidSessionDescriptor;
+}
+
+void FRewindDebugger::CloseActiveTraceSession()
+{
+	using namespace UE::TraceBasedDebuggers;
+	if (const TSharedPtr<FTraceSessionsManager> TraceManager = FRewindDebuggerModule::Get().GetTraceSessionsManager())
+	{
+		if (CurrentSessionDescriptor.IsValid())
+		{
+			TraceManager->CloseSession(CurrentSessionDescriptor.SessionName);
+		}
+	}
+
+	CurrentSessionDescriptor = {};
+}
+
+void FRewindDebugger::CloseSessionByRemoteSessionID(const FGuid& RemoteSessionID)
+{
+	using namespace UE::TraceBasedDebuggers;
+
+	if (!CurrentSessionDescriptor.IsValid()
+		|| CurrentSessionDescriptor.RemoteSessionID != RemoteSessionID)
+	{
+		return;
+	}
+
+	CloseActiveTraceSession();
+}
+
+void FRewindDebugger::StopActiveTraceSessions() const
+{
+	using namespace UE::TraceBasedDebuggers;
+	if (const TSharedPtr<FTraceSessionsManager> TraceManager = FRewindDebuggerModule::Get().GetTraceSessionsManager())
+	{
+		if (CurrentSessionDescriptor.IsValid())
+		{
+			TraceManager->StopSession(CurrentSessionDescriptor.SessionName);
+		}
+	}
+}
+#endif // UE_WITH_TRACE_BASED_DEBUGGERS_ANALYSIS
+
 void FRewindDebugger::Tick(float DeltaTime)
 {
 	TRACE_CPUPROFILER_EVENT_SCOPE(FRewindDebugger::Tick);
@@ -1028,18 +1168,23 @@ void FRewindDebugger::Tick(float DeltaTime)
 			}
 
 			const double RecordingDurationValue = GameplayProvider->GetRecordingDuration();
+			bool bNewDataAvailable = false;
 			if (IsTraceFileLoaded() && RecordingDurationValue > RecordingDuration.Get())
 			{
-				// while trace file is loading up, force the trace range to update.
+				// while trace file is loading up or analysing a live session, force the trace range to update.
 				SetCurrentViewRange(GetCurrentViewRange());
+				bNewDataAvailable = true;
 			}
 			RecordingDuration.Set(RecordingDurationValue);
 
 			RefreshDebugTracks();
 
-			if (bPIESimulating)
+			const bool bRequiresScrubTimeUpdate = bPIESimulating || bNewDataAvailable;
+
+			if (bRequiresScrubTimeUpdate)
 			{
-				if (IsRecording())
+				if (IsRecording()
+					|| bNewDataAvailable)
 				{
 					TRACE_CPUPROFILER_EVENT_SCOPE(FRewindDebugger::Tick_UpdateSimulating);
 					SetCurrentScrubTime(RecordingDurationValue);
@@ -1313,7 +1458,7 @@ void FRewindDebugger::MakeWorldsMenu(UToolMenu* Menu)
 
 void FRewindDebugger::RegisterToolBar()
 {
-	UToolMenu* Menu = UToolMenus::Get()->RegisterMenu("RewindDebugger.ToolBar", NAME_None, EMultiBoxType::ToolBar);
+	UToolMenu* Menu = UToolMenus::Get()->RegisterMenu(FRewindDebuggerModule::MainToolBarName, NAME_None, EMultiBoxType::ToolBar);
 
 	FToolMenuSection& Section = Menu->FindOrAddSection("VCRControls");
 	const FRewindDebuggerCommands& Commands = FRewindDebuggerCommands::Get();
@@ -1498,19 +1643,6 @@ void FRewindDebugger::RegisterToolBar()
 		TAttribute<FText>(),
 		FSlateIcon("RewindDebuggerStyle", "RewindDebugger.LastFrame.small")));
 
-	Section.AddEntry(FToolMenuEntry::InitToolBarButton(
-		Commands.StartRecording,
-		FText(),
-		TAttribute<FText>(),
-		FSlateIcon("RewindDebuggerStyle", "RewindDebugger.StartRecording.small")));
-
-
-	Section.AddEntry(FToolMenuEntry::InitToolBarButton(
-		Commands.StopRecording,
-		FText(),
-		TAttribute<FText>(),
-		FSlateIcon("RewindDebuggerStyle", "RewindDebugger.StopRecording.small")));
-
 	Section.AddSeparator(NAME_None);
 
 	Section.AddEntry(FToolMenuEntry::InitToolBarButton(
@@ -1525,12 +1657,6 @@ void FRewindDebugger::RegisterToolBar()
 		TAttribute<FText>(),
 		FSlateIcon(FAppStyle::GetAppStyleSetName(), "Icons.FolderOpen")));
 
-	Section.AddEntry(FToolMenuEntry::InitToolBarButton(
-		Commands.SaveTrace,
-		FText(),
-		TAttribute<FText>(),
-		FSlateIcon(FAppStyle::GetAppStyleSetName(), "Icons.Save")));
-
 	Section.AddEntry(FToolMenuEntry::InitToolBarButton(
 		Commands.ClearTrace,
 		FText(),
diff --git a/Engine/Plugins/Animation/GameplayInsights/Source/RewindDebugger/Private/RewindDebugger.h b/Engine/Plugins/Animation/GameplayInsights/Source/RewindDebugger/Private/RewindDebugger.h
--- a/Engine/Plugins/Animation/GameplayInsights/Source/RewindDebugger/Private/RewindDebugger.h
+++ b/Engine/Plugins/Animation/GameplayInsights/Source/RewindDebugger/Private/RewindDebugger.h
@@ -1,4 +1,4 @@
-// Copyright Epic Games, Inc. All Rights Reserved.
+﻿// Copyright Epic Games, Inc. All Rights Reserved.
 #pragma once
 
 #if UE_ENABLE_INCLUDE_ORDER_DEPRECATED_IN_5_7
@@ -11,11 +11,18 @@
 #include "IRewindDebugger.h"
 #include "IRewindDebuggerTrackCreator.h"
 #include "RewindDebuggerTrack.h"
+#include "TraceSessionsManager.h"
 
 namespace TraceServices
 {
-	class IAnalysisSession;
+class IAnalysisSession;
 }
+
+namespace UE::RewindDebugger
+{
+struct FDebuggerRecordingControls;
+}
+
 class FMenuBuilder;
 class IGameplayProvider;
 class SDockTab;
@@ -57,6 +64,7 @@ public:
 	}
 
 	virtual const TraceServices::IAnalysisSession* GetAnalysisSession() const override;
+	virtual TSharedPtr<const TraceServices::IAnalysisSession> GetAnalysisSessionAsShared() const override;
 	virtual uint64 GetRootObjectId() const override;
 	virtual bool GetRootObjectPosition(FVector& OutPosition) const override;
 	virtual void SetRootObjectPosition(const TOptional<FVector>& InPosition) override;
@@ -84,8 +92,6 @@ public:
 		return DisplayWorldId == WorldId;
 	}
 
-	void OnConnection();
-
 	void GetTargetObjectIds(TArray<RewindDebugger::FObjectId>& OutTargetObjectIds) const;
 
 	TArray<TSharedPtr<RewindDebugger::FRewindDebuggerTrack>>& GetTracks()
@@ -121,10 +127,7 @@ public:
 
 	bool CanClearTrace() const;
 	void ClearTrace();
-
-	bool CanSaveTrace() const;
-	void SaveTrace(FString FileName);
-	void SaveTrace();
+	void ClearTraceAndActiveSession();
 
 	bool ShouldAutoRecordOnPIE() const;
 	void SetShouldAutoRecordOnPIE(bool InValue);
@@ -133,7 +136,7 @@ public:
 	void SetShouldAutoEject(bool InValue);
 
 	/** Stop recording: Stop tracing Object + Animation Data. */
-	void StopRecording();
+	void StopRecording() const;
 	bool CanStopRecording() const
 	{
 		return IsRecording();
@@ -219,7 +222,33 @@ public:
 		return TrackTypes;
 	};
 
+#if UE_WITH_TRACE_BASED_DEBUGGERS_ANALYSIS
+public:
+	bool ConnectToLiveSession(uint32 SessionID, FStringView InSessionAddress);
+	bool ConnectToLiveSession_Direct(const FGuid& RemoteSessionID);
+	bool ConnectToLiveSession_Relay(const FGuid& RemoteSessionID);
+
+	void CloseSessionByRemoteSessionID(const FGuid& RemoteSessionID);
+
+	const UE::TraceBasedDebuggers::FTraceSessionDescriptor& GetCurrentSessionDescriptor() const
+	{
+		return CurrentSessionDescriptor;
+	}
+
+private:
+	void OpenSession(const UE::TraceBasedDebuggers::FTraceSessionDescriptor& SessionDescriptor);
+	void CloseActiveTraceSession();
+	void StopActiveTraceSessions() const;
+
+	const UE::TraceBasedDebuggers::FTraceSessionDescriptor& GetSessionDescriptor(const FGuid& RemoteSessionID) const;
+#endif // UE_WITH_TRACE_BASED_DEBUGGERS_ANALYSIS
+
 private:
+
+#if WITH_TRACE_BASED_DEBUGGERS
+	void StopLocalRecording() const;
+#endif
+
 	static void RefreshDebuggedObjects(TArray<TSharedPtr<RewindDebugger::FRewindDebuggerTrack>>& InTracks, TArray<TSharedPtr<FDebugObjectInfo>>& OutObjects);
 
 	void OnTrackListChanged();
@@ -230,12 +259,18 @@ private:
 	void OnPIEStopped(bool bSimulating);
 	void OnPIESingleStepped(bool bSimulating);
 
+	bool CanUsePlaybackControls() const;
+
 	FDelegateHandle PreBeginPIEHandle;
 	FDelegateHandle PausePIEHandle;
 	FDelegateHandle ResumePIEHandle;
 	FDelegateHandle SingleStepPIEHandle;
 	FDelegateHandle ShutdownPIEHandle;
 
+	FDelegateHandle RecordingStartedHandle;
+	FDelegateHandle RecordingStoppedHandle;
+	FDelegateHandle ClearRecordingHandle;
+
 	void SetCurrentScrubTime(double Time);
 
 	TBindableProperty<double> TraceTime;
@@ -254,7 +289,6 @@ private:
 	FOnTrackListChanged TrackListChangedDelegate;
 	FOnTrackCursor TrackCursorDelegate;
 
-	bool bTraceJustConnected = false;
 	bool bPIEStarted = false;
 	bool bPIESimulating = false;
 
@@ -280,15 +314,18 @@ private:
 
 	static void GetScrubTimeInformation(double InDebugTime, FScrubTimeInformation& InOutTimeInformation, uint16 InRecordingIndex, const TraceServices::IAnalysisSession* AnalysisSession);
 
+#if UE_WITH_TRACE_BASED_DEBUGGERS_ANALYSIS
+	UE::TraceBasedDebuggers::FTraceSessionDescriptor CurrentSessionDescriptor;
+#endif // UE_WITH_TRACE_BASED_DEBUGGERS_ANALYSIS
 	TArray<TSharedPtr<FDebugObjectInfo>> DebuggedObjects;
 	mutable TSharedPtr<FDebugObjectInfo> SelectedObject;
 
 	TArray<TSharedPtr<RewindDebugger::FRewindDebuggerTrack>> DebugTracks;
 	TSharedPtr<RewindDebugger::FRewindDebuggerTrack> SelectedTrack;
+	TSharedPtr<UE::RewindDebugger::FDebuggerRecordingControls> RecordingControls;
 
 	TArray<RewindDebugger::FObjectId> CandidateIds;
 
-	mutable class IUnrealInsightsModule* UnrealInsightsModule = nullptr;
 	FTSTicker::FDelegateHandle TickerHandle;
 
 	TOptional<FVector> RootObjectPosition;
diff --git a/Engine/Plugins/Animation/GameplayInsights/Source/RewindDebugger/Private/RewindDebuggerCommands.h b/Engine/Plugins/Animation/GameplayInsights/Source/RewindDebugger/Private/RewindDebuggerCommands.h
--- a/Engine/Plugins/Animation/GameplayInsights/Source/RewindDebugger/Private/RewindDebuggerCommands.h
+++ b/Engine/Plugins/Animation/GameplayInsights/Source/RewindDebugger/Private/RewindDebuggerCommands.h
@@ -17,13 +17,12 @@ public:
 		: TCommands<FRewindDebuggerCommands>(TEXT("RewindDebugger"), LOCTEXT("RewindDebugger", "Rewind Debugger"), NAME_None, "RewindDebuggerStyle")
 	{ }
 
-	// TCommands interface
-
+	//~ TCommands interface
 	virtual void RegisterCommands() override
 	{
 		UI_COMMAND(PauseOrPlay, "Pause or Play Recording", "Toggle recording playing", EUserInterfaceActionType::Button, FInputChord(EKeys::SpaceBar) );
-		UI_COMMAND(StartRecording, "Start Recording", "Start recording Animation Insights data", EUserInterfaceActionType::Button, FInputChord(EModifierKey::Control | EModifierKey::Shift, EKeys::R));
-		UI_COMMAND(StopRecording, "Stop Recording", "Stop recording Animation Insights data", EUserInterfaceActionType::Button, FInputChord(EModifierKey::Control | EModifierKey::Shift, EKeys::S));
+		UI_COMMAND(StartRecording, "Start Recording", "Start recording gameplay data", EUserInterfaceActionType::Button, FInputChord(EModifierKey::Control | EModifierKey::Shift, EKeys::R));
+		UI_COMMAND(StopRecording, "Stop Recording", "Stop recording gameplay data", EUserInterfaceActionType::Button, FInputChord(EModifierKey::Control | EModifierKey::Shift, EKeys::S));
 		UI_COMMAND(FirstFrame, "First Frame", "Jump to first recorded frame", EUserInterfaceActionType::Button, FInputChord(EKeys::Up));
 		UI_COMMAND(PreviousFrame, "Previous Frame", "Step one frame back", EUserInterfaceActionType::Button, FInputChord(EKeys::Left));
 		UI_COMMAND(ReversePlay, "Reverse Play", "Playback recorded data in reverse", EUserInterfaceActionType::Button, FInputChord(EModifierKey::Control | EModifierKey::Shift, EKeys::SpaceBar));
@@ -35,7 +34,6 @@ public:
 		UI_COMMAND(AutoRecord, "Auto Record", "Automatically start recording when PIE is started", EUserInterfaceActionType::ToggleButton, FInputChord());
 		UI_COMMAND(OpenTrace, "Open Recording", "Open a Trace file containing a rewind debugger recording", EUserInterfaceActionType::Button, FInputChord());
 		UI_COMMAND(AttachToSession, "Attach to Session", "Attach to a live session", EUserInterfaceActionType::Button, FInputChord());
-		UI_COMMAND(SaveTrace, "Save Recording", "Save a copy of the current recording", EUserInterfaceActionType::Button, FInputChord());
 		UI_COMMAND(ClearTrace, "Clear Recording", "Clear current recording data", EUserInterfaceActionType::Button, FInputChord());
 	}
 
@@ -54,7 +52,6 @@ public:
 	TSharedPtr<FUICommandInfo> AutoRecord;
 	TSharedPtr<FUICommandInfo> OpenTrace;
 	TSharedPtr<FUICommandInfo> AttachToSession;
-	TSharedPtr<FUICommandInfo> SaveTrace;
 	TSharedPtr<FUICommandInfo> ClearTrace;
 };
 
diff --git a/Engine/Plugins/Animation/GameplayInsights/Source/RewindDebugger/Private/RewindDebuggerModule.cpp b/Engine/Plugins/Animation/GameplayInsights/Source/RewindDebugger/Private/RewindDebuggerModule.cpp
--- a/Engine/Plugins/Animation/GameplayInsights/Source/RewindDebugger/Private/RewindDebuggerModule.cpp
+++ b/Engine/Plugins/Animation/GameplayInsights/Source/RewindDebugger/Private/RewindDebuggerModule.cpp
@@ -13,9 +13,12 @@
 #include "PropertyTraceMenu.h"
 #include "RewindDebugger.h"
 #include "RewindDebuggerCommands.h"
+#include "RewindDebuggerEngineEditorBridge.h"
+#include "RewindDebuggerRuntime/RewindDebuggerRuntime.h"
 #include "RewindDebuggerStyle.h"
 #include "SRewindDebugger.h"
 #include "ToolMenus.h"
+#include "TraceSessionsManager.h"
 #include "Widgets/Docking/SDockTab.h"
 #include "WorkspaceMenuStructure.h"
 #include "WorkspaceMenuStructureModule.h"
@@ -26,6 +29,15 @@ const FName FRewindDebuggerModule::MainTabName("RewindDebugger2");
 const FName FRewindDebuggerModule::DetailsTabName("RewindDebuggerDetails2");
 const FName FRewindDebuggerModule::MainMenuName("RewindDebugger.MainMenu");
 const FName FRewindDebuggerModule::TrackContextMenuName("RewindDebugger.TrackContextMenu");
+const FName FRewindDebuggerModule::MainToolBarName = FName("RewindDebugger.ToolBar");
+const FName FRewindDebuggerModule::MainStatusBarName = FName("RewindDebugger.StatusBar");
+
+#include "TraceServices/ModuleService.h"
+
+FRewindDebuggerModule& FRewindDebuggerModule::Get()
+{
+	return FModuleManager::Get().LoadModuleChecked<FRewindDebuggerModule>(TEXT("RewindDebugger"));
+}
 
 FName FRewindDebuggerModule::GetMainTabName() const
 {
@@ -37,6 +49,11 @@ FName FRewindDebuggerModule::GetDetailsTabName() const
 	return DetailsTabName;
 }
 
+FName FRewindDebuggerModule::GetMainToolbarName() const
+{
+	return MainToolBarName;
+}
+
 TSharedRef<SDockTab> FRewindDebuggerModule::SpawnRewindDebuggerDetailsTab(const FSpawnTabArgs& SpawnTabArgs)
 {
 	if (FRewindDebugger::Instance() == nullptr)
@@ -49,7 +66,7 @@ TSharedRef<SDockTab> FRewindDebuggerModule::SpawnRewindDebuggerDetailsTab(const
 	RewindDebugger->SetIsDetailsPanelOpen(true);
 	
 	TSharedRef<SDockTab> MajorTab = SNew(SDockTab)
-		.TabRole(ETabRole::PanelTab);
+		.TabRole(PanelTab);
 
 	MajorTab->SetOnTabClosed(SDockTab::FOnTabClosedCallback::CreateLambda( [](TSharedRef<SDockTab>) { FRewindDebugger::Instance()->SetIsDetailsPanelOpen(false); }));
 
@@ -66,7 +83,7 @@ TSharedRef<SDockTab> FRewindDebuggerModule::SpawnRewindDebuggerTab(const FSpawnT
 	}
 	
 	const TSharedRef<SDockTab> MajorTab = SNew(SDockTab)
-		.TabRole(ETabRole::PanelTab)
+		.TabRole(PanelTab)
 		.OnTabClosed_Lambda([this](TSharedRef<SDockTab>)
 		{
 			// clear reference to widget so it will be destroyed
@@ -140,54 +157,35 @@ TSharedRef<SDockTab> FRewindDebuggerModule::SpawnRewindDebuggerTab(const FSpawnT
 							FIsActionChecked(),
 							FIsActionButtonVisible());
 
-	CommandList->MapAction(Commands.StartRecording, 
-							FExecuteAction::CreateRaw(DebuggerInstance, &FRewindDebugger::StartRecording), 
-							FCanExecuteAction::CreateRaw(DebuggerInstance, &FRewindDebugger::CanStartRecording),
-							FIsActionChecked(),
-							FIsActionButtonVisible::CreateLambda([]() { return !FRewindDebugger::Instance()->IsRecording();}));
-
-	CommandList->MapAction(Commands.StopRecording,
-							 FExecuteAction::CreateRaw(DebuggerInstance, &FRewindDebugger::StopRecording),
-							 FCanExecuteAction::CreateRaw(DebuggerInstance, &FRewindDebugger::CanStopRecording),
-							 FIsActionChecked(),
-							 FIsActionButtonVisible::CreateRaw(DebuggerInstance, &FRewindDebugger::CanStopRecording));
-
 	CommandList->MapAction(Commands.AutoEject,
-							 FExecuteAction::CreateLambda([DebuggerInstance]() { DebuggerInstance->SetShouldAutoEject(!DebuggerInstance->ShouldAutoEject()); }),
-							 FCanExecuteAction(),
-							 FCanExecuteAction::CreateRaw(DebuggerInstance, &FRewindDebugger::ShouldAutoEject),
-							 FIsActionButtonVisible());
-							 
+							FExecuteAction::CreateLambda([DebuggerInstance]() { DebuggerInstance->SetShouldAutoEject(!DebuggerInstance->ShouldAutoEject()); }),
+							FCanExecuteAction(),
+							FCanExecuteAction::CreateRaw(DebuggerInstance, &FRewindDebugger::ShouldAutoEject),
+							FIsActionButtonVisible());
+
 	CommandList->MapAction(Commands.AutoRecord,
-							 FExecuteAction::CreateLambda([DebuggerInstance]() { DebuggerInstance->SetShouldAutoRecordOnPIE(!DebuggerInstance->ShouldAutoRecordOnPIE()); }),
-							 FCanExecuteAction(),
-							 FCanExecuteAction::CreateRaw(DebuggerInstance, &FRewindDebugger::ShouldAutoRecordOnPIE),
-							 FIsActionButtonVisible());
-	
+							FExecuteAction::CreateLambda([DebuggerInstance]() { DebuggerInstance->SetShouldAutoRecordOnPIE(!DebuggerInstance->ShouldAutoRecordOnPIE()); }),
+							FCanExecuteAction(),
+							FCanExecuteAction::CreateRaw(DebuggerInstance, &FRewindDebugger::ShouldAutoRecordOnPIE),
+							FIsActionButtonVisible());
+
 	CommandList->MapAction(Commands.OpenTrace,
-							 FExecuteAction::CreateRaw(DebuggerInstance, &FRewindDebugger::OpenTrace),
-							 FCanExecuteAction::CreateRaw(DebuggerInstance, &FRewindDebugger::CanOpenTrace),
-							 FIsActionChecked(),
-							 FIsActionButtonVisible());
-	
+							FExecuteAction::CreateRaw(DebuggerInstance, &FRewindDebugger::OpenTrace),
+							FCanExecuteAction::CreateRaw(DebuggerInstance, &FRewindDebugger::CanOpenTrace),
+							FIsActionChecked(),
+							FIsActionButtonVisible());
+
 	CommandList->MapAction(Commands.AttachToSession,
-    						 FExecuteAction::CreateRaw(DebuggerInstance, &FRewindDebugger::AttachToSession),
-    						 FCanExecuteAction::CreateRaw(DebuggerInstance, &FRewindDebugger::CanOpenTrace),
-    						 FIsActionChecked(),
-    						 FIsActionButtonVisible());
-		
-	CommandList->MapAction(Commands.SaveTrace,
-							 FExecuteAction::CreateRaw(DebuggerInstance, &FRewindDebugger::SaveTrace),
-							 FCanExecuteAction::CreateRaw(DebuggerInstance, &FRewindDebugger::CanSaveTrace),
-							 FIsActionChecked(),
-							 FIsActionButtonVisible());
-	
-	CommandList->MapAction(Commands.ClearTrace,
-							 FExecuteAction::CreateRaw(DebuggerInstance, &FRewindDebugger::ClearTrace),
-							 FCanExecuteAction::CreateRaw(DebuggerInstance, &FRewindDebugger::CanClearTrace),
-							 FIsActionChecked(),
-							 FIsActionButtonVisible());
+							FExecuteAction::CreateRaw(DebuggerInstance, &FRewindDebugger::AttachToSession),
+							FCanExecuteAction::CreateRaw(DebuggerInstance, &FRewindDebugger::CanOpenTrace),
+							FIsActionChecked(),
+							FIsActionButtonVisible());
 
+	CommandList->MapAction(Commands.ClearTrace,
+							FExecuteAction::CreateRaw(DebuggerInstance, &FRewindDebugger::ClearTraceAndActiveSession),
+							FCanExecuteAction::CreateRaw(DebuggerInstance, &FRewindDebugger::CanClearTrace),
+							FIsActionChecked(),
+							FIsActionButtonVisible());
 
 	// Register PIE Rewind Debugger Commands
 	if (GEditor != nullptr)
@@ -233,8 +231,20 @@ static FBlueprintDoubleClickHandler BlueprintDoubleClickHandler;
 
 void FRewindDebuggerModule::StartupModule()
 {
-	UToolMenus::Get()->RegisterMenu(FRewindDebuggerModule::MainMenuName);
-	UToolMenus::Get()->RegisterMenu(FRewindDebuggerModule::TrackContextMenuName);
+	// @todo: this should only require UE_WITH_TRACE_BASED_DEBUGGERS_ANALYSIS,
+	// but we currently need the runtime instance to create the trace sessions manager
+	// in order to generate the file name. This dependency needs to be removed if we want
+	// some programs to only uses the analysis functionalities.
+#if UE_WITH_TRACE_BASED_DEBUGGERS_ANALYSIS && WITH_TRACE_BASED_DEBUGGERS
+	TraceSessionsManager = MakeShared<UE::TraceBasedDebuggers::FTraceSessionsManager>(
+		*RewindDebugger::FRewindDebuggerRuntime::Instance()
+		, UE::RewindDebugger::FRewindDebuggerEngineEditorBridge::Get()
+		, /*TraceModule*/nullptr
+		, /*SaveDirectorySubPathInUserDir*/TEXT("RewindDebugger"));
+#endif // UE_WITH_TRACE_BASED_DEBUGGERS_ANALYSIS && WITH_TRACE_BASED_DEBUGGERS
+
+	UToolMenus::Get()->RegisterMenu(MainMenuName);
+	UToolMenus::Get()->RegisterMenu(TrackContextMenuName);
 
 	FRewindDebuggerStyle::Initialize();
 	FRewindDebuggerCommands::Register();
diff --git a/Engine/Plugins/Animation/GameplayInsights/Source/RewindDebugger/Private/RewindDebuggerModule.h b/Engine/Plugins/Animation/GameplayInsights/Source/RewindDebugger/Private/RewindDebuggerModule.h
--- a/Engine/Plugins/Animation/GameplayInsights/Source/RewindDebugger/Private/RewindDebuggerModule.h
+++ b/Engine/Plugins/Animation/GameplayInsights/Source/RewindDebugger/Private/RewindDebuggerModule.h
@@ -6,6 +6,11 @@
 #include "RewindDebuggerCamera.h"
 #include "RewindDebuggerAnimation.h"
 
+namespace UE::TraceBasedDebuggers
+{
+class FTraceSessionsManager;
+}
+
 class SDockTab;
 class SRewindDebugger;
 class SRewindDebuggerDetails;
@@ -13,20 +18,36 @@ class SRewindDebuggerDetails;
 class FRewindDebuggerModule : public IRewindDebuggerModule
 {
 public:
-	// IModuleInterface interface
+	//~ IModuleInterface interface
 	virtual void StartupModule() override;
 	virtual void ShutdownModule() override;
 
+	static FRewindDebuggerModule& Get();
 	virtual FName GetMainTabName() const override;
+	virtual FName GetMainToolbarName() const override;
 	virtual FName GetDetailsTabName() const override;
 	virtual TSharedRef<SDockTab> SpawnRewindDebuggerTab(const FSpawnTabArgs& SpawnTabArgs) override;
 	virtual TSharedRef<SDockTab> SpawnRewindDebuggerDetailsTab(const FSpawnTabArgs& SpawnTabArgs) override;
 
+	static const FName MainMenuName;
+	static const FName MainToolBarName;
+	static const FName MainStatusBarName;
 	static const FName MainTabName;
 	static const FName DetailsTabName;
-	static const FName MainMenuName;
 	static const FName TrackContextMenuName;
 
+#if UE_WITH_TRACE_BASED_DEBUGGERS_ANALYSIS
+public:
+	/** Returns the trace sessions manager instance used by RewindDebugger */
+	TSharedPtr<UE::TraceBasedDebuggers::FTraceSessionsManager>& GetTraceSessionsManager()
+	{
+		return TraceSessionsManager;
+	}
+
+private:
+	TSharedPtr<UE::TraceBasedDebuggers::FTraceSessionsManager> TraceSessionsManager;
+#endif // UE_WITH_TRACE_BASED_DEBUGGERS_ANALYSIS
+
 private:
 	TSharedPtr<SRewindDebugger> RewindDebuggerWidget;
 	TSharedPtr<SRewindDebuggerDetails> RewindDebuggerDetailsWidget;
diff --git a/Engine/Plugins/Animation/GameplayInsights/Source/RewindDebugger/Private/RewindDebuggerRecordingControls.cpp b/Engine/Plugins/Animation/GameplayInsights/Source/RewindDebugger/Private/RewindDebuggerRecordingControls.cpp
new file mode 100644
--- /dev/null
+++ b/Engine/Plugins/Animation/GameplayInsights/Source/RewindDebugger/Private/RewindDebuggerRecordingControls.cpp
@@ -0,0 +1,120 @@
+// Copyright Epic Games, Inc. All Rights Reserved.
+
+#include "RewindDebuggerRecordingControls.h"
+
+#include "RewindDebugger.h"
+#include "RewindDebuggerEngineEditorBridge.h"
+#include "RewindDebuggerRemoteSessionsHandler.h"
+#include "RewindDebuggerStyle.h"
+#include "SessionInfo.h"
+
+namespace UE::RewindDebugger
+{
+
+FDebuggerRecordingControls::FDebuggerRecordingControls(const FName& InStatusBarId
+	, const TSharedRef<TraceBasedDebuggers::FRemoteSessionsManager>& InRemoteSessionManager
+	, const FGuid& InDebuggerId
+	, const FLogCategoryBase& InLogCategory)
+	: FRecordingControls(InStatusBarId, InRemoteSessionManager, InDebuggerId, InLogCategory)
+{
+	StartRecordingBrush = FRewindDebuggerStyle::Get().GetBrush("RewindDebugger.StartRecording");
+	StopRecordingBrush = FRewindDebuggerStyle::Get().GetBrush("RewindDebugger.StopRecording");
+	RecordingSessionIcon = FSlateIcon("RewindDebuggerStyle", FName("RewindDebugger.StartRecording"), FName("RewindDebugger.StartRecording"));
+
+#if WITH_TRACE_BASED_DEBUGGERS
+	TraceRelayTransportInstanceWeakPtr = FRewindDebuggerEngineEditorBridge::Get().GetTraceRelayTransportInstance().ToWeakPtr();
+#endif
+}
+
+bool FDebuggerRecordingControls::CanStartRecording() const
+{
+	return !IsRecording() && IsRecordingToggleButtonEnabled();
+}
+
+void FDebuggerRecordingControls::StartRecording()
+{
+	if (!IsRecording())
+	{
+		ToggleRecordingState();
+	}
+}
+
+bool FDebuggerRecordingControls::CanStopRecording() const
+{
+	return IsRecording() && IsRecordingToggleButtonEnabled();
+}
+
+void FDebuggerRecordingControls::StopRecording()
+{
+	if (IsRecording())
+	{
+		ToggleRecordingState();
+	}
+}
+
+bool FDebuggerRecordingControls::ConnectToLiveSession_RelayInternal(const FGuid& InSessionId)
+{
+#if UE_WITH_TRACE_BASED_DEBUGGERS_ANALYSIS
+	return FRewindDebugger::Instance()->ConnectToLiveSession_Relay(InSessionId);
+#else
+	return false;
+#endif
+}
+
+bool FDebuggerRecordingControls::ConnectToLiveSession_DirectInternal(const FGuid& InSessionId)
+{
+#if UE_WITH_TRACE_BASED_DEBUGGERS_ANALYSIS
+	return FRewindDebugger::Instance()->ConnectToLiveSession_Direct(InSessionId);
+#else
+	return false;
+#endif
+}
+
+bool FDebuggerRecordingControls::IsRecordingToggleButtonEnabledInternal(const TSharedRef<TraceBasedDebuggers::FSessionInfo>& InSessionInfo) const
+{
+	return (InSessionInfo->BuildTargetType == EBuildTargetType::Editor && FRewindDebugger::Instance()->IsPIESimulating())
+		|| InSessionInfo->BuildTargetType != EBuildTargetType::Editor;
+}
+
+TNotNull<TUniquePtr<TraceBasedDebuggers::FStartRecordingCommandMessage>> FDebuggerRecordingControls::BuildStartRecordingParamsInternal()
+{
+	return MakeUnique<FStartRecordingCommand>();
+}
+
+void FDebuggerRecordingControls::SendStartRecordingCommandInternal(const FMessageAddress& InAddress, TUniquePtr<TraceBasedDebuggers::FStartRecordingCommandMessage> InStartRecordingCommand)
+{
+	if (const TSharedPtr<TraceBasedDebuggers::FRemoteSessionsManager> SessionsManager = RemoteSessionManagerWeakPtr.Pin())
+	{
+		const FStartRecordingCommand* StartRecordingCommand = static_cast<FStartRecordingCommand*>(InStartRecordingCommand.Get());
+		SessionsManager->SendCommand(InAddress, *StartRecordingCommand);
+	}
+}
+
+void FDebuggerRecordingControls::SendStopRecordingCommandInternal(const FMessageAddress& InAddress)
+{
+	if (const TSharedPtr<TraceBasedDebuggers::FRemoteSessionsManager> SessionsManager = RemoteSessionManagerWeakPtr.Pin())
+	{
+		SessionsManager->SendCommand(InAddress, FStopRecordingCommand{});
+	}
+}
+
+void FDebuggerRecordingControls::OnToggleRecordingStateInternal(const TSharedRef<TraceBasedDebuggers::FSessionInfo>& InSessionInfo)
+{
+	ToggleSingleSessionRecordingState(InSessionInfo);
+}
+
+void FDebuggerRecordingControls::CloseSessionByRemoteSessionIDInternal(const FGuid& InSessionId)
+{
+#if UE_WITH_TRACE_BASED_DEBUGGERS_ANALYSIS
+	FRewindDebugger::Instance()->CloseSessionByRemoteSessionID(InSessionId);
+#endif
+}
+
+const TraceBasedDebuggers::FTraceSessionDescriptor& FDebuggerRecordingControls::GetCurrentSessionDescriptorInternal()
+{
+#if UE_WITH_TRACE_BASED_DEBUGGERS_ANALYSIS
+	return FRewindDebugger::Instance()->GetCurrentSessionDescriptor();
+#endif
+}
+
+} // UE::RewindDebugger
\ No newline at end of file
diff --git a/Engine/Plugins/Animation/GameplayInsights/Source/RewindDebugger/Private/RewindDebuggerRecordingControls.h b/Engine/Plugins/Animation/GameplayInsights/Source/RewindDebugger/Private/RewindDebuggerRecordingControls.h
new file mode 100644
--- /dev/null
+++ b/Engine/Plugins/Animation/GameplayInsights/Source/RewindDebugger/Private/RewindDebuggerRecordingControls.h
@@ -0,0 +1,38 @@
+// Copyright Epic Games, Inc. All Rights Reserved.
+
+#pragma once
+
+#include "Widgets/RecordingControls.h"
+
+namespace UE::RewindDebugger
+{
+
+struct FDebuggerRecordingControls : TraceBasedDebuggers::FRecordingControls
+{
+	FDebuggerRecordingControls(const FName& InStatusBarId
+		, const TSharedRef<TraceBasedDebuggers::FRemoteSessionsManager>& InRemoteSessionManager
+		, const FGuid& InDebuggerId
+		, const FLogCategoryBase& InLogCategory);
+
+	using FRecordingControls::IsRecording;
+
+	bool CanStartRecording() const;
+	void StartRecording();
+
+	bool CanStopRecording() const;
+	void StopRecording();
+protected:
+	virtual bool ConnectToLiveSession_RelayInternal(const FGuid& InSessionId) override;
+	virtual bool ConnectToLiveSession_DirectInternal(const FGuid& InSessionId) override;
+	virtual bool IsRecordingToggleButtonEnabledInternal(const TSharedRef<TraceBasedDebuggers::FSessionInfo>& InSessionInfo) const override;
+
+	virtual TNotNull<TUniquePtr<TraceBasedDebuggers::FStartRecordingCommandMessage>> BuildStartRecordingParamsInternal() override;
+	virtual void SendStartRecordingCommandInternal(const FMessageAddress& InAddress, TUniquePtr<TraceBasedDebuggers::FStartRecordingCommandMessage> InStartRecordingCommand) override;
+	virtual void SendStopRecordingCommandInternal(const FMessageAddress& InAddress) override;
+
+	virtual void OnToggleRecordingStateInternal(const TSharedRef<TraceBasedDebuggers::FSessionInfo>& InSessionInfo) override;
+	virtual void CloseSessionByRemoteSessionIDInternal(const FGuid& InSessionId) override;
+	virtual const TraceBasedDebuggers::FTraceSessionDescriptor& GetCurrentSessionDescriptorInternal() override;
+};
+
+}
diff --git a/Engine/Plugins/Animation/GameplayInsights/Source/RewindDebugger/Private/RewindDebuggerStyle.cpp b/Engine/Plugins/Animation/GameplayInsights/Source/RewindDebugger/Private/RewindDebuggerStyle.cpp
--- a/Engine/Plugins/Animation/GameplayInsights/Source/RewindDebugger/Private/RewindDebuggerStyle.cpp
+++ b/Engine/Plugins/Animation/GameplayInsights/Source/RewindDebugger/Private/RewindDebuggerStyle.cpp
@@ -1,4 +1,4 @@
-// Copyright Epic Games, Inc. All Rights Reserved.
+﻿// Copyright Epic Games, Inc. All Rights Reserved.
 
 #include "RewindDebuggerStyle.h"
 #include "Styling/AppStyle.h"
@@ -21,6 +21,8 @@ FRewindDebuggerStyle::FRewindDebuggerStyle() :
 	SetCoreContentRoot(FPaths::EngineContentDir() / TEXT(""));
 
 	// Toolbar
+	Set("RewindDebugger.StartRecording", new CORE_IMAGE_BRUSH_SVG("Editor/Slate/Sequencer/PlaybackControls/PlayControlsRecord", Icon24x24));
+	Set("RewindDebugger.StopRecording", new CORE_IMAGE_BRUSH_SVG("Editor/Slate/Sequencer/PlaybackControls/PlayControlsStop", Icon24x24));
 	Set("RewindDebugger.StartRecording.small", new CORE_IMAGE_BRUSH("Editor/Slate/Sequencer/Transport_Bar/Record_24x", Icon24x24));
 	Set("RewindDebugger.StopRecording.small", new CORE_IMAGE_BRUSH("Editor/Slate/Sequencer/Transport_Bar/Recording_24x", Icon24x24));
 	Set("RewindDebugger.FirstFrame.small", new CORE_IMAGE_BRUSH("Editor/Slate/Sequencer/Transport_Bar/Go_To_Front_24x", Icon24x24));
diff --git a/Engine/Plugins/Animation/GameplayInsights/Source/RewindDebugger/Private/SRewindDebugger.cpp b/Engine/Plugins/Animation/GameplayInsights/Source/RewindDebugger/Private/SRewindDebugger.cpp
--- a/Engine/Plugins/Animation/GameplayInsights/Source/RewindDebugger/Private/SRewindDebugger.cpp
+++ b/Engine/Plugins/Animation/GameplayInsights/Source/RewindDebugger/Private/SRewindDebugger.cpp
@@ -1,4 +1,4 @@
-// Copyright Epic Games, Inc. All Rights Reserved.
+﻿// Copyright Epic Games, Inc. All Rights Reserved.
 
 #include "SRewindDebugger.h"
 #include "ActorPickerMode.h"
@@ -14,15 +14,18 @@
 #include "Modules/ModuleManager.h"
 #include "ObjectTrace.h"
 #include "RewindDebuggerCommands.h"
+#include "RewindDebuggerEngineEditorBridge.h"
 #include "RewindDebuggerModule.h"
+#include "RewindDebuggerRuntime/RewindDebuggerRuntime.h"
 #include "RewindDebuggerSettings.h"
 #include "RewindDebuggerStyle.h"
-#include "SceneOutlinerModule.h"
 #include "Selection.h"
 #include "SSimpleTimeSlider.h"
+#include "StatusBarSubsystem.h"
 #include "Styling/SlateIconFinder.h"
 #include "ToolMenu.h"
 #include "ToolMenus.h"
+#include "Widgets/SConnectionStatusBar.h"
 #include "Widgets/Images/SImage.h"
 #include "Widgets/Input/SButton.h"
 #include "Widgets/Input/SComboButton.h"
@@ -39,6 +42,18 @@ SRewindDebugger::SRewindDebugger()
 {
 }
 
+SRewindDebugger::~SRewindDebugger()
+{
+#if WITH_TRACE_BASED_DEBUGGERS
+	// Sessions discovery no longer needed after closing the main debugger window
+	using namespace UE::RewindDebugger;
+	if (TSharedPtr<UE::TraceBasedDebuggers::FRemoteSessionsManager> RemoteSessionManager = FRewindDebuggerEngineEditorBridge::Get().GetSessionsManager())
+	{
+		RemoteSessionManager->StopSessionDiscovery(DebuggerGuid);
+	}
+#endif // WITH_TRACE_BASED_DEBUGGERS
+}
+
 void SRewindDebugger::TrackCursor(const bool bReverse)
 {
 	const float ScrubTime = ScrubTimeAttribute.Get();
@@ -106,7 +121,6 @@ bool SRewindDebugger::ShouldHideTrackType(const FName& TrackType) const
 	return Settings.HiddenTrackTypes.Contains(TrackType);
 }
 
-
 void SRewindDebugger::ToggleDisplayEmptyTracks()
 {
 	Settings.bShowEmptyObjectTracks = !Settings.bShowEmptyObjectTracks;
@@ -317,10 +331,7 @@ void SRewindDebugger::Construct(const FArguments& InArgs, TSharedRef<FUICommandL
 		.OnScrubPositionChanged_Lambda(
 			[this](double NewScrubTime, bool bIsScrubbing)
 			{
-				if (bIsScrubbing)
-				{
-					OnScrubPositionChanged.ExecuteIfBound(NewScrubTime, bIsScrubbing);
-				}
+				OnScrubPositionChanged.ExecuteIfBound(NewScrubTime, bIsScrubbing);
 			})
 		.OnSelectionChanged_Lambda(
 			[this](TSharedPtr<RewindDebugger::FRewindDebuggerTrack> SelectedItem, ESelectInfo::Type SelectInfo)
@@ -334,8 +345,45 @@ void SRewindDebugger::Construct(const FArguments& InArgs, TSharedRef<FUICommandL
 				}
 			});
 
+	TSharedPtr<SWidget> StatusBarWidget = SNullWidget::NullWidget;
+
+#if WITH_TRACE_BASED_DEBUGGERS
+	using namespace UE::RewindDebugger;
+	using namespace UE::TraceBasedDebuggers;
+	if (const TSharedPtr<FRemoteSessionsManager> RemoteSessionManager = FRewindDebuggerEngineEditorBridge::Get().GetSessionsManager())
+	{
+		// Start remote sessions discovery
+		RemoteSessionManager->StartSessionDiscovery(DebuggerGuid);
+
+		// Add status bar for connection statistics updates
+		if (UStatusBarSubsystem* StatusBarSubsystem = GEditor ? GEditor->GetEditorSubsystem<UStatusBarSubsystem>() : nullptr)
+		{
+			using namespace UE::Editor::Toolbars;
+			constexpr ECreateStatusBarOptions StatusBarOptions =
+				ECreateStatusBarOptions::HideContentBrowser
+				| ECreateStatusBarOptions::HideOutputLog
+				| ECreateStatusBarOptions::HideSourceControl;
+
+			StatusBarWidget = StatusBarSubsystem->MakeStatusBarWidget(FRewindDebuggerModule::MainStatusBarName, ConstructUnderMajorTab, StatusBarOptions);
+
+			const FName StatusBarToolbarID = FName(FRewindDebuggerModule::MainStatusBarName.ToString() + TEXT(".Toolbar"));
+			if (UToolMenu* StatusToolBarMenu = UToolMenus::Get()->ExtendMenu(StatusBarToolbarID))
+			{
+				FToolMenuSection& ConnectionStatusSection = StatusToolBarMenu->AddSection("ConnectionStatus");
+
+				constexpr bool bNoIndent = true;
+				constexpr bool bSearchable = false;
+				ConnectionStatusSection.AddEntry(FToolMenuEntry::InitWidget("ConnectionStatusBar"
+					, SNew(SConnectionStatusBar, UE::RewindDebugger::DebuggerGuid, RemoteSessionManager.ToSharedRef())
+					, FText::GetEmpty()
+					, bNoIndent
+					, bSearchable));
+			}
+		}
+	}
+#endif // WITH_TRACE_BASED_DEBUGGERS
 
-	const TSharedRef<SWidget> ToolBar = UToolMenus::Get()->GenerateWidget("RewindDebugger.ToolBar", ToolMenuContext);
+	const TSharedRef<SWidget> ToolBar = UToolMenus::Get()->GenerateWidget(FRewindDebuggerModule::MainToolBarName, ToolMenuContext);
 
 	ChildSlot
 		[
@@ -486,10 +534,15 @@ void SRewindDebugger::Construct(const FArguments& InArgs, TSharedRef<FUICommandL
 										[
 											ScrollBar.ToSharedRef()
 										]
-
 								]
 						]
 				]
+				+ SVerticalBox::Slot()
+				.Padding(0.0f, 2.0f, 0.0f, 0.0f)
+				.AutoHeight()
+				[
+					StatusBarWidget.ToSharedRef()
+				]
 		];
 }
 
diff --git a/Engine/Plugins/Animation/GameplayInsights/Source/RewindDebugger/Private/SRewindDebugger.h b/Engine/Plugins/Animation/GameplayInsights/Source/RewindDebugger/Private/SRewindDebugger.h
--- a/Engine/Plugins/Animation/GameplayInsights/Source/RewindDebugger/Private/SRewindDebugger.h
+++ b/Engine/Plugins/Animation/GameplayInsights/Source/RewindDebugger/Private/SRewindDebugger.h
@@ -55,6 +55,7 @@ public:
 	void Construct(const FArguments& InArgs, TSharedRef<FUICommandList> CommandList, const TSharedRef<SDockTab>& ConstructUnderMajorTab, const TSharedPtr<SWindow>& ConstructUnderWindow);
 
 	SRewindDebugger();
+	virtual ~SRewindDebugger();
 	void TrackCursor(bool bReverse);
 	void RefreshTracks();
 
diff --git a/Engine/Plugins/Animation/GameplayInsights/Source/RewindDebugger/Private/SRewindDebuggerTimelines.cpp b/Engine/Plugins/Animation/GameplayInsights/Source/RewindDebugger/Private/SRewindDebuggerTimelines.cpp
--- a/Engine/Plugins/Animation/GameplayInsights/Source/RewindDebugger/Private/SRewindDebuggerTimelines.cpp
+++ b/Engine/Plugins/Animation/GameplayInsights/Source/RewindDebugger/Private/SRewindDebuggerTimelines.cpp
@@ -1,7 +1,6 @@
 // Copyright Epic Games, Inc. All Rights Reserved.
 
 #include "SRewindDebuggerTimelines.h"
-#include "ISequencerWidgetsModule.h"
 #include "ObjectTrace.h"
 #include "RewindDebugger.h"
 #include "RewindDebuggerStyle.h"
diff --git a/Engine/Plugins/Animation/GameplayInsights/Source/RewindDebugger/Private/SRewindDebuggerTrackTree.cpp b/Engine/Plugins/Animation/GameplayInsights/Source/RewindDebugger/Private/SRewindDebuggerTrackTree.cpp
--- a/Engine/Plugins/Animation/GameplayInsights/Source/RewindDebugger/Private/SRewindDebuggerTrackTree.cpp
+++ b/Engine/Plugins/Animation/GameplayInsights/Source/RewindDebugger/Private/SRewindDebuggerTrackTree.cpp
@@ -1,7 +1,6 @@
 // Copyright Epic Games, Inc. All Rights Reserved.
 
 #include "SRewindDebuggerTrackTree.h"
-#include "ISequencerWidgetsModule.h"
 #include "ObjectTrace.h"
 #include "RewindDebugger.h"
 #include "RewindDebuggerStyle.h"
diff --git a/Engine/Plugins/Animation/GameplayInsights/Source/RewindDebugger/RewindDebugger.Build.cs b/Engine/Plugins/Animation/GameplayInsights/Source/RewindDebugger/RewindDebugger.Build.cs
--- a/Engine/Plugins/Animation/GameplayInsights/Source/RewindDebugger/RewindDebugger.Build.cs
+++ b/Engine/Plugins/Animation/GameplayInsights/Source/RewindDebugger/RewindDebugger.Build.cs
@@ -10,36 +10,28 @@ namespace UnrealBuildTool.Rules
 			new string[]
 			{
 				"ActorPickerMode",
-				"AnimationBlueprintEditor",
-				"ApplicationCore",
-				"AssetRegistry",
 				"Core",
 				"CoreUObject",
-				"EditorFramework",
-				"EditorWidgets",
+				"DeveloperSettings",
 				"Engine",
 				"GameplayInsights",
-				"GameplayInsightsEditor",
 				"InputCore",
-				"Kismet",
-				"LevelEditor",
-				"Persona",
-				"SceneOutliner",
-				"SequencerWidgets",
-				"SlateCore",
+				"LevelSequence", // required for hard coded level sequence special case in ObjectTrack
+				"RewindDebuggerInterface",
+				"RewindDebuggerRuntime",
 				"Slate",
+				"SlateCore",
+				"Sockets",
+				"StatusBar",
 				"ToolMenus",
-				"TraceLog",
+				"ToolWidgets",
 				"TraceAnalysis",
-				"TraceServices",
+				"TraceBasedDebuggers",
+				"TraceBasedDebuggersAnalysis",
 				"TraceInsights",
+				"TraceLog",
+				"TraceServices",
 				"UnrealEd",
-				"RewindDebuggerInterface",
-				"RewindDebuggerRuntime",
-				"ToolWidgets",
-				"MainFrame",
-				"DeveloperSettings", 
-				"LevelSequence", "LevelSequence",
 			});
 		}
 	}
diff --git a/Engine/Plugins/Animation/GameplayInsights/Source/RewindDebuggerRuntime/Internal/RewindDebuggerRemoteSessionsHandler.h b/Engine/Plugins/Animation/GameplayInsights/Source/RewindDebuggerRuntime/Internal/RewindDebuggerRemoteSessionsHandler.h
new file mode 100644
--- /dev/null
+++ b/Engine/Plugins/Animation/GameplayInsights/Source/RewindDebuggerRuntime/Internal/RewindDebuggerRemoteSessionsHandler.h
@@ -0,0 +1,39 @@
+// Copyright Epic Games, Inc. All Rights Reserved.
+
+#pragma once
+
+#include "RemoteSessionsManager.h"
+#include "RewindDebuggerRemoteSessionsHandler.generated.h"
+
+#define UE_API REWINDDEBUGGERRUNTIME_API
+
+namespace UE::RewindDebugger
+{
+USTRUCT()
+struct FStartRecordingCommand : public UE::TraceBasedDebuggers::FStartRecordingCommandMessage
+{
+GENERATED_BODY()
+};
+
+
+USTRUCT()
+struct FStopRecordingCommand : public UE::TraceBasedDebuggers::FStopRecordingCommandMessage
+{
+GENERATED_BODY()
+};
+
+#if WITH_TRACE_BASED_DEBUGGERS
+/**
+ * Remote sessions handler to handle RewindDebugger specific message types.
+ */
+struct FSessionsHandler : TraceBasedDebuggers::IRemoteSessionsHandler
+{
+	UE_API virtual void OnCreatingMessageEndpoint(const TSharedPtr<TraceBasedDebuggers::FRemoteSessionsManager>&, FMessageEndpointBuilder&) override;
+	UE_API virtual void OnSubscribingMessageTypes(const TSharedPtr<TraceBasedDebuggers::FRemoteSessionsManager>&, FMessageEndpoint&) override;
+	UE_API virtual void OnBuildingFullSessionInfoResponseMessage(TraceBasedDebuggers::FFullSessionInfoResponseMessage&) override;
+	UE_API virtual void OnHandlingFullSessionInfoResponseMessage(const TraceBasedDebuggers::FFullSessionInfoResponseMessage&, const TSharedPtr<TraceBasedDebuggers::FSessionInfo>&) override;
+};
+#endif // WITH_TRACE_BASED_DEBUGGERS
+} // UE::RewindDebugger
+
+#undef UE_API
\ No newline at end of file
diff --git a/Engine/Plugins/Animation/GameplayInsights/Source/RewindDebuggerRuntime/Private/RewindDebuggerEngineEditorBridge.cpp b/Engine/Plugins/Animation/GameplayInsights/Source/RewindDebuggerRuntime/Private/RewindDebuggerEngineEditorBridge.cpp
new file mode 100644
--- /dev/null
+++ b/Engine/Plugins/Animation/GameplayInsights/Source/RewindDebuggerRuntime/Private/RewindDebuggerEngineEditorBridge.cpp
@@ -0,0 +1,219 @@
+// Copyright Epic Games, Inc. All Rights Reserved.
+
+#include "RewindDebuggerEngineEditorBridge.h"
+
+#if WITH_TRACE_BASED_DEBUGGERS
+
+#include "RewindDebuggerRemoteSessionsHandler.h"
+#include "TraceBasedDebuggerRuntime.h"
+#include "TraceDataRelayTransport.h"
+#include "Engine/Engine.h"
+#include "Misc/App.h"
+#include "UObject/Package.h"
+
+namespace UE::RewindDebugger
+{
+using namespace ::RewindDebugger;
+
+bool FRewindDebuggerEngineEditorBridge::bIsInstantiated = false;
+
+FRewindDebuggerEngineEditorBridge::FRewindDebuggerEngineEditorBridge(): FEngineEditorBridge(LogRewindDebuggerRuntime)
+{
+	RegisterRemoteSessionsHandler(MakeShared<FSessionsHandler>());
+}
+
+FRewindDebuggerEngineEditorBridge& FRewindDebuggerEngineEditorBridge::Get()
+{
+	static FRewindDebuggerEngineEditorBridge EngineEditorBridge;
+
+	// Keep track of the instantiation of the bridge.
+	// Doing so allows external modules responsible for calling Teardown
+	// to avoid calling Get (and instantiate the bridge) if the bridge was not instantiated.
+	bIsInstantiated = true;
+	return EngineEditorBridge;
+}
+
+bool FRewindDebuggerEngineEditorBridge::IsInstantiated()
+{
+	return bIsInstantiated;
+}
+
+void FRewindDebuggerEngineEditorBridge::OnInitializeInternal()
+{
+	using namespace UE::TraceBasedDebuggers;
+
+#if WITH_EDITOR
+	PIEStartedHandle = FWorldDelegates::OnPIEStarted.AddRaw(this, &FRewindDebuggerEngineEditorBridge::HandlePIEStarted);
+#endif
+
+	if (FRewindDebuggerRuntime* Runtime = FRewindDebuggerRuntime::Instance())
+	{
+		//Runtime->ClearRecording.AddRaw(this, &FRewindDebugger::OnClearRecording);
+		RecordingStartedHandle = Runtime->RecordingStarted.AddRaw(this, &FRewindDebuggerEngineEditorBridge::HandleRecordingStarted);
+		RecordingStartFailedHandle = Runtime->RecordingStartFailed.AddRaw(this, &FRewindDebuggerEngineEditorBridge::HandleRecordingStartFailed);
+		RecordingStoppedHandle = Runtime->RecordingStopped.AddRaw(this, &FRewindDebuggerEngineEditorBridge::HandleRecordingStopped);
+
+		Runtime->OnTraceConnectionDetailsUpdated().AddRaw(this, &FRewindDebuggerEngineEditorBridge::HandleTraceConnectionDetailsUpdated);
+
+		// If we were already recording, we need to make sure we run the initialization step to set up the session broadcast ticker
+		// and the collision channel serialization
+		if (Runtime->IsRecording())
+		{
+			HandleRecordingStarted();
+		}
+	}
+}
+
+void FRewindDebuggerEngineEditorBridge::OnTearDownInternal()
+{
+#if WITH_EDITOR
+	FWorldDelegates::OnPIEStarted.Remove(PIEStartedHandle);
+#endif
+
+	FTSTicker::RemoveTicker(DeferredShowMessageOnScreenHandle);
+
+	if (FRewindDebuggerRuntime* Runtime = FRewindDebuggerRuntime::Instance())
+	{
+		Runtime->RecordingStarted.Remove(RecordingStartedHandle);
+		Runtime->RecordingStartFailed.Remove(RecordingStartFailedHandle);
+		Runtime->RecordingStopped.Remove(RecordingStoppedHandle);
+		Runtime->OnTraceConnectionDetailsUpdated().RemoveAll(this);
+
+		// Make sure of removing the message from the screen in case the recording didn't quite stop yet
+		if (Runtime->IsRecording())
+		{
+			HandleRecordingStopped();
+		}
+	}
+}
+
+bool FRewindDebuggerEngineEditorBridge::AddOnScreenRecordingMessage(float DummyDeltaTime)
+{
+	constexpr bool bContinueLooping = false;
+	if (!GEngine)
+	{
+		return bContinueLooping;
+	}
+
+	if (!IsInGameThread())
+	{
+		if (!DeferredShowMessageOnScreenHandle.IsValid())
+		{
+			DeferredShowMessageOnScreenHandle = FTSTicker::GetCoreTicker().AddTicker(FTickerDelegate::CreateRaw(this, &FRewindDebuggerEngineEditorBridge::AddOnScreenRecordingMessage));
+		}
+
+		return bContinueLooping;
+	}
+
+	DeferredShowMessageOnScreenHandle.Reset();
+
+	const FText RecordingStartedMessage = NSLOCTEXT("RewindDebugger", "OnScreenRecordingStartedMessage", "Rewind Debugger recording in progress...");
+
+	if (RecordingMessageKey == 0)
+	{
+		RecordingMessageKey = GetTypeHash(RecordingStartedMessage.ToString());
+	}
+
+	// Add a long duration value, we will remove the message manually when the recording stops
+	// @todo: consider using FSlateNotificationManager if properly supported on all platforms
+	constexpr float MessageDurationSeconds = 3600.0f;
+	GEngine->AddOnScreenDebugMessage(RecordingMessageKey, MessageDurationSeconds, FColor::Red, RecordingStartedMessage.ToString());
+
+	return bContinueLooping;
+}
+
+void FRewindDebuggerEngineEditorBridge::RemoveOnScreenRecordingMessage()
+{
+	ensure(IsInGameThread());
+
+	if (!GEngine)
+	{
+		return;
+	}
+
+	if (DeferredShowMessageOnScreenHandle.IsValid())
+	{
+		FTSTicker::RemoveTicker(DeferredShowMessageOnScreenHandle);
+		DeferredShowMessageOnScreenHandle.Reset();
+	}
+	else if (RecordingMessageKey != 0)
+	{
+		GEngine->RemoveOnScreenDebugMessage(RecordingMessageKey);
+	}
+}
+
+void FRewindDebuggerEngineEditorBridge::OnRecordingStartedInternal()
+{
+	AddOnScreenRecordingMessage();
+}
+
+void FRewindDebuggerEngineEditorBridge::OnRecordingStoppedInternal()
+{
+	RemoveOnScreenRecordingMessage();
+}
+
+void FRewindDebuggerEngineEditorBridge::BuildRecordingStatusInternal(TraceBasedDebuggers::FRecordingStatusMessage& OutStatusMessage) const
+{
+	if (const FRewindDebuggerRuntime* Runtime = FRewindDebuggerRuntime::Instance())
+	{
+		OutStatusMessage.DebuggerId = Runtime->GetDebuggerId();
+		OutStatusMessage.bIsRecording = Runtime->IsRecording();
+		OutStatusMessage.ElapsedTime = Runtime->GetAccumulatedRecordingTime();
+		OutStatusMessage.BufferedDataBytesSize = Runtime->GetBufferedDataBytesSize();
+	}
+}
+
+void FRewindDebuggerEngineEditorBridge::HandleTraceConnectionDetailsUpdated() const
+{
+	TraceBasedDebuggers::FTraceConnectionDetailsMessage ConnectionDetailsMessage;
+	ConnectionDetailsMessage.InstanceId = FApp::GetInstanceId();
+
+	if (const FRewindDebuggerRuntime* Runtime = FRewindDebuggerRuntime::Instance())
+	{
+		ConnectionDetailsMessage.TraceDetails = Runtime->GetCurrentTraceConnectionDetails();
+	}
+
+	GetSessionsManager()->PublishMessage(ConnectionDetailsMessage);
+}
+
+void FRewindDebuggerEngineEditorBridge::HandleRecordingStartFailed(const FText& InFailureReason) const
+{
+#if !WITH_EDITOR
+	if (GEngine)
+	{
+		// In non-editor builds we don't have an error pop-up, therefore we want to show the error message on screen
+		const FText ErrorMessage = FText::FormatOrdered(NSLOCTEXT("RewindDebugger", "StartRecordingFailedOnScreenMessage", "Failed to start debug trace recording. {0}"), InFailureReason);
+
+		constexpr float MessageDurationSeconds = 4.0f;
+		GEngine->AddOnScreenDebugMessage(RecordingMessageKey, MessageDurationSeconds, FColor::Red, ErrorMessage.ToString());
+	}
+#endif
+}
+
+void FRewindDebuggerEngineEditorBridge::HandlePIEStarted(UGameInstance* GameInstance)
+{
+	if (const FRewindDebuggerRuntime* Runtime = FRewindDebuggerRuntime::Instance())
+	{
+		if (Runtime->IsRecording())
+		{
+			AddOnScreenRecordingMessage();
+		}
+	}
+}
+
+} // UE::RewindDebugger
+
+#else
+
+namespace UE::RewindDebugger
+{
+
+FRewindDebuggerEngineEditorBridge& FRewindDebuggerEngineEditorBridge::Get()
+{
+	static FRewindDebuggerEngineEditorBridge EngineEditorBridge;
+	return EngineEditorBridge;
+}
+
+} // UE::RewindDebugger
+
+#endif // WITH_TRACE_BASED_DEBUGGERS
diff --git a/Engine/Plugins/Animation/GameplayInsights/Source/RewindDebuggerRuntime/Private/RewindDebuggerRemoteSessionsHandler.cpp b/Engine/Plugins/Animation/GameplayInsights/Source/RewindDebuggerRuntime/Private/RewindDebuggerRemoteSessionsHandler.cpp
new file mode 100644
--- /dev/null
+++ b/Engine/Plugins/Animation/GameplayInsights/Source/RewindDebuggerRuntime/Private/RewindDebuggerRemoteSessionsHandler.cpp
@@ -0,0 +1,88 @@
+// Copyright Epic Games, Inc. All Rights Reserved.
+
+#if WITH_TRACE_BASED_DEBUGGERS
+
+#include "RewindDebuggerRemoteSessionsHandler.h"
+#include "MessageEndpointBuilder.h"
+#include "RewindDebuggerRuntime/RewindDebuggerRuntime.h"
+#include "RewindDebuggerRuntimeModule.h"
+#include "SessionInfo.h"
+
+namespace UE::RewindDebugger
+{
+
+void FSessionsHandler::OnCreatingMessageEndpoint(
+	const TSharedPtr<TraceBasedDebuggers::FRemoteSessionsManager>& SessionsManager
+	, FMessageEndpointBuilder& InEndPointBuilder)
+{
+	InEndPointBuilder
+		.Handling<FStartRecordingCommand>([](const FStartRecordingCommand& InMessage, const TSharedRef<IMessageContext>& InContext)
+			{
+				UE_AUTORTFM_ONCOMMIT(InMessage, InContext)
+				{
+					if (::RewindDebugger::FRewindDebuggerRuntime* Runtime = ::RewindDebugger::FRewindDebuggerRuntime::Instance())
+					{
+						Runtime->StartRecording(InMessage);
+					}
+				};
+			})
+		.Handling<FStopRecordingCommand>([](const FStopRecordingCommand&, const TSharedRef<IMessageContext>&)
+			{
+				UE_AUTORTFM_ONCOMMIT()
+				{
+					if (::RewindDebugger::FRewindDebuggerRuntime* Runtime = ::RewindDebugger::FRewindDebuggerRuntime::Instance())
+					{
+						Runtime->StopRecording();
+					}
+				};
+			});
+}
+
+void FSessionsHandler::OnSubscribingMessageTypes(
+	const TSharedPtr<TraceBasedDebuggers::FRemoteSessionsManager>& InSessionsManager
+	, FMessageEndpoint& InMessageEndpoint)
+{
+	if (InSessionsManager)
+	{
+		InSessionsManager->RegisterExternalSupportedMessageType<FStartRecordingCommand>();
+		InSessionsManager->RegisterExternalSupportedMessageType<FStopRecordingCommand>();
+	}
+
+	InMessageEndpoint.Subscribe<FStartRecordingCommand>();
+	InMessageEndpoint.Subscribe<FStopRecordingCommand>();
+}
+
+void FSessionsHandler::OnBuildingFullSessionInfoResponseMessage(
+	TraceBasedDebuggers::FFullSessionInfoResponseMessage& InMessage)
+{
+	if (const ::RewindDebugger::FRewindDebuggerRuntime* Runtime = ::RewindDebugger::FRewindDebuggerRuntime::Instance())
+	{
+		InMessage.bIsRecording = Runtime->IsRecording();
+		InMessage.DebuggerId = DebuggerGuid;
+	}
+
+	// Session information currently doesn't require any RewindDebugger specific data, but if it's eventually
+	// needed then this is where it needs to be added to the message.
+	// e.g.,
+	//    FRewindDebuggerSessionInfoResponseData MessageData;
+	//    InMessage.DebuggerSpecificData = FInstancedStruct::Make(MoveTemp(MessageData));
+}
+
+void FSessionsHandler::OnHandlingFullSessionInfoResponseMessage(
+	const TraceBasedDebuggers::FFullSessionInfoResponseMessage& InMessage
+	, const TSharedPtr<TraceBasedDebuggers::FSessionInfo>& InSessionInfo)
+{
+	// Session information currently doesn't require any RewindDebugger specific data, but if it's eventually
+	// needed then this is where it needs to be extracted from the message and stored in the session info.
+	// e.g.,
+	//    if (const FRewindDebuggerSessionInfoResponseData* MessageData = InMessage.GetDebuggerData<FRewindDebuggerSessionInfoResponseData>())
+	//    {
+	//        FRewindDebuggerSessionData NewSessionData;
+	//        ... setup data ...
+	//        InSessionInfo->SetDebuggerData<FRewindDebuggerSessionData>(MoveTemp(NewSessionData));
+	//    }
+}
+
+} // UE::RewindDebugger
+
+#endif // WITH_TRACE_BASED_DEBUGGERS
\ No newline at end of file
diff --git a/Engine/Plugins/Animation/GameplayInsights/Source/RewindDebuggerRuntime/Private/RewindDebuggerRuntime.cpp b/Engine/Plugins/Animation/GameplayInsights/Source/RewindDebuggerRuntime/Private/RewindDebuggerRuntime.cpp
--- a/Engine/Plugins/Animation/GameplayInsights/Source/RewindDebuggerRuntime/Private/RewindDebuggerRuntime.cpp
+++ b/Engine/Plugins/Animation/GameplayInsights/Source/RewindDebuggerRuntime/Private/RewindDebuggerRuntime.cpp
@@ -1,16 +1,16 @@
 // Copyright Epic Games, Inc. All Rights Reserved.
 
+#if WITH_TRACE_BASED_DEBUGGERS
+
 #include "RewindDebuggerRuntime/RewindDebuggerRuntime.h"
 
-#include "ObjectTrace.h"
 #include "EngineUtils.h"
-#include "Engine/World.h"
+#include "Features/IModularFeatures.h"
 #include "GameFramework/Pawn.h"
-#include "Misc/FileHelper.h"
-#include "ProfilingDebugging/TraceAuxiliary.h"
-#include "UObject/UObjectIterator.h"
+#include "ObjectTrace.h"
+#include "RemoteSessionsManager.h"
 #include "RewindDebuggerRuntimeInterface/IRewindDebuggerRuntimeExtension.h"
-#include "Features/IModularFeatures.h"
+#include "UObject/UObjectIterator.h"
 
 DEFINE_LOG_CATEGORY(LogRewindDebuggerRuntime)
 
@@ -19,6 +19,14 @@ namespace RewindDebugger
 
 	FRewindDebuggerRuntime* FRewindDebuggerRuntime::InternalInstance = nullptr;
 
+	FRewindDebuggerRuntime::FRewindDebuggerRuntime()
+		: FRuntimeModule(LogRewindDebuggerRuntime, UE::RewindDebugger::DebuggerGuid)
+	{
+		DebuggerName = TEXT("RewindDebugger");
+		AutoRecordCommand = TEXT("StartRewindDebuggerRecording");
+		HostAddressCommand = TEXT("HostAddress");
+	}
+
 	void FRewindDebuggerRuntime::Initialize()
 	{
 		InternalInstance = new FRewindDebuggerRuntime();
@@ -30,6 +38,17 @@ namespace RewindDebugger
 		InternalInstance = nullptr;
 	}
 
+	void FRewindDebuggerRuntime::StartRecording()
+	{
+		using namespace UE::TraceBasedDebuggers;
+		FRuntimeModule::StartRecording(FStartRecordingCommandMessage
+			{
+				.Target = TEXT("127.0.0.1"),
+				.RecordingMode = ERecordingMode::Live,
+				.TransportMode = ETraceTransportMode::TraceServer
+			});
+	}
+
 	static void IterateExtensions(TFunction<void(IRewindDebuggerRuntimeExtension* Extension)> IteratorFunction)
 	{
 		// update extensions
@@ -42,158 +61,85 @@ namespace RewindDebugger
 			IteratorFunction(Extension);
 		}
 	}
-	
-	static void DisableAllTraceChannels()
+
+	void FRewindDebuggerRuntime::OnEnableRequiredTraceChannelsInternal()
 	{
-		UE::Trace::EnumerateChannels([](const ANSICHAR* ChannelName, bool bEnabled, void*)
-		{
-			if (bEnabled)
+		UE::Trace::ToggleChannel(TEXT("Object"), true);
+		UE::Trace::ToggleChannel(TEXT("ObjectProperties"), true);
+		UE::Trace::ToggleChannel(TEXT("Frame"), true);
+	}
+
+	void FRewindDebuggerRuntime::OnRecordingStartingInternal()
+	{
+#if OBJECT_TRACE_ENABLED
+		FObjectTrace::Reset();
+#endif
+
+		ClearRecording.Broadcast();
+	
+		// update extensions
+		IterateExtensions([](IRewindDebuggerRuntimeExtension* Extension)
 			{
-				FString ChannelNameFString(ChannelName);
-				UE::Trace::ToggleChannel(ChannelNameFString.GetCharArray().GetData(), false);
+				Extension->Clear();
 			}
-		}
-		, nullptr);
+		);
 	}
-	
-	void FRewindDebuggerRuntime::StartRecordingWithArgs(const TArray<FString>& Args)
-	{
-		if (Args.Num() > 0)
-		{
-			FTraceAuxiliary::EConnectionType TraceType = FTraceAuxiliary::EConnectionType::None;
-			FString TraceDestination;
 
-			for (const FString& Arg : Args)
+	void FRewindDebuggerRuntime::OnRecordingStartedInternal()
+	{
+		// Previous implementation calling the following code was executed on the game thread
+		// after calling FTraceAuxiliary::Start, but we are now waiting for the connection
+		// to be ready (i.e. FTraceAuxiliary::OnConnection) which comes from a worker thread.
+		// Extensions code in 'RecordingStarted' might require to run on the Game thread so this
+		// is why we defer the callbacks.
+		ExecuteOnGameThread(UE_SOURCE_LOCATION, [this]
 			{
-				if (Arg.StartsWith(TEXT("-tracefile"), ESearchCase::IgnoreCase))
-				{
-					ensureMsgf(TraceType == FTraceAuxiliary::EConnectionType::None, TEXT("RewindDebugger.StartRecording: Specifying more than 1 trace destination is not supported. Received: %s"), *FString::Join(Args, TEXT(" ")));
-				
-					TraceType = FTraceAuxiliary::EConnectionType::File;
-					
-					// Try to extract filename.
-					if (FParse::Value(*Arg, TEXT("-tracefile="), TraceDestination))
+				// update extensions
+				IterateExtensions([](IRewindDebuggerRuntimeExtension* Extension)
 					{
-						// Make sure it's a valid filename
-						FText FilenameError;
-						if (!FFileHelper::IsFilenameValidForSaving(TraceDestination, FilenameError))
-						{
-							ensureMsgf(false, TEXT("RewindDebugger.StartRecording: Specified filename is not supported: %s"), *FilenameError.ToString());
-							TraceDestination = "";
-						}
+						Extension->RecordingStarted();
 					}
-				}
-				else if (Arg.StartsWith(TEXT("-tracehost"), ESearchCase::IgnoreCase))
+				);
+
+#if OBJECT_TRACE_ENABLED
+				// trace each play-in-editor world, and all the actors in it.
+				for (TObjectIterator<UWorld> World; World; ++World)
 				{
-					ensureMsgf(TraceType == FTraceAuxiliary::EConnectionType::None, TEXT("RewindDebugger.StartRecording: Specifying more than 1 trace destination is not supported. Received: %s"), *FString::Join(Args, TEXT(" ")));
 
-					TraceType = FTraceAuxiliary::EConnectionType::Network;
+					FObjectTrace::ResetWorldElapsedTime(*World);
+
+					TRACE_WORLD(*World);
 
-					if (FParse::Value(*Arg, TEXT("-tracehost="), TraceDestination))
+					for (TActorIterator<AController> Iterator(*World); Iterator; ++Iterator)
 					{
-						// Should we validate that TraceDestination is valid ip address?
+						if (APawn* Pawn = Iterator->GetPawn())
+						{
+							TRACE_PAWN_POSSESS(static_cast<UObject*>(*Iterator), static_cast<UObject*>(Pawn));
+						}
 					}
 				}
-				else
-				{
-					ensureMsgf(false, TEXT("RewindDebugger.StartRecording: Received unknown argument: %s"), *Arg);
-				}
-			}
-
-			if (TraceType != FTraceAuxiliary::EConnectionType::None)
-			{
-				StartRecording(TraceType, *TraceDestination);
-				return;
-			}
-		}
+#endif // OBJECT_TRACE_ENABLED
+			});
 
-		// No destination was specified, just use the default.
-		StartRecording();
+		RecordingStarted.Broadcast();
 	}
-	
-	void FRewindDebuggerRuntime::StartRecording()
+
+	void FRewindDebuggerRuntime::OnRecordingStartFailedInternal(const FText& FailureReason)
 	{
-		StartRecording(FTraceAuxiliary::EConnectionType::Network, TEXT("127.0.0.1"));
+		RecordingStartFailed.Broadcast(FailureReason);
 	}
 
-	void FRewindDebuggerRuntime::StartRecording(FTraceAuxiliary::EConnectionType TraceType, const TCHAR* TraceDestination)
+	void FRewindDebuggerRuntime::OnRecordingStoppedInternal()
 	{
-#if OBJECT_TRACE_ENABLED
-		FObjectTrace::Reset();
-
-		ClearRecording.Broadcast();
-
-		// update extensions
-		IterateExtensions([](IRewindDebuggerRuntimeExtension* Extension)
-			{
-				Extension->Clear();
-			}
-		);
-        	
-        // Disable all trace channels, and then enable only the ones needed by RewindDebugger
-        // for systems with RewindDebugger integration, they should enable their channel(s) in an Extension in "RecordingStarted"
-        DisableAllTraceChannels();
-
-		// Clear all buffered data and prevent data from previous recordings from leaking into the new recording
-		FTraceAuxiliary::FOptions Options;
-		Options.bExcludeTail = true;
-	
-		// FTraceAuxiliary::OnConnection.AddRaw(this, &FRewindDebugger::OnConnection); 
-	
-		FTraceAuxiliary::Start(TraceType, TraceDestination, TEXT(""), &Options, LogRewindDebuggerRuntime);
-		
-		UE::Trace::ToggleChannel(TEXT("Object"), true);
-		UE::Trace::ToggleChannel(TEXT("ObjectProperties"), true);
-		UE::Trace::ToggleChannel(TEXT("Frame"), true);
-		
-		bIsRecording = true;
-
 		// update extensions
 		IterateExtensions([](IRewindDebuggerRuntimeExtension* Extension)
 			{
-				Extension->RecordingStarted();
+				Extension->RecordingStopped();
 			}
 		);
-	
-		// trace each play-in-editor world, and all the actors in it.
-		for (TObjectIterator<UWorld> World; World; ++World)
-		{
-			FObjectTrace::ResetWorldElapsedTime(*World);
-			
-			TRACE_WORLD(*World);
-				
-			for (TActorIterator<AController> Iterator(*World); Iterator; ++Iterator)
-			{
-				if (APawn* Pawn = Iterator->GetPawn())
-				{
-					TRACE_PAWN_POSSESS(static_cast<UObject*>(*Iterator), static_cast<UObject*>(Pawn));
-				}
-			}
-		}
 
-		RecordingStarted.Broadcast();
-	#endif // OBJECT_TRACE_ENABLED
+		RecordingStopped.Broadcast();
 	}
-	
-	void FRewindDebuggerRuntime::StopRecording()
-	{
-		if (bIsRecording)
-		{
-			// // update extensions
-			IterateExtensions([](IRewindDebuggerRuntimeExtension* Extension)
-				{
-					Extension->RecordingStopped();
-				}
-			);
-		
-			bIsRecording = false;
-		
-			DisableAllTraceChannels();
-			FTraceAuxiliary::Stop();
-		
-			RecordingStopped.Broadcast();
-		}
-	}
-
-
 }
+
+#endif // WITH_TRACE_BASED_DEBUGGERS
\ No newline at end of file
diff --git a/Engine/Plugins/Animation/GameplayInsights/Source/RewindDebuggerRuntime/Private/RewindDebuggerRuntimeModule.cpp b/Engine/Plugins/Animation/GameplayInsights/Source/RewindDebuggerRuntime/Private/RewindDebuggerRuntimeModule.cpp
--- a/Engine/Plugins/Animation/GameplayInsights/Source/RewindDebuggerRuntime/Private/RewindDebuggerRuntimeModule.cpp
+++ b/Engine/Plugins/Animation/GameplayInsights/Source/RewindDebuggerRuntime/Private/RewindDebuggerRuntimeModule.cpp
@@ -1,12 +1,29 @@
 // Copyright Epic Games, Inc. All Rights Reserved.
 
+#if WITH_TRACE_BASED_DEBUGGERS
 #include "RewindDebuggerRuntimeModule.h"
 
-#include "RewindDebuggerRuntime/RewindDebuggerRuntime.h"
-#include "Modules/ModuleManager.h"
+#include "EngineEditorBridge.h"
+#include "Features/IModularFeatures.h"
 #include "HAL/IConsoleManager.h"
+#include "Modules/ModuleManager.h"
 #include "ObjectTrace.h"
-#include "Features/IModularFeatures.h"
+#include "RewindDebuggerAnimationRuntime.h"
+#include "RewindDebuggerEngineEditorBridge.h"
+#include "RewindDebuggerRuntime/RewindDebuggerRuntime.h"
+
+class FRewindDebuggerRuntimeModule : public IModuleInterface
+{
+public:
+	//~ IModuleInterface interface
+	virtual void StartupModule() override;
+	virtual void ShutdownModule() override;
+	
+private:
+	TArray<IConsoleObject*> ConsoleObjects;
+
+	RewindDebugger::FRewindDebuggerAnimationRuntime AnimationExtension;
+};
 
 void FRewindDebuggerRuntimeModule::StartupModule()
 {
@@ -14,28 +31,51 @@ void FRewindDebuggerRuntimeModule::StartupModule()
 	{
 		RewindDebugger::FRewindDebuggerRuntime::Initialize();
 	}
-	
+
 	ConsoleObjects.Add(IConsoleManager::Get().RegisterConsoleCommand(
 		TEXT("RewindDebugger.StartRecording"),
 		TEXT("Starts making a rewind debugger recording."),
-		FConsoleCommandWithArgsDelegate::CreateRaw(RewindDebugger::FRewindDebuggerRuntime::Instance(), &RewindDebugger::FRewindDebuggerRuntime::StartRecordingWithArgs),
+		FConsoleCommandWithArgsDelegate::CreateLambda([](const TArray<FString>& Args)
+			{
+				UE_AUTORTFM_ONCOMMIT(Args)
+				{
+					if (RewindDebugger::FRewindDebuggerRuntime* Runtime = RewindDebugger::FRewindDebuggerRuntime::Instance())
+					{
+						Runtime->StartRecording(Args);
+					}
+				};
+			}),
 		ECVF_Default
 	));
-	
+
 	ConsoleObjects.Add(IConsoleManager::Get().RegisterConsoleCommand(
-   		TEXT("RewindDebugger.StopRecording"),
-   		TEXT("Stops the current rewind debugger recording."),
-		FConsoleCommandDelegate::CreateRaw(RewindDebugger::FRewindDebuggerRuntime::Instance(), &RewindDebugger::FRewindDebuggerRuntime::StopRecording),
-   		ECVF_Default
-   	));
-	
+		TEXT("RewindDebugger.StopRecording"),
+		TEXT("Stops the current rewind debugger recording."),
+		FConsoleCommandWithArgsDelegate::CreateLambda([](const TArray<FString>& Args)
+			{
+				UE_AUTORTFM_ONCOMMIT(Args)
+				{
+					if (RewindDebugger::FRewindDebuggerRuntime* Runtime = RewindDebugger::FRewindDebuggerRuntime::Instance())
+					{
+						Runtime->StopRecording();
+					}
+				};
+			}),
+		ECVF_Default
+	));
+
 	IModularFeatures::Get().RegisterModularFeature(IRewindDebuggerRuntimeExtension::ModularFeatureName, &AnimationExtension);
+
+	FCoreDelegates::OnFEngineLoopInitComplete.AddLambda([]()
+		{
+			UE::RewindDebugger::FRewindDebuggerEngineEditorBridge::Get().Initialize();
+		});
 }
 
 void FRewindDebuggerRuntimeModule::ShutdownModule()
 {
-	IConsoleManager &ConsoleManager = IConsoleManager::Get();
-	for (IConsoleObject *ConsoleObject : ConsoleObjects)
+	IConsoleManager& ConsoleManager = IConsoleManager::Get();
+	for (IConsoleObject* ConsoleObject : ConsoleObjects)
 	{
 		if (ConsoleObject)
 		{
@@ -43,10 +83,20 @@ void FRewindDebuggerRuntimeModule::ShutdownModule()
 		}
 	}
 	ConsoleObjects.Empty();
-	
+
 	RewindDebugger::FRewindDebuggerRuntime::Shutdown();
-	
+
 	IModularFeatures::Get().UnregisterModularFeature(IRewindDebuggerRuntimeExtension::ModularFeatureName, &AnimationExtension);
+
+	// Calling Teardown only if the bridge was instantiated since we don't
+	// want calling 'Get()' which will create the global instance during the shutdown.
+	using namespace UE::RewindDebugger;
+	if (FRewindDebuggerEngineEditorBridge::IsInstantiated())
+	{
+		FRewindDebuggerEngineEditorBridge::Get().TearDown();
+	}
 }
 
 IMPLEMENT_MODULE(FRewindDebuggerRuntimeModule, RewindDebuggerRuntime);
+
+#endif // WITH_TRACE_BASED_DEBUGGERS
\ No newline at end of file
diff --git a/Engine/Plugins/Animation/GameplayInsights/Source/RewindDebuggerRuntime/Private/RewindDebuggerRuntimeModule.h b/Engine/Plugins/Animation/GameplayInsights/Source/RewindDebuggerRuntime/Private/RewindDebuggerRuntimeModule.h
--- a/Engine/Plugins/Animation/GameplayInsights/Source/RewindDebuggerRuntime/Private/RewindDebuggerRuntimeModule.h
+++ b/Engine/Plugins/Animation/GameplayInsights/Source/RewindDebuggerRuntime/Private/RewindDebuggerRuntimeModule.h
@@ -2,19 +2,8 @@
 
 #pragma once
 
+#if UE_ENABLE_INCLUDE_ORDER_DEPRECATED_IN_5_8
 #include "CoreMinimal.h"
 #include "Modules/ModuleInterface.h"
 #include "RewindDebuggerAnimationRuntime.h"
-
-class FRewindDebuggerRuntimeModule : public IModuleInterface
-{
-public:
-	// IModuleInterface interface
-	virtual void StartupModule() override;
-	virtual void ShutdownModule() override;
-	
-private:
-	TArray<IConsoleObject*> ConsoleObjects;
-
-	RewindDebugger::FRewindDebuggerAnimationRuntime AnimationExtension;
-};
+#endif // UE_ENABLE_INCLUDE_ORDER_DEPRECATED_IN_5_8
diff --git a/Engine/Plugins/Animation/GameplayInsights/Source/RewindDebuggerRuntime/Public/RewindDebuggerEngineEditorBridge.h b/Engine/Plugins/Animation/GameplayInsights/Source/RewindDebuggerRuntime/Public/RewindDebuggerEngineEditorBridge.h
new file mode 100644
--- /dev/null
+++ b/Engine/Plugins/Animation/GameplayInsights/Source/RewindDebuggerRuntime/Public/RewindDebuggerEngineEditorBridge.h
@@ -0,0 +1,75 @@
+// Copyright Epic Games, Inc. All Rights Reserved.
+
+#pragma once
+
+#if WITH_TRACE_BASED_DEBUGGERS
+
+#include "Delegates/IDelegateInstance.h"
+#include "EngineEditorBridge.h"
+#include "RewindDebuggerRuntime/RewindDebuggerRuntime.h"
+
+#define UE_API REWINDDEBUGGERRUNTIME_API
+
+namespace UE::RewindDebugger
+{
+
+/**
+ * Object that bridges the RewindDebugger runtime module and the Engine/Editor.
+ * This layer is not strictly necessary for RewindDebugger since its runtime currently has access to the engine module, but we use it
+ * to share a common stack of layers with other trace-based debuggers.
+ */
+class FRewindDebuggerEngineEditorBridge : public TraceBasedDebuggers::FEngineEditorBridge
+{
+public:
+	FRewindDebuggerEngineEditorBridge();
+
+	~FRewindDebuggerEngineEditorBridge() = default;
+
+	UE_API static FRewindDebuggerEngineEditorBridge& Get();
+
+	/** @return Whether the static instance has been created (i.e., Get() was called and teardown is required) */
+	UE_API static bool IsInstantiated();
+
+private:
+	virtual void OnInitializeInternal() override;
+	virtual void OnTearDownInternal() override;
+	virtual void OnRecordingStartedInternal() override;
+	virtual void OnRecordingStoppedInternal() override;
+	virtual void BuildRecordingStatusInternal(TraceBasedDebuggers::FRecordingStatusMessage& OutStatusMessage) const override;
+
+	bool AddOnScreenRecordingMessage(float DummyDeltaTime = 0.1f);
+	void RemoveOnScreenRecordingMessage();
+	void HandleRecordingStartFailed(const FText& InFailureReason) const;
+	void HandlePIEStarted(UGameInstance* GameInstance);
+
+	void HandleTraceConnectionDetailsUpdated() const;
+
+	FDelegateHandle RecordingStartedHandle;
+	FDelegateHandle RecordingStoppedHandle;
+	FDelegateHandle RecordingStartFailedHandle;
+	uint64 RecordingMessageKey = 0;
+
+#if WITH_EDITOR
+	FDelegateHandle PIEStartedHandle;
+#endif
+
+	FTSTicker::FDelegateHandle DeferredShowMessageOnScreenHandle;
+	static bool bIsInstantiated;
+};
+
+} // UE::RewindDebugger
+
+#else
+
+#define UE_API REWINDDEBUGGERRUNTIME_API
+namespace UE::RewindDebugger
+{
+class FRewindDebuggerEngineEditorBridge
+{
+public:
+	static FRewindDebuggerEngineEditorBridge& Get();
+};
+} // UE::RewindDebugger
+#endif
+
+#undef UE_API
diff --git a/Engine/Plugins/Animation/GameplayInsights/Source/RewindDebuggerRuntime/Public/RewindDebuggerRuntime/RewindDebuggerRuntime.h b/Engine/Plugins/Animation/GameplayInsights/Source/RewindDebuggerRuntime/Public/RewindDebuggerRuntime/RewindDebuggerRuntime.h
--- a/Engine/Plugins/Animation/GameplayInsights/Source/RewindDebuggerRuntime/Public/RewindDebuggerRuntime/RewindDebuggerRuntime.h
+++ b/Engine/Plugins/Animation/GameplayInsights/Source/RewindDebuggerRuntime/Public/RewindDebuggerRuntime/RewindDebuggerRuntime.h
@@ -1,41 +1,63 @@
-// Copyright Epic Games, Inc. All Rights Reserved.
+﻿// Copyright Epic Games, Inc. All Rights Reserved.
 
 #pragma once
 
+#if UE_ENABLE_INCLUDE_ORDER_DEPRECATED_IN_5_8
 #include "CoreMinimal.h"
-#include "ProfilingDebugging/TraceAuxiliary.h" // for FTraceAuxiliary::EConnectionType
+#include "ProfilingDebugging/TraceAuxiliary.h"
+#endif // UE_ENABLE_INCLUDE_ORDER_DEPRECATED_IN_5_8
+
+#if WITH_TRACE_BASED_DEBUGGERS
+
+#include "TraceBasedDebuggerRuntime.h"
 
 #define UE_API REWINDDEBUGGERRUNTIME_API
 
 REWINDDEBUGGERRUNTIME_API DECLARE_LOG_CATEGORY_EXTERN(LogRewindDebuggerRuntime, Log, All);
 
+namespace UE::RewindDebugger
+{
+static constexpr FGuid DebuggerGuid = FGuid(0xC4E4CF24, 0x57DB4959, 0x8EDBD741, 0x673D25EC);
+} // UE::RewindDebugger
+
 
 namespace RewindDebugger
 {
 
-	class FRewindDebuggerRuntime
+	class FRewindDebuggerRuntime : public UE::TraceBasedDebuggers::FRuntimeModule
 	{
 	public:
+		UE_API FRewindDebuggerRuntime();
+
 		static UE_API void Initialize();
 		static UE_API void Shutdown();
-		static FRewindDebuggerRuntime* Instance() { return InternalInstance; }
-			
-		UE_API void StartRecording();
-		UE_API void StopRecording();
+		static FRewindDebuggerRuntime* Instance()
+		{
+			return InternalInstance;
+		}
 
-		UE_API void StartRecordingWithArgs(const TArray<FString>& Args);
+		UE_API void StartRecording();
 
-		bool IsRecording() const { return bIsRecording; }
+		//~ For StartRecording overloads taking an arguments list or a FStartRecordingCommandMessage struct
+		using FRuntimeModule::StartRecording;
 
-		FSimpleMulticastDelegate RecordingStarted;
 		FSimpleMulticastDelegate ClearRecording;
-		FSimpleMulticastDelegate RecordingStopped;
+		UE::TraceBasedDebuggers::FRecordingStateChangedDelegate RecordingStarted;
+		UE::TraceBasedDebuggers::FRecordingStartFailedDelegate RecordingStartFailed;
+		UE::TraceBasedDebuggers::FRecordingStateChangedDelegate RecordingStopped;
+
 	private:
-		UE_API void StartRecording(FTraceAuxiliary::EConnectionType TraceType, const TCHAR* TraceDestination);
+		//~ Begin FRuntimeModule interface
+		virtual void OnRecordingStartingInternal() override;
+		virtual void OnRecordingStartedInternal() override;
+		virtual void OnRecordingStartFailedInternal(const FText& FailureReason) override;
+		virtual void OnRecordingStoppedInternal() override;
+		virtual void OnEnableRequiredTraceChannelsInternal() override;
+		//~ End FRuntimeModule interface
 
-		bool bIsRecording = false;
 		static UE_API FRewindDebuggerRuntime* InternalInstance;
 	};
 }
 
 #undef UE_API
+#endif // WITH_TRACE_BASED_DEBUGGERS
diff --git a/Engine/Plugins/Animation/GameplayInsights/Source/RewindDebuggerRuntime/RewindDebuggerRuntime.Build.cs b/Engine/Plugins/Animation/GameplayInsights/Source/RewindDebuggerRuntime/RewindDebuggerRuntime.Build.cs
--- a/Engine/Plugins/Animation/GameplayInsights/Source/RewindDebuggerRuntime/RewindDebuggerRuntime.Build.cs
+++ b/Engine/Plugins/Animation/GameplayInsights/Source/RewindDebuggerRuntime/RewindDebuggerRuntime.Build.cs
@@ -6,6 +6,8 @@ namespace UnrealBuildTool.Rules
 	{
 		public RewindDebuggerRuntime(ReadOnlyTargetRules Target) : base(Target)
 		{
+			bAllowUETypesInNamespaces = true;
+
 			PrivateDependencyModuleNames.AddRange(
 			new string[]
 			{
@@ -13,6 +15,7 @@ namespace UnrealBuildTool.Rules
 				"CoreUObject",
 				"Engine",
 				"RewindDebuggerRuntimeInterface",
+				"TraceBasedDebuggers",
 				"TraceLog",
 			});
 		}
diff --git a/Engine/Plugins/Animation/GameplayInsights/Source/RewindDebuggerVLog/Private/RewindDebuggerVLog.cpp b/Engine/Plugins/Animation/GameplayInsights/Source/RewindDebuggerVLog/Private/RewindDebuggerVLog.cpp
--- a/Engine/Plugins/Animation/GameplayInsights/Source/RewindDebuggerVLog/Private/RewindDebuggerVLog.cpp
+++ b/Engine/Plugins/Animation/GameplayInsights/Source/RewindDebuggerVLog/Private/RewindDebuggerVLog.cpp
@@ -4,6 +4,7 @@
 
 #include "FindInBlueprintManager.h"
 #include "IRewindDebugger.h"
+#include "IRewindDebuggerModule.h"
 #include "IVisualLoggerProvider.h"
 #include "LogVisualizerSettings.h"
 #include "ObjectTrace.h"
@@ -70,7 +71,13 @@ FRewindDebuggerVLog::~FRewindDebuggerVLog()
 
 void FRewindDebuggerVLog::Initialize()
 {
-	UToolMenu* Menu = UToolMenus::Get()->FindMenu("RewindDebugger.ToolBar");
+	const IRewindDebuggerModule* RewindDebuggerModule = FModuleManager::GetModulePtr<IRewindDebuggerModule>("RewindDebugger");
+	if (!ensure(RewindDebuggerModule))
+	{
+		return;
+	}
+
+	UToolMenu* Menu = UToolMenus::Get()->FindMenu(RewindDebuggerModule->GetMainToolbarName());
 
 	FToolMenuSection& NewSection = Menu->AddSection("Visual Logger", LOCTEXT("Visual Logger", "Visual Logger"));
 
@@ -117,11 +124,10 @@ void FRewindDebuggerVLog::RenderLogEntry(const FVisualLogEntry& Entry, UCanvas*
 			RenderingActor->AddLogEntry(Entry);
 		}
 	}
-	else
+	else if (const IRewindDebugger* RewindDebugger = IRewindDebugger::Instance())
 	{
-		UWorld* World = IRewindDebugger::Instance()->GetWorldToVisualize();
+		UWorld* World = RewindDebugger->GetWorldToVisualize();
 		FVisualLogEntryRenderer::RenderLogEntry(World, Entry, &MatchCategoryFilters, Canvas, GEngine->GetMediumFont(), MonospaceFont.Get(), ScreenTextY);
-
 	}
 }
 
@@ -185,24 +191,26 @@ void FRewindDebuggerVLog::MakeLogLevelMenu(UToolMenu* Menu) const
 void FRewindDebuggerVLog::MakeCategoriesMenu(UToolMenu* Menu) const
 {
 	FToolMenuSection& Section = Menu->AddSection("Categories");
-	IRewindDebugger* RewindDebugger = IRewindDebugger::Instance();
-	if (const TraceServices::IAnalysisSession* Session = RewindDebugger->GetAnalysisSession())
+	const IRewindDebugger* RewindDebugger = IRewindDebugger::Instance();
+
+	if (const TraceServices::IAnalysisSession* Session = RewindDebugger ? RewindDebugger->GetAnalysisSession() : nullptr)
 	{
 		TraceServices::FAnalysisSessionReadScope SessionReadScope(*Session);
 		if (const IVisualLoggerProvider* VisualLoggerProvider = Session->ReadProvider<IVisualLoggerProvider>("VisualLoggerProvider"))
 		{
 			VisualLoggerProvider->EnumerateCategories([&Section, this](const FName& Category)
 				{
+					const FText CategoryAsText = FText::FromName(Category);
 					Section.AddMenuEntry(Category,
-						FText::FromName(Category),
-						FText(),
+						CategoryAsText,
+						FText::Format(LOCTEXT("CategoryDesc", "Enable/Disable {0} category"), CategoryAsText),
 						FSlateIcon(),
 						FUIAction(
 							FExecuteAction::CreateLambda([Category, this]() { ToggleCategory(Category); }),
 							FCanExecuteAction(),
 							FGetActionCheckState::CreateLambda([Category, this]() { return IsCategoryActive(Category) ? ECheckBoxState::Checked : ECheckBoxState::Unchecked; })
 						),
-						EUserInterfaceActionType::Check
+						EUserInterfaceActionType::ToggleButton
 					);
 				});
 		}
diff --git a/Engine/Plugins/ChaosVD/Source/ChaosVD/Private/Widgets/ChaosVDRecordingControls.cpp b/Engine/Plugins/ChaosVD/Source/ChaosVD/Private/Widgets/ChaosVDRecordingControls.cpp
--- a/Engine/Plugins/ChaosVD/Source/ChaosVD/Private/Widgets/ChaosVDRecordingControls.cpp
+++ b/Engine/Plugins/ChaosVD/Source/ChaosVD/Private/Widgets/ChaosVDRecordingControls.cpp
@@ -232,7 +232,7 @@ void FChaosVDRecordingControls::CloseSessionByRemoteSessionIDInternal(const FGui
 	}
 }
 
-UE::TraceBasedDebuggers::FTraceSessionDescriptor& FChaosVDRecordingControls::GetCurrentSessionDescriptorInternal()
+const UE::TraceBasedDebuggers::FTraceSessionDescriptor& FChaosVDRecordingControls::GetCurrentSessionDescriptorInternal()
 {
 	if (const TSharedPtr<SChaosVDMainTab> MainTab = MainTabWeakPtr.Pin())
 	{
diff --git a/Engine/Plugins/ChaosVD/Source/ChaosVD/Private/Widgets/ChaosVDRecordingControls.h b/Engine/Plugins/ChaosVD/Source/ChaosVD/Private/Widgets/ChaosVDRecordingControls.h
--- a/Engine/Plugins/ChaosVD/Source/ChaosVD/Private/Widgets/ChaosVDRecordingControls.h
+++ b/Engine/Plugins/ChaosVD/Source/ChaosVD/Private/Widgets/ChaosVDRecordingControls.h
@@ -1,4 +1,4 @@
-﻿// Copyright Epic Games, Inc. All Rights Reserved.
+// Copyright Epic Games, Inc. All Rights Reserved.
 
 #pragma once
 
@@ -47,7 +47,7 @@ protected:
 
 	virtual void CloseSessionByRemoteSessionIDInternal(const FGuid& InSessionId) override;
 
-	virtual UE::TraceBasedDebuggers::FTraceSessionDescriptor& GetCurrentSessionDescriptorInternal() override;
+	virtual const UE::TraceBasedDebuggers::FTraceSessionDescriptor& GetCurrentSessionDescriptorInternal() override;
 	//~ End FRecordingControls interface
 
 	TSharedRef<SWidget> GenerateDataChannelsMenu();
diff --git a/Engine/Plugins/Experimental/Animation/TrajectoryTools/Source/TrajectoryTools/Private/TrajectoryRewindDebuggerExtension.cpp b/Engine/Plugins/Experimental/Animation/TrajectoryTools/Source/TrajectoryTools/Private/TrajectoryRewindDebuggerExtension.cpp
--- a/Engine/Plugins/Experimental/Animation/TrajectoryTools/Source/TrajectoryTools/Private/TrajectoryRewindDebuggerExtension.cpp
+++ b/Engine/Plugins/Experimental/Animation/TrajectoryTools/Source/TrajectoryTools/Private/TrajectoryRewindDebuggerExtension.cpp
@@ -7,6 +7,7 @@
 #include "IGameplayProvider.h"
 #include "Debug/DebugDrawService.h"
 #include "DrawDebugHelpers.h"
+#include "IRewindDebuggerModule.h"
 #include "Components/SkeletalMeshComponent.h"
 #include "Engine/Canvas.h"
 #include "SceneView.h"
@@ -159,7 +160,13 @@ void FRewindDebuggerTrajectory::UpdateState(IRewindDebugger* RewindDebugger)
 
 void FRewindDebuggerTrajectory::GenerateMenu()
 {
-	if (UToolMenu* Menu = UToolMenus::Get()->FindMenu("RewindDebugger.ToolBar"))
+	const IRewindDebuggerModule* RewindDebuggerModule = FModuleManager::GetModulePtr<IRewindDebuggerModule>("RewindDebugger");
+	if (!ensure(RewindDebuggerModule))
+	{
+		return;
+	}
+
+	if (UToolMenu* Menu = UToolMenus::Get()->FindMenu(RewindDebuggerModule->GetMainToolbarName()))
 	{
 		FToolMenuSection& Section = Menu->AddSection("TrajectoryWorkflows", LOCTEXT("Trajectory Tools","Trajectory Tools"));
 		
diff --git a/Engine/Source/Editor/RewindDebuggerInterface/Public/IRewindDebugger.h b/Engine/Source/Editor/RewindDebuggerInterface/Public/IRewindDebugger.h
--- a/Engine/Source/Editor/RewindDebuggerInterface/Public/IRewindDebugger.h
+++ b/Engine/Source/Editor/RewindDebuggerInterface/Public/IRewindDebugger.h
@@ -1,4 +1,4 @@
-// Copyright Epic Games, Inc. All Rights Reserved.
+﻿// Copyright Epic Games, Inc. All Rights Reserved.
 #pragma once
 
 #if UE_ENABLE_INCLUDE_ORDER_DEPRECATED_IN_5_7
@@ -90,9 +90,20 @@ public:
 	/** @return the current visible range in Rewind Debugger recording time units */
 	virtual const TRange<double>& GetCurrentViewRange() const = 0;
 
-	/** @return the current analysis session */
+	/**
+	 * Returns pointer to the current analysis sessions for immediate access.
+	 * If the goal is to store in a variable for later access then GetAnalysisSessionAsShared should be used instead.
+	 * @return the current analysis session
+	 * @see GetAnalysisSessionAsShared
+	 */
 	virtual const TraceServices::IAnalysisSession* GetAnalysisSession() const = 0;
 
+	/**
+	 * Returns shared pointer to the current analysis sessions.
+	 * @return the current analysis session
+	 */
+	virtual TSharedPtr<const TraceServices::IAnalysisSession> GetAnalysisSessionAsShared() const = 0;
+
 	/** @return insights id for the root object currently debugged */
 	virtual uint64 GetRootObjectId() const = 0;
 
diff --git a/Engine/Source/Editor/RewindDebuggerInterface/Public/IRewindDebuggerModule.h b/Engine/Source/Editor/RewindDebuggerInterface/Public/IRewindDebuggerModule.h
--- a/Engine/Source/Editor/RewindDebuggerInterface/Public/IRewindDebuggerModule.h
+++ b/Engine/Source/Editor/RewindDebuggerInterface/Public/IRewindDebuggerModule.h
@@ -18,6 +18,9 @@ public:
 	/** @return the name of the RewindDebugger details window that could be used to invoke the associated tab in the Editor. */
 	virtual FName GetDetailsTabName() const = 0;
 
+	/** @return the name of the main toolbar in the RewindDebugger window that could be used to add new entries. */
+	virtual FName GetMainToolbarName() const = 0;
+
 	virtual TSharedRef<SDockTab> SpawnRewindDebuggerTab(const FSpawnTabArgs& SpawnTabArgs) = 0;
 	virtual TSharedRef<SDockTab> SpawnRewindDebuggerDetailsTab(const FSpawnTabArgs& SpawnTabArgs) = 0;
 };
\ No newline at end of file
diff --git a/Engine/Source/Runtime/Experimental/TraceBasedDebuggersAnalysis/Public/Widgets/RecordingControls.h b/Engine/Source/Runtime/Experimental/TraceBasedDebuggersAnalysis/Public/Widgets/RecordingControls.h
--- a/Engine/Source/Runtime/Experimental/TraceBasedDebuggersAnalysis/Public/Widgets/RecordingControls.h
+++ b/Engine/Source/Runtime/Experimental/TraceBasedDebuggersAnalysis/Public/Widgets/RecordingControls.h
@@ -90,7 +90,7 @@ protected:
 	virtual void OnToggleRecordingStateInternal(const TSharedRef<FSessionInfo>& InSessionInfo) = 0;
 	virtual void CloseSessionByRemoteSessionIDInternal(const FGuid& InSessionId) = 0;
 
-	virtual FTraceSessionDescriptor& GetCurrentSessionDescriptorInternal() = 0;
+	virtual const FTraceSessionDescriptor& GetCurrentSessionDescriptorInternal() = 0;
 
 	UE_API void SelectTargetSession(FGuid SessionId);
 	UE_API TSharedPtr<FSessionInfo> GetCurrentSessionInfo() const;
