[RewindDebugger] added analysis opened/closed notification for editor extensions
- updated some editor extensions to use RuntimeExtension instead when it was the intent (e.g., controlling traces)

#Synced-CL 49420911

diff --git a/Engine/Plugins/Animation/GameplayInsights/Source/RewindDebugger/Private/RewindDebugger.cpp b/Engine/Plugins/Animation/GameplayInsights/Source/RewindDebugger/Private/RewindDebugger.cpp
--- a/Engine/Plugins/Animation/GameplayInsights/Source/RewindDebugger/Private/RewindDebugger.cpp
+++ b/Engine/Plugins/Animation/GameplayInsights/Source/RewindDebugger/Private/RewindDebugger.cpp
@@ -73,8 +73,6 @@ FRewindDebugger::FRewindDebugger()
 	if (FRewindDebuggerRuntime* Runtime = FRewindDebuggerRuntime::Instance())
 	{
 		ClearRecordingHandle = Runtime->ClearRecording.AddRaw(this, &FRewindDebugger::OnClearRecording);
-		RecordingStartedHandle = Runtime->RecordingStarted.AddRaw(this, &FRewindDebugger::OnRecordingStarted);
-		RecordingStoppedHandle = Runtime->RecordingStopped.AddRaw(this, &FRewindDebugger::OnRecordingStopped);
 	}
 #endif // WITH_TRACE_BASED_DEBUGGERS
 
@@ -156,8 +154,6 @@ FRewindDebugger::~FRewindDebugger()
 	if (RewindDebugger::FRewindDebuggerRuntime* Runtime = RewindDebugger::FRewindDebuggerRuntime::Instance())
 	{
 		Runtime->ClearRecording.Remove(ClearRecordingHandle);
-		Runtime->RecordingStarted.Remove(RecordingStartedHandle);
-		Runtime->RecordingStopped.Remove(RecordingStoppedHandle);
 	}
 #endif // WITH_TRACE_BASED_DEBUGGERS
 }
@@ -453,37 +449,11 @@ void FRewindDebugger::OnClearRecording()
 	);
 }
 
-void FRewindDebugger::OnRecordingStarted()
-{
-	// Previous implementation calling the following code was executed on the game thread
-	// after calling FTraceAuxiliary::Start, but we are now waiting for the connection
-	// to be ready (i.e. FTraceAuxiliary::OnConnection) which comes from a worker thread.
-	// Extensions code in 'RecordingStarted' might require to run on the Game thread so this
-	// is why we defer the callbacks.
-	ExecuteOnGameThread(UE_SOURCE_LOCATION, [this]
-		{
-			IterateExtensions([this](IRewindDebuggerExtension* Extension)
-				{
-					Extension->RecordingStarted(this);
-				});
-		});
-}
-
-void FRewindDebugger::OnRecordingStopped()
-{
-	IterateExtensions([this](IRewindDebuggerExtension* Extension)
-		{
-			Extension->RecordingStopped(this);
-		}
-	);
-}
-
 bool FRewindDebugger::CanOpenTrace() const
 {
 	return !bPIEStarted && !IsRecording();
 }
 
-
 void FRewindDebugger::OpenTrace(const FString& FilePath)
 {
 	ClearTraceAndActiveSession();
@@ -1077,6 +1047,15 @@ void FRewindDebugger::OpenSession(const UE::TraceBasedDebuggers::FTraceSessionDe
 {
 	ensureMsgf(!CurrentSessionDescriptor.IsValid(), TEXT("Expecting current session to be closed before opening a new one"));
 	CurrentSessionDescriptor = SessionDescriptor;
+
+	// Notify extensions
+	IterateExtensions([this](IRewindDebuggerExtension* Extension)
+		{
+			Extension->AnalysisSessionOpened(this);
+		});
+
+	// Invalidate world Id to reselect one from the new analysis
+	bDisplayWorldIdValid = false;
 }
 
 const UE::TraceBasedDebuggers::FTraceSessionDescriptor& FRewindDebugger::GetSessionDescriptor(const FGuid& RemoteSessionID) const
@@ -1099,6 +1078,12 @@ void FRewindDebugger::CloseActiveTraceSession()
 		if (CurrentSessionDescriptor.IsValid())
 		{
 			TraceManager->CloseSession(CurrentSessionDescriptor.SessionName);
+
+			// Notify extensions
+			IterateExtensions([this](IRewindDebuggerExtension* Extension)
+				{
+					Extension->AnalysisSessionClosed(this);
+				});
 		}
 	}
 
@@ -1388,6 +1373,7 @@ void FRewindDebugger::SetDisplayWorld(uint64 WorldId)
 			Extension->Update(0.0, this);
 		});
 }
+
 void FRewindDebugger::MakeWorldsMenu(UToolMenu* Menu)
 {
 	const FRewindDebugger* RewindDebugger = Instance();
diff --git a/Engine/Plugins/Animation/GameplayInsights/Source/RewindDebugger/Private/RewindDebugger.h b/Engine/Plugins/Animation/GameplayInsights/Source/RewindDebugger/Private/RewindDebugger.h
--- a/Engine/Plugins/Animation/GameplayInsights/Source/RewindDebugger/Private/RewindDebugger.h
+++ b/Engine/Plugins/Animation/GameplayInsights/Source/RewindDebugger/Private/RewindDebugger.h
@@ -116,8 +116,6 @@ public:
 	/** Start a new Recording:  Start tracing Objects, increment the current recording index, and reset the recording elapsed time to 0 */
 	virtual void StartRecording() const override;
 	void OnClearRecording();
-	void OnRecordingStarted();
-	void OnRecordingStopped();
 
 	bool CanOpenTrace() const;
 	void OpenTrace(const FString& FilePath);
@@ -267,8 +265,6 @@ private:
 	FDelegateHandle SingleStepPIEHandle;
 	FDelegateHandle ShutdownPIEHandle;
 
-	FDelegateHandle RecordingStartedHandle;
-	FDelegateHandle RecordingStoppedHandle;
 	FDelegateHandle ClearRecordingHandle;
 
 	void SetCurrentScrubTime(double Time);
diff --git a/Engine/Plugins/Animation/GameplayInsights/Source/RewindDebuggerRuntime/Private/RewindDebuggerEngineEditorBridge.cpp b/Engine/Plugins/Animation/GameplayInsights/Source/RewindDebuggerRuntime/Private/RewindDebuggerEngineEditorBridge.cpp
--- a/Engine/Plugins/Animation/GameplayInsights/Source/RewindDebuggerRuntime/Private/RewindDebuggerEngineEditorBridge.cpp
+++ b/Engine/Plugins/Animation/GameplayInsights/Source/RewindDebuggerRuntime/Private/RewindDebuggerEngineEditorBridge.cpp
@@ -48,7 +48,6 @@ void FRewindDebuggerEngineEditorBridge::OnInitializeInternal()
 
 	if (FRewindDebuggerRuntime* Runtime = FRewindDebuggerRuntime::Instance())
 	{
-		//Runtime->ClearRecording.AddRaw(this, &FRewindDebugger::OnClearRecording);
 		RecordingStartedHandle = Runtime->RecordingStarted.AddRaw(this, &FRewindDebuggerEngineEditorBridge::HandleRecordingStarted);
 		RecordingStartFailedHandle = Runtime->RecordingStartFailed.AddRaw(this, &FRewindDebuggerEngineEditorBridge::HandleRecordingStartFailed);
 		RecordingStoppedHandle = Runtime->RecordingStopped.AddRaw(this, &FRewindDebuggerEngineEditorBridge::HandleRecordingStopped);
diff --git a/Engine/Plugins/Cameras/GameplayCameras/Source/GameplayCameras/GameplayCameras.Build.cs b/Engine/Plugins/Cameras/GameplayCameras/Source/GameplayCameras/GameplayCameras.Build.cs
--- a/Engine/Plugins/Cameras/GameplayCameras/Source/GameplayCameras/GameplayCameras.Build.cs
+++ b/Engine/Plugins/Cameras/GameplayCameras/Source/GameplayCameras/GameplayCameras.Build.cs
@@ -19,6 +19,7 @@ public class GameplayCameras : ModuleRules
 				"HeadMountedDisplay",
 				"MovieScene",
 				"MovieSceneTracks",
+				"RewindDebuggerRuntimeInterface",
 				"StateTreeModule",
 				"TemplateSequence",
 				"TraceLog"
diff --git a/Engine/Plugins/Cameras/GameplayCameras/Source/GameplayCameras/Private/GameplayCamerasModule.cpp b/Engine/Plugins/Cameras/GameplayCameras/Source/GameplayCameras/Private/GameplayCamerasModule.cpp
--- a/Engine/Plugins/Cameras/GameplayCameras/Source/GameplayCameras/Private/GameplayCamerasModule.cpp
+++ b/Engine/Plugins/Cameras/GameplayCameras/Source/GameplayCameras/Private/GameplayCamerasModule.cpp
@@ -9,6 +9,12 @@
 #include "ShowFlags.h"
 #include "UObject/UObjectBase.h"
 
+#if UE_GAMEPLAY_CAMERAS_TRACE
+#include "Debug/CameraSystemTrace.h"
+#include "Features/IModularFeatures.h"
+#include "RewindDebuggerRuntimeInterface/IRewindDebuggerRuntimeExtension.h"
+#endif // UE_GAMEPLAY_CAMERAS_TRACE
+
 #define LOCTEXT_NAMESPACE "GameplayCamerasModule"
 
 DEFINE_LOG_CATEGORY(LogCameraSystem);
@@ -18,7 +24,26 @@ namespace UE::Cameras
 
 TCustomShowFlag<> GameplayCamerasShowFlag(TEXT("GameplayCameras"), true, SFG_Developer, LOCTEXT("ShowFlagDisplayName", "Gameplay Cameras"));
 
-}  // namespace UE::Cameras
+#if UE_GAMEPLAY_CAMERAS_TRACE
+/**
+ * Rewind debugger runtime extension to control camera system traces.
+ */
+class FCameraSystemRewindDebuggerRuntimeExtension : public IRewindDebuggerRuntimeExtension
+{
+public:
+	virtual void RecordingStarted() override
+	{
+		Trace::ToggleChannel(*FCameraSystemTrace::ChannelName, true);
+	}
+
+	virtual void RecordingStopped() override
+	{
+		Trace::ToggleChannel(*FCameraSystemTrace::ChannelName, false);
+	}
+};
+#endif // UE_GAMEPLAY_CAMERAS_TRACE
+
+} // namespace UE::Cameras
 
 IGameplayCamerasModule& IGameplayCamerasModule::Get()
 {
@@ -37,10 +62,23 @@ public:
 #if UE_GAMEPLAY_CAMERAS_DEBUG
 		UE::Cameras::FCameraDebugColors::RegisterBuiltinColorSchemes();
 #endif  // UE_GAMEPLAY_CAMERAS_DEBUG
+
+#if UE_GAMEPLAY_CAMERAS_TRACE
+		using namespace UE::Cameras;
+
+		RewindDebuggerRuntimeExtension = MakeShared<FCameraSystemRewindDebuggerRuntimeExtension>();
+
+		IModularFeatures& ModularFeatures = IModularFeatures::Get();
+		ModularFeatures.RegisterModularFeature(IRewindDebuggerRuntimeExtension::ModularFeatureName, RewindDebuggerRuntimeExtension.Get());
+#endif  // UE_GAMEPLAY_CAMERAS_TRACE
 	}
 
 	virtual void ShutdownModule() override
 	{
+#if UE_GAMEPLAY_CAMERAS_TRACE
+		IModularFeatures& ModularFeatures = IModularFeatures::Get();
+		ModularFeatures.UnregisterModularFeature(IRewindDebuggerRuntimeExtension::ModularFeatureName, RewindDebuggerRuntimeExtension.Get());
+#endif  // UE_GAMEPLAY_CAMERAS_TRACE
 		UnregisterBuiltInBlendableStructs();
 	}
 
@@ -125,6 +163,10 @@ private:
 #if WITH_EDITOR
 	TSharedPtr<IGameplayCamerasLiveEditManager> LiveEditManager;
 #endif
+
+#if UE_GAMEPLAY_CAMERAS_TRACE
+	TSharedPtr<UE::Cameras::FCameraSystemRewindDebuggerRuntimeExtension> RewindDebuggerRuntimeExtension;
+#endif  // UE_GAMEPLAY_CAMERAS_TRACE
 };
 
 IMPLEMENT_MODULE(FGameplayCamerasModule, GameplayCameras);
diff --git a/Engine/Plugins/Cameras/GameplayCameras/Source/GameplayCamerasEditor/Private/Trace/CameraSystemRewindDebuggerExtension.cpp b/Engine/Plugins/Cameras/GameplayCameras/Source/GameplayCamerasEditor/Private/Trace/CameraSystemRewindDebuggerExtension.cpp
--- a/Engine/Plugins/Cameras/GameplayCameras/Source/GameplayCamerasEditor/Private/Trace/CameraSystemRewindDebuggerExtension.cpp
+++ b/Engine/Plugins/Cameras/GameplayCameras/Source/GameplayCamerasEditor/Private/Trace/CameraSystemRewindDebuggerExtension.cpp
@@ -28,11 +28,6 @@ FCameraSystemRewindDebuggerExtension::~FCameraSystemRewindDebuggerExtension()
 	EnsureDebugDrawDelegate(false);
 }
 
-void FCameraSystemRewindDebuggerExtension::RecordingStarted(IRewindDebugger* RewindDebugger)
-{
-	UE::Trace::ToggleChannel(*FCameraSystemTrace::ChannelName, true);
-}
-
 void FCameraSystemRewindDebuggerExtension::Update(float DeltaTime, IRewindDebugger* RewindDebugger)
 {
 	if (RewindDebugger->IsPIESimulating() || RewindDebugger->GetRecordingDuration() == 0.0)
@@ -84,11 +79,6 @@ void FCameraSystemRewindDebuggerExtension::Update(float DeltaTime, IRewindDebugg
 	}
 }
 
-void FCameraSystemRewindDebuggerExtension::RecordingStopped(IRewindDebugger* RewindDebugger)
-{
-	UE::Trace::ToggleChannel(*FCameraSystemTrace::ChannelName, false);
-}
-
 void FCameraSystemRewindDebuggerExtension::Clear(IRewindDebugger* RewindDebugger)
 {
 	EnsureDebugDrawDelegate(false);
diff --git a/Engine/Plugins/Cameras/GameplayCameras/Source/GameplayCamerasEditor/Private/Trace/CameraSystemRewindDebuggerExtension.h b/Engine/Plugins/Cameras/GameplayCameras/Source/GameplayCamerasEditor/Private/Trace/CameraSystemRewindDebuggerExtension.h
--- a/Engine/Plugins/Cameras/GameplayCameras/Source/GameplayCamerasEditor/Private/Trace/CameraSystemRewindDebuggerExtension.h
+++ b/Engine/Plugins/Cameras/GameplayCameras/Source/GameplayCamerasEditor/Private/Trace/CameraSystemRewindDebuggerExtension.h
@@ -1,4 +1,4 @@
-// Copyright Epic Games, Inc. All Rights Reserved.
+﻿// Copyright Epic Games, Inc. All Rights Reserved.
 
 #pragma once
 
@@ -32,9 +32,7 @@ public:
 	virtual ~FCameraSystemRewindDebuggerExtension();
 
 	// IRewindDebuggerExtension interface.
-	virtual void RecordingStarted(IRewindDebugger* RewindDebugger) override;
 	virtual void Update(float DeltaTime, IRewindDebugger* RewindDebugger) override;
-	virtual void RecordingStopped(IRewindDebugger* RewindDebugger) override;
 	virtual void Clear(IRewindDebugger* RewindDebugger) override;
 
 private:
diff --git a/Engine/Plugins/Developer/FieldNotificationTrace/Source/FieldNotificationTraceEditor/FieldNotificationTraceEditor.Build.cs b/Engine/Plugins/Developer/FieldNotificationTrace/Source/FieldNotificationTraceEditor/FieldNotificationTraceEditor.Build.cs
--- a/Engine/Plugins/Developer/FieldNotificationTrace/Source/FieldNotificationTraceEditor/FieldNotificationTraceEditor.Build.cs
+++ b/Engine/Plugins/Developer/FieldNotificationTrace/Source/FieldNotificationTraceEditor/FieldNotificationTraceEditor.Build.cs
@@ -23,6 +23,7 @@ public class FieldNotificationTraceEditor : ModuleRules
 			{
 				"GameplayInsights",
 				"RewindDebuggerInterface",
+				"RewindDebuggerRuntimeInterface",
 				"Slate",
 				"SlateCore",
 				"TraceAnalysis",
diff --git a/Engine/Plugins/Developer/FieldNotificationTrace/Source/FieldNotificationTraceEditor/Private/FieldNotificationRewindDebugger.cpp b/Engine/Plugins/Developer/FieldNotificationTrace/Source/FieldNotificationTraceEditor/Private/FieldNotificationRewindDebugger.cpp
--- a/Engine/Plugins/Developer/FieldNotificationTrace/Source/FieldNotificationTraceEditor/Private/FieldNotificationRewindDebugger.cpp
+++ b/Engine/Plugins/Developer/FieldNotificationTrace/Source/FieldNotificationTraceEditor/Private/FieldNotificationRewindDebugger.cpp
@@ -1,14 +1,14 @@
 // Copyright Epic Games, Inc. All Rights Reserved.
 
 #include "FieldNotificationRewindDebugger.h"
-
 #include "FieldNotificationTraceProvider.h"
+#include "IRewindDebugger.h"
 //#include "IGameplayProvider.h"
 
 namespace UE::FieldNotification
 {
 
-void FRewindDebugger::Update(float DeltaTime, IRewindDebugger* RewindDebugger)
+void FRewindDebuggerExtension::Update(float DeltaTime, IRewindDebugger* RewindDebugger)
 {
 	const TraceServices::IAnalysisSession* AnalysisSession = RewindDebugger->GetAnalysisSession();
 	TraceServices::FAnalysisSessionReadScope SessionReadScope(*AnalysisSession);
@@ -18,20 +18,20 @@ void FRewindDebugger::Update(float DeltaTime, IRewindDebugger* RewindDebugger)
 	
 	const TraceServices::IFrameProvider& FrameProvider = TraceServices::ReadFrameProvider(*RewindDebugger->GetAnalysisSession());
 	TraceServices::FFrame Frame;
-	if (FrameProvider.GetFrameFromTime(ETraceFrameType::TraceFrameType_Game, RewindDebugger->CurrentTraceTime(), Frame))
+	if (FrameProvider.GetFrameFromTime(TraceFrameType_Game, RewindDebugger->CurrentTraceTime(), Frame))
 	{
 		// each tick update the UMG Preview Window with what we are trying to debug
 	}
 }
 
-void FRewindDebugger::RecordingStarted(IRewindDebugger*)
+void FRewindDebuggerRuntimeExtension::RecordingStarted()
 {
-	UE::Trace::ToggleChannel(TEXT("FieldNotificationChannel"), true);
+	Trace::ToggleChannel(TEXT("FieldNotificationChannel"), true);
 }
 
-void FRewindDebugger::RecordingStopped(IRewindDebugger*)
+void FRewindDebuggerRuntimeExtension::RecordingStopped()
 {
-	UE::Trace::ToggleChannel(TEXT("FieldNotificationChannel"), false);
+	Trace::ToggleChannel(TEXT("FieldNotificationChannel"), false);
 }
 
 } // namespace
diff --git a/Engine/Plugins/Developer/FieldNotificationTrace/Source/FieldNotificationTraceEditor/Private/FieldNotificationRewindDebugger.h b/Engine/Plugins/Developer/FieldNotificationTrace/Source/FieldNotificationTraceEditor/Private/FieldNotificationRewindDebugger.h
--- a/Engine/Plugins/Developer/FieldNotificationTrace/Source/FieldNotificationTraceEditor/Private/FieldNotificationRewindDebugger.h
+++ b/Engine/Plugins/Developer/FieldNotificationTrace/Source/FieldNotificationTraceEditor/Private/FieldNotificationRewindDebugger.h
@@ -1,19 +1,26 @@
 // Copyright Epic Games, Inc. All Rights Reserved.
 #pragma once
 
-#include "CoreMinimal.h"
-#include "IRewindDebugger.h"
 #include "IRewindDebuggerExtension.h"
+#include "RewindDebuggerRuntimeInterface/IRewindDebuggerRuntimeExtension.h"
+
+class IRewindDebugger;
 
 namespace UE::FieldNotification
 {
 
-class FRewindDebugger : public IRewindDebuggerExtension
+class FRewindDebuggerExtension : public IRewindDebuggerExtension
 {
 public:
 	virtual void Update(float DeltaTime, IRewindDebugger* RewindDebugger) override;
-	virtual void RecordingStarted(IRewindDebugger* RewindDebugger) override;
-	virtual void RecordingStopped(IRewindDebugger* RewindDebugger) override;
+};
+
+//~ @todo this extension should be moved to a runtime module to support trace recording on non-editor targets
+class FRewindDebuggerRuntimeExtension : public IRewindDebuggerRuntimeExtension
+{
+public:
+	virtual void RecordingStarted() override;
+	virtual void RecordingStopped() override;
 };
 
 } //namespace
diff --git a/Engine/Plugins/Developer/FieldNotificationTrace/Source/FieldNotificationTraceEditor/Private/FieldNotificationTraceEditorModule.cpp b/Engine/Plugins/Developer/FieldNotificationTrace/Source/FieldNotificationTraceEditor/Private/FieldNotificationTraceEditorModule.cpp
--- a/Engine/Plugins/Developer/FieldNotificationTrace/Source/FieldNotificationTraceEditor/Private/FieldNotificationTraceEditorModule.cpp
+++ b/Engine/Plugins/Developer/FieldNotificationTrace/Source/FieldNotificationTraceEditor/Private/FieldNotificationTraceEditorModule.cpp
@@ -13,7 +13,10 @@ namespace UE::FieldNotification
 
 void FTraceEditorModule::StartupModule()
 {
-	IModularFeatures::Get().RegisterModularFeature(IRewindDebuggerExtension::ModularFeatureName, &RewindDebugger);
+	IModularFeatures::Get().RegisterModularFeature(IRewindDebuggerExtension::ModularFeatureName, &RewindDebuggerExtension);
+
+	//~ @todo this extension should be moved to a runtime module to support trace recording on non-editor targets
+	IModularFeatures::Get().RegisterModularFeature(IRewindDebuggerRuntimeExtension::ModularFeatureName, &RewindDebuggerRuntimeExtension);
 	IModularFeatures::Get().RegisterModularFeature(RewindDebugger::IRewindDebuggerTrackCreator::ModularFeatureName, &TrackCreator);
 	IModularFeatures::Get().RegisterModularFeature(TraceServices::ModuleFeatureName, &TraceModule);
 }
@@ -23,7 +26,8 @@ void FTraceEditorModule::ShutdownModule()
 {
 	IModularFeatures::Get().UnregisterModularFeature(TraceServices::ModuleFeatureName, &TraceModule);
 	IModularFeatures::Get().UnregisterModularFeature(RewindDebugger::IRewindDebuggerTrackCreator::ModularFeatureName, &TrackCreator);
-	IModularFeatures::Get().UnregisterModularFeature(IRewindDebuggerExtension::ModularFeatureName, &RewindDebugger);
+	IModularFeatures::Get().UnregisterModularFeature(IRewindDebuggerRuntimeExtension::ModularFeatureName, &RewindDebuggerRuntimeExtension);
+	IModularFeatures::Get().UnregisterModularFeature(IRewindDebuggerExtension::ModularFeatureName, &RewindDebuggerExtension);
 }
 
 } //namespace 
diff --git a/Engine/Plugins/Developer/FieldNotificationTrace/Source/FieldNotificationTraceEditor/Private/FieldNotificationTraceEditorModule.h b/Engine/Plugins/Developer/FieldNotificationTrace/Source/FieldNotificationTraceEditor/Private/FieldNotificationTraceEditorModule.h
--- a/Engine/Plugins/Developer/FieldNotificationTrace/Source/FieldNotificationTraceEditor/Private/FieldNotificationTraceEditorModule.h
+++ b/Engine/Plugins/Developer/FieldNotificationTrace/Source/FieldNotificationTraceEditor/Private/FieldNotificationTraceEditorModule.h
@@ -25,7 +25,8 @@ public:
 	//~ End IModuleInterface interface
 
 private:
-	FRewindDebugger RewindDebugger;
+	FRewindDebuggerExtension RewindDebuggerExtension;
+	FRewindDebuggerRuntimeExtension RewindDebuggerRuntimeExtension;
 	FTracksCreator TrackCreator;
 	FTraceServiceModule TraceModule;
 };
diff --git a/Engine/Plugins/Experimental/Animation/TrajectoryTools/Source/TrajectoryTools/Private/TrajectoryRewindDebuggerExtension.cpp b/Engine/Plugins/Experimental/Animation/TrajectoryTools/Source/TrajectoryTools/Private/TrajectoryRewindDebuggerExtension.cpp
--- a/Engine/Plugins/Experimental/Animation/TrajectoryTools/Source/TrajectoryTools/Private/TrajectoryRewindDebuggerExtension.cpp
+++ b/Engine/Plugins/Experimental/Animation/TrajectoryTools/Source/TrajectoryTools/Private/TrajectoryRewindDebuggerExtension.cpp
@@ -95,13 +95,13 @@ void FRewindDebuggerTrajectory::Update(float DeltaTime, IRewindDebugger* RewindD
 	check(State.Trajectories.Num() == State.ObjectInfos.Num());
 }
 
-void FRewindDebuggerTrajectory::RecordingStarted(IRewindDebugger* RewindDebugger)
+void FRewindDebuggerTrajectory::AnalysisSessionOpened(IRewindDebugger* RewindDebugger)
 {
 	// Clear up any cached state / variables.
 	Reset();
 }
 
-void FRewindDebuggerTrajectory::RecordingStopped(IRewindDebugger* RewindDebugger)
+void FRewindDebuggerTrajectory::AnalysisSessionClosed(IRewindDebugger* RewindDebugger)
 {
 	
 }
diff --git a/Engine/Plugins/Experimental/Animation/TrajectoryTools/Source/TrajectoryTools/Public/TrajectoryRewindDebuggerExtension.h b/Engine/Plugins/Experimental/Animation/TrajectoryTools/Source/TrajectoryTools/Public/TrajectoryRewindDebuggerExtension.h
--- a/Engine/Plugins/Experimental/Animation/TrajectoryTools/Source/TrajectoryTools/Public/TrajectoryRewindDebuggerExtension.h
+++ b/Engine/Plugins/Experimental/Animation/TrajectoryTools/Source/TrajectoryTools/Public/TrajectoryRewindDebuggerExtension.h
@@ -31,8 +31,8 @@ public:
 
 	/** Begin IRewindDebuggerExtension implementation */
 	virtual void Update(float DeltaTime, IRewindDebugger* RewindDebugger) override;
-	virtual void RecordingStarted(IRewindDebugger* RewindDebugger) override;
-	virtual void RecordingStopped(IRewindDebugger* RewindDebugger) override;
+	virtual void AnalysisSessionOpened(IRewindDebugger* RewindDebugger) override;
+	virtual void AnalysisSessionClosed(IRewindDebugger* RewindDebugger) override;
 	virtual void Clear(IRewindDebugger* RewindDebugger) override;
 	virtual FString GetName() { return TEXT("RewindDebuggerTrajectory"); }
 	/* End IRewindDebuggerExtension implementation */
diff --git a/Engine/Plugins/Runtime/StateTree/Source/StateTreeEditorModule/Private/Debugger/StateTreeRewindDebuggerExtensions.cpp b/Engine/Plugins/Runtime/StateTree/Source/StateTreeEditorModule/Private/Debugger/StateTreeRewindDebuggerExtensions.cpp
--- a/Engine/Plugins/Runtime/StateTree/Source/StateTreeEditorModule/Private/Debugger/StateTreeRewindDebuggerExtensions.cpp
+++ b/Engine/Plugins/Runtime/StateTree/Source/StateTreeEditorModule/Private/Debugger/StateTreeRewindDebuggerExtensions.cpp
@@ -42,12 +42,12 @@ void FRewindDebuggerPlaybackExtension::Update(float DeltaTime, IRewindDebugger*
 	}
 }
 
-void FRewindDebuggerPlaybackExtension::RecordingStarted(IRewindDebugger* RewindDebugger)
+void FRewindDebuggerPlaybackExtension::AnalysisSessionOpened(IRewindDebugger* RewindDebugger)
 {
 	StateTree::Delegates::OnTraceAnalysisStateChanged.Broadcast(EStateTreeTraceAnalysisStatus::Started);
 }
 
-void FRewindDebuggerPlaybackExtension::RecordingStopped(IRewindDebugger* RewindDebugger)
+void FRewindDebuggerPlaybackExtension::AnalysisSessionClosed(IRewindDebugger* RewindDebugger)
 {
 	StateTree::Delegates::OnTraceAnalysisStateChanged.Broadcast(EStateTreeTraceAnalysisStatus::Stopped);
 }
diff --git a/Engine/Plugins/Runtime/StateTree/Source/StateTreeEditorModule/Private/Debugger/StateTreeRewindDebuggerExtensions.h b/Engine/Plugins/Runtime/StateTree/Source/StateTreeEditorModule/Private/Debugger/StateTreeRewindDebuggerExtensions.h
--- a/Engine/Plugins/Runtime/StateTree/Source/StateTreeEditorModule/Private/Debugger/StateTreeRewindDebuggerExtensions.h
+++ b/Engine/Plugins/Runtime/StateTree/Source/StateTreeEditorModule/Private/Debugger/StateTreeRewindDebuggerExtensions.h
@@ -12,8 +12,8 @@ namespace UE::StateTreeDebugger
 class FRewindDebuggerPlaybackExtension final : public IRewindDebuggerExtension
 {
 	virtual void Update(float DeltaTime, IRewindDebugger* RewindDebugger) override;
-	virtual void RecordingStarted(IRewindDebugger* RewindDebugger) override;
-	virtual void RecordingStopped(IRewindDebugger* RewindDebugger) override;
+	virtual void AnalysisSessionOpened(IRewindDebugger* RewindDebugger) override;
+	virtual void AnalysisSessionClosed(IRewindDebugger* RewindDebugger) override;
 	virtual void Clear(IRewindDebugger* RewindDebugger) override;
 
 	/** Last scrub time we received. Used to avoid redundant updates. */
diff --git a/Engine/Source/Editor/RewindDebuggerInterface/Public/IRewindDebuggerExtension.h b/Engine/Source/Editor/RewindDebuggerInterface/Public/IRewindDebuggerExtension.h
--- a/Engine/Source/Editor/RewindDebuggerInterface/Public/IRewindDebuggerExtension.h
+++ b/Engine/Source/Editor/RewindDebuggerInterface/Public/IRewindDebuggerExtension.h
@@ -1,4 +1,4 @@
-// Copyright Epic Games, Inc. All Rights Reserved.
+﻿// Copyright Epic Games, Inc. All Rights Reserved.
 #pragma once
 
 #if UE_ENABLE_INCLUDE_ORDER_DEPRECATED_IN_5_7
@@ -29,15 +29,25 @@ public:
 	/** called while scrubbing, playing back, or paused */
 	virtual void Update(float DeltaTime, IRewindDebugger* RewindDebugger)
 	{
-	};
+	}
+
+	/** called when opening an analysis session */
+	virtual void AnalysisSessionOpened(IRewindDebugger* RewindDebugger)
+	{
+	}
+
+	/** called when closing an analysis session */
+	virtual void AnalysisSessionClosed(IRewindDebugger* RewindDebugger)
+	{
+	}
 
-	/** called when recording has started */
-	virtual void RecordingStarted(IRewindDebugger* RewindDebugger)
+	UE_DEPRECATED(5.8, TEXT("Use AnalysisSessionOpened instead or implement an extension inheriting from IRewindDebuggerRuntimeExtension"))
+	virtual void RecordingStarted(IRewindDebugger* RewindDebugger) final
 	{
 	};
 
-	/** called when recording has ended */
-	virtual void RecordingStopped(IRewindDebugger* RewindDebugger)
+	UE_DEPRECATED(5.8, TEXT("Use AnalysisSessionClosed instead or implement an extension inheriting from IRewindDebuggerRuntimeExtension"))
+	virtual void RecordingStopped(IRewindDebugger* RewindDebugger) final
 	{
 	};
 
