[RewindDebugger] reenabling server world filtering in GameplayProvider while investigating regression and implementing proper fix

#Synced-CL 49426448

diff --git a/Engine/Plugins/Animation/GameplayInsights/Source/GameplayInsights/Private/GameplayProvider.cpp b/Engine/Plugins/Animation/GameplayInsights/Source/GameplayInsights/Private/GameplayProvider.cpp
--- a/Engine/Plugins/Animation/GameplayInsights/Source/GameplayInsights/Private/GameplayProvider.cpp
+++ b/Engine/Plugins/Animation/GameplayInsights/Source/GameplayInsights/Private/GameplayProvider.cpp
@@ -1,4 +1,4 @@
-// Copyright Epic Games, Inc. All Rights Reserved.
+﻿// Copyright Epic Games, Inc. All Rights Reserved.
 
 #include "GameplayProvider.h"
 
@@ -772,8 +772,9 @@ void FGameplayProvider::AppendRecordingInfo(uint64 InWorldId, double InProfileTi
 		{
 			// We should put each world into its own timeline in the future, to support frame stepping in sync with server frames
 			if (WorldInfos[*WorldIndex].NetMode == FWorldInfo::ENetMode::Standalone
-				|| WorldInfos[*WorldIndex].NetMode == FWorldInfo::ENetMode::Client
-				|| WorldInfos[*WorldIndex].NetMode == FWorldInfo::ENetMode::DedicatedServer)
+				|| WorldInfos[*WorldIndex].NetMode == FWorldInfo::ENetMode::Client)
+				// Disabling while investigating regression and implementing proper fix
+				//|| WorldInfos[*WorldIndex].NetMode == FWorldInfo::ENetMode::DedicatedServer)
 			{
 				FRecordingInfoMessage NewRecordingInfo;
 				NewRecordingInfo.WorldId = InWorldId;
