[RewindDebugger] stop the active analysis session (which closes the file writer) when stopping PIE

#Synced-CL 49442481

diff --git a/Engine/Plugins/Animation/GameplayInsights/Source/RewindDebugger/Private/RewindDebugger.cpp b/Engine/Plugins/Animation/GameplayInsights/Source/RewindDebugger/Private/RewindDebugger.cpp
--- a/Engine/Plugins/Animation/GameplayInsights/Source/RewindDebugger/Private/RewindDebugger.cpp
+++ b/Engine/Plugins/Animation/GameplayInsights/Source/RewindDebugger/Private/RewindDebugger.cpp
@@ -254,13 +254,16 @@ void FRewindDebugger::OnPIEStopped(bool bSimulating)
 		const UWorld* World = GetWorldToVisualize();
 		SetCurrentScrubTime(FObjectTrace::GetWorldElapsedTime(World));
 #endif // OBJECT_TRACE_ENABLED
+
+		// @todo we should make sure to only stop recording for the local session
+		// or prevent users from selecting a remote session while PIE is active and recording
+		StopRecording();
+
+		StopActiveTraceSession();
 	}
 
 	bPIEStarted = false;
 	bPIESimulating = false;
-
-	StopRecording();
-
 	bDisplayWorldIdValid = false;
 }
 
@@ -564,9 +567,7 @@ void FRewindDebugger::ClearTrace()
 
 void FRewindDebugger::ClearTraceAndActiveSession()
 {
-#if UE_WITH_TRACE_BASED_DEBUGGERS_ANALYSIS
 	CloseActiveTraceSession();
-#endif
 	ClearTrace();
 }
 
@@ -1070,26 +1071,6 @@ const UE::TraceBasedDebuggers::FTraceSessionDescriptor& FRewindDebugger::GetSess
 	return InvalidSessionDescriptor;
 }
 
-void FRewindDebugger::CloseActiveTraceSession()
-{
-	using namespace UE::TraceBasedDebuggers;
-	if (const TSharedPtr<FTraceSessionsManager> TraceManager = FRewindDebuggerModule::Get().GetTraceSessionsManager())
-	{
-		if (CurrentSessionDescriptor.IsValid())
-		{
-			TraceManager->CloseSession(CurrentSessionDescriptor.SessionName);
-
-			// Notify extensions
-			IterateExtensions([this](IRewindDebuggerExtension* Extension)
-				{
-					Extension->AnalysisSessionClosed(this);
-				});
-		}
-	}
-
-	CurrentSessionDescriptor = {};
-}
-
 void FRewindDebugger::CloseSessionByRemoteSessionID(const FGuid& RemoteSessionID)
 {
 	using namespace UE::TraceBasedDebuggers;
@@ -1103,18 +1084,45 @@ void FRewindDebugger::CloseSessionByRemoteSessionID(const FGuid& RemoteSessionID
 	CloseActiveTraceSession();
 }
 
-void FRewindDebugger::StopActiveTraceSessions() const
+#endif // UE_WITH_TRACE_BASED_DEBUGGERS_ANALYSIS
+
+void FRewindDebugger::StopActiveTraceSession() const
 {
+#if UE_WITH_TRACE_BASED_DEBUGGERS_ANALYSIS
 	using namespace UE::TraceBasedDebuggers;
-	if (const TSharedPtr<FTraceSessionsManager> TraceManager = FRewindDebuggerModule::Get().GetTraceSessionsManager())
+
+	if (CurrentSessionDescriptor.IsValid())
 	{
-		if (CurrentSessionDescriptor.IsValid())
+		if (const TSharedPtr<FTraceSessionsManager> TraceManager = FRewindDebuggerModule::Get().GetTraceSessionsManager())
 		{
 			TraceManager->StopSession(CurrentSessionDescriptor.SessionName);
 		}
 	}
+#endif // UE_WITH_TRACE_BASED_DEBUGGERS_ANALYSIS
 }
+
+void FRewindDebugger::CloseActiveTraceSession()
+{
+#if UE_WITH_TRACE_BASED_DEBUGGERS_ANALYSIS
+	using namespace UE::TraceBasedDebuggers;
+
+	if (CurrentSessionDescriptor.IsValid())
+	{
+		if (const TSharedPtr<FTraceSessionsManager> TraceManager = FRewindDebuggerModule::Get().GetTraceSessionsManager())
+		{
+			TraceManager->CloseSession(CurrentSessionDescriptor.SessionName);
+
+			// Notify extensions
+			IterateExtensions([this](IRewindDebuggerExtension* Extension)
+				{
+					Extension->AnalysisSessionClosed(this);
+				});
+		}
+	}
+
+	CurrentSessionDescriptor = {};
 #endif // UE_WITH_TRACE_BASED_DEBUGGERS_ANALYSIS
+}
 
 void FRewindDebugger::Tick(float DeltaTime)
 {
diff --git a/Engine/Plugins/Animation/GameplayInsights/Source/RewindDebugger/Private/RewindDebugger.h b/Engine/Plugins/Animation/GameplayInsights/Source/RewindDebugger/Private/RewindDebugger.h
--- a/Engine/Plugins/Animation/GameplayInsights/Source/RewindDebugger/Private/RewindDebugger.h
+++ b/Engine/Plugins/Animation/GameplayInsights/Source/RewindDebugger/Private/RewindDebugger.h
@@ -235,14 +235,14 @@ public:
 
 private:
 	void OpenSession(const UE::TraceBasedDebuggers::FTraceSessionDescriptor& SessionDescriptor);
-	void CloseActiveTraceSession();
-	void StopActiveTraceSessions() const;
-
 	const UE::TraceBasedDebuggers::FTraceSessionDescriptor& GetSessionDescriptor(const FGuid& RemoteSessionID) const;
 #endif // UE_WITH_TRACE_BASED_DEBUGGERS_ANALYSIS
 
 private:
 
+	void CloseActiveTraceSession();
+	void StopActiveTraceSession() const;
+
 #if WITH_TRACE_BASED_DEBUGGERS
 	void StopLocalRecording() const;
 #endif
