[RewindDebugger] added missing module implementation for targets not using WITH_TRACE_BASED_DEBUGGERS
#Synced-CL 49574741

diff --git a/Engine/Plugins/Animation/GameplayInsights/Source/RewindDebuggerRuntime/Private/RewindDebuggerRuntimeModule.cpp b/Engine/Plugins/Animation/GameplayInsights/Source/RewindDebuggerRuntime/Private/RewindDebuggerRuntimeModule.cpp
--- a/Engine/Plugins/Animation/GameplayInsights/Source/RewindDebuggerRuntime/Private/RewindDebuggerRuntimeModule.cpp
+++ b/Engine/Plugins/Animation/GameplayInsights/Source/RewindDebuggerRuntime/Private/RewindDebuggerRuntimeModule.cpp
@@ -1,12 +1,12 @@
-// Copyright Epic Games, Inc. All Rights Reserved.
+﻿// Copyright Epic Games, Inc. All Rights Reserved.
 
-#if WITH_TRACE_BASED_DEBUGGERS
 #include "RewindDebuggerRuntimeModule.h"
+#include "Modules/ModuleManager.h"
 
+#if WITH_TRACE_BASED_DEBUGGERS
 #include "EngineEditorBridge.h"
 #include "Features/IModularFeatures.h"
 #include "HAL/IConsoleManager.h"
-#include "Modules/ModuleManager.h"
 #include "ObjectTrace.h"
 #include "RewindDebuggerAnimationRuntime.h"
 #include "RewindDebuggerEngineEditorBridge.h"
@@ -99,4 +99,8 @@ void FRewindDebuggerRuntimeModule::ShutdownModule()
 
 IMPLEMENT_MODULE(FRewindDebuggerRuntimeModule, RewindDebuggerRuntime);
 
+#else
+
+IMPLEMENT_MODULE(FDefaultModuleImpl, RewindDebuggerRuntime);
+
 #endif // WITH_TRACE_BASED_DEBUGGERS
\ No newline at end of file
