Temporary workaround to fix Rewind Debugger scrubbing in PIE

#Synced-CL 49562535

diff --git a/Engine/Plugins/Animation/GameplayInsights/Source/RewindDebugger/Private/RewindDebugger.cpp b/Engine/Plugins/Animation/GameplayInsights/Source/RewindDebugger/Private/RewindDebugger.cpp
--- a/Engine/Plugins/Animation/GameplayInsights/Source/RewindDebugger/Private/RewindDebugger.cpp
+++ b/Engine/Plugins/Animation/GameplayInsights/Source/RewindDebugger/Private/RewindDebugger.cpp
@@ -599,7 +599,7 @@ void FRewindDebugger::SetShouldAutoEject(bool value)
 
 bool FRewindDebugger::CanUsePlaybackControls() const
 {
-	return !bPIESimulating && !IsRecording();
+	return !bPIESimulating;// && !IsRecording();
 }
 
 bool FRewindDebugger::CanPause() const
