[RewindDebugger] added Auto-Scroll functionality
- playback controls are now always available and will disable auto-scroll
- manual scrubbing will disable auto-scroll
- recording started on PIE will activate auto-scroll
- resuming recorded PIE session will reactivate auto-scroll

#Synced-CL 49574741

diff --git a/Engine/Plugins/Animation/GameplayInsights/Source/RewindDebugger/Private/RewindDebugger.cpp b/Engine/Plugins/Animation/GameplayInsights/Source/RewindDebugger/Private/RewindDebugger.cpp
--- a/Engine/Plugins/Animation/GameplayInsights/Source/RewindDebugger/Private/RewindDebugger.cpp
+++ b/Engine/Plugins/Animation/GameplayInsights/Source/RewindDebugger/Private/RewindDebugger.cpp
@@ -1,4 +1,4 @@
-﻿// Copyright Epic Games, Inc. All Rights Reserved.
+// Copyright Epic Games, Inc. All Rights Reserved.
 
 #include "RewindDebugger.h"
 
@@ -22,6 +22,7 @@
 #include "RewindDebuggerRecordingControls.h"
 #include "RewindDebuggerRuntime/RewindDebuggerRuntime.h"
 #include "RewindDebuggerSettings.h"
+#include "RewindDebuggerStyle.h"
 #include "RewindDebuggerTrackCreators.h"
 #include "SessionInfo.h"
 #include "SModalSessionBrowser.h"
@@ -194,6 +195,11 @@ void FRewindDebugger::OnPIEStarted(bool bSimulating)
 	if (ShouldAutoRecordOnPIE())
 	{
 		StartRecording();
+		bRecordingPIESession = true;
+
+		// Reactivate auto-scroll assuming user wants to visualize the newest data until
+		// he manually scrubs or uses the playback controls.
+		bAutoScrollToLastItem = true;
 	}
 }
 
@@ -202,14 +208,6 @@ void FRewindDebugger::OnPIEPaused(bool bSimulating)
 	bPIESimulating = false;
 	ControlState = EControlState::Pause;
 
-#if OBJECT_TRACE_ENABLED
-	if (IsRecording())
-	{
-		const UWorld* World = GetWorldToVisualize();
-		SetCurrentScrubTime(FObjectTrace::GetWorldElapsedTime(World));
-	}
-#endif // OBJECT_TRACE_ENABLED
-
 	if (ShouldAutoEject() && FPlayWorldCommandCallbacks::IsInPIE())
 	{
 		bool CanEject = false;
@@ -229,6 +227,13 @@ void FRewindDebugger::OnPIEResumed(bool bSimulating)
 {
 	bPIESimulating = true;
 
+	if (IsRecording()
+		&& bRecordingPIESession)
+	{
+		// Reactivate auto-scroll assuming user is done inspecting recorded data
+		bAutoScrollToLastItem = true;
+	}
+
 	if (ShouldAutoEject() && FPlayWorldCommandCallbacks::IsInSIE())
 	{
 		GEditor->RequestToggleBetweenPIEandSIE();
@@ -237,29 +242,20 @@ void FRewindDebugger::OnPIEResumed(bool bSimulating)
 
 void FRewindDebugger::OnPIESingleStepped(bool bSimulating)
 {
-#if OBJECT_TRACE_ENABLED
-	if (IsRecording())
+	if (IsRecording()
+		&& bRecordingPIESession)
 	{
-		const UWorld* World = GetWorldToVisualize();
-		SetCurrentScrubTime(FObjectTrace::GetWorldElapsedTime(World));
+		// Reactivate auto-scroll assuming user wants to visualize data for the next frame
+		bAutoScrollToLastItem = true;
 	}
-#endif // OBJECT_TRACE_ENABLED
 }
 
 void FRewindDebugger::OnPIEStopped(bool bSimulating)
 {
-	if (IsRecording() && bPIESimulating)
+	if (IsRecording()
+		&& bRecordingPIESession)
 	{
-#if OBJECT_TRACE_ENABLED
-		const UWorld* World = GetWorldToVisualize();
-		SetCurrentScrubTime(FObjectTrace::GetWorldElapsedTime(World));
-#endif // OBJECT_TRACE_ENABLED
-
-		// @todo we should make sure to only stop recording for the local session
-		// or prevent users from selecting a remote session while PIE is active and recording
 		StopRecording();
-
-		StopActiveTraceSession();
 	}
 
 	bPIEStarted = false;
@@ -419,6 +415,19 @@ void FRewindDebugger::StartRecording() const
 	}
 }
 
+void FRewindDebugger::OnRecordingStarted()
+{
+}
+
+void FRewindDebugger::OnRecordingStopped()
+{
+	if (bRecordingPIESession)
+	{
+		bRecordingPIESession = false;
+		StopActiveTraceSession();
+	}
+}
+
 void FRewindDebugger::StopRecording() const
 {
 	if (RecordingControls)
@@ -454,7 +463,7 @@ void FRewindDebugger::OnClearRecording()
 
 bool FRewindDebugger::CanOpenTrace() const
 {
-	return !bPIEStarted && !IsRecording();
+	return !IsRecording();
 }
 
 void FRewindDebugger::OpenTrace(const FString& FilePath)
@@ -551,7 +560,7 @@ void FRewindDebugger::ClearTrace()
 	CurrentTraceRange.SetLowerBoundValue(0);
 	CurrentTraceRange.SetUpperBoundValue(0);
 	RecordingDuration.Set(0.0);
-	SetCurrentScrubTime(0.0);
+	SetCurrentScrubTime(0.0, EAutoScrollChange::PreserveCurrentState);
 
 	TrackSelectionChanged(nullptr);
 
@@ -571,6 +580,26 @@ void FRewindDebugger::ClearTraceAndActiveSession()
 	ClearTrace();
 }
 
+void FRewindDebugger::ToggleAutoScroll()
+{
+	bAutoScrollToLastItem = !bAutoScrollToLastItem;
+	if (bAutoScrollToLastItem)
+	{
+		SetCurrentScrubTime(RecordingDuration.Get(), EAutoScrollChange::PreserveCurrentState);
+		TrackCursorDelegate.ExecuteIfBound(/*bReverse*/false);
+	}
+}
+
+bool FRewindDebugger::CanEnableAutoScroll() const
+{
+	return GetAnalysisSession() != nullptr;
+}
+
+bool FRewindDebugger::IsAutoScrollEnabled() const
+{
+	return bAutoScrollToLastItem;
+}
+
 bool FRewindDebugger::ShouldAutoRecordOnPIE() const
 {
 	return URewindDebuggerSettings::Get().bShouldAutoRecordOnPIE;
@@ -599,7 +628,8 @@ void FRewindDebugger::SetShouldAutoEject(bool value)
 
 bool FRewindDebugger::CanUsePlaybackControls() const
 {
-	return !bPIESimulating;// && !IsRecording();
+	// We now allow playback controls at any time and using them will disable auto-scroll
+	return true;
 }
 
 bool FRewindDebugger::CanPause() const
@@ -634,10 +664,11 @@ void FRewindDebugger::Play()
 {
 	if (CanPlay())
 	{
-		if (CurrentScrubTime >= RecordingDuration.Get())
-		{
-			SetCurrentScrubTime(0);
-		}
+		// Loop from the beginning if already at the end,
+		// otherwise keep the current time but make sure to disable the auto-scroll
+		const double AdjustedTime = (CurrentScrubTime >= RecordingDuration.Get()) ? 0 : CurrentScrubTime;
+		SetCurrentScrubTime(AdjustedTime, EAutoScrollChange::Disable);
+		TrackCursorDelegate.ExecuteIfBound(/*bReverse*/false);
 
 		ControlState = EControlState::Play;
 	}
@@ -652,10 +683,11 @@ void FRewindDebugger::PlayReverse()
 {
 	if (CanPlayReverse())
 	{
-		if (CurrentScrubTime <= 0)
-		{
-			SetCurrentScrubTime(RecordingDuration.Get());
-		}
+		// Loop from the end if already at the beginning,
+		// otherwise keep the current time but make sure to disable the auto-scroll
+		const double AdjustedTime = (CurrentScrubTime <= 0) ? RecordingDuration.Get() : CurrentScrubTime;
+		SetCurrentScrubTime(AdjustedTime, EAutoScrollChange::Disable);
+		TrackCursorDelegate.ExecuteIfBound(/*bReverse*/true);
 
 		ControlState = EControlState::PlayReverse;
 	}
@@ -671,8 +703,8 @@ void FRewindDebugger::ScrubToStart()
 	if (CanScrub())
 	{
 		Pause();
-		SetCurrentScrubTime(0);
-		TrackCursorDelegate.ExecuteIfBound(false);
+		SetCurrentScrubTime(0, EAutoScrollChange::Disable);
+		TrackCursorDelegate.ExecuteIfBound(/*bReverse*/false);
 	}
 }
 
@@ -681,8 +713,8 @@ void FRewindDebugger::ScrubToEnd()
 	if (CanScrub())
 	{
 		Pause();
-		SetCurrentScrubTime(RecordingDuration.Get());
-		TrackCursorDelegate.ExecuteIfBound(false);
+		SetCurrentScrubTime(RecordingDuration.Get(), EAutoScrollChange::Disable);
+		TrackCursorDelegate.ExecuteIfBound(/*bReverse*/false);
 	}
 }
 
@@ -703,8 +735,8 @@ void FRewindDebugger::Step(const int32 InNumberOfFrames)
 
 			if (NewScrubTime.IsSet())
 			{
-				SetCurrentScrubTime(NewScrubTime.GetValue());
-				TrackCursorDelegate.ExecuteIfBound(false);
+				SetCurrentScrubTime(NewScrubTime.GetValue(), EAutoScrollChange::Disable);
+				TrackCursorDelegate.ExecuteIfBound(/*bReverse*/false);
 				return;
 			}
 		}
@@ -723,9 +755,8 @@ void FRewindDebugger::Step(const int32 InNumberOfFrames)
 						ScrubTimeInformation.FrameIndex = FMath::Clamp<int64>(ScrubTimeInformation.FrameIndex + InNumberOfFrames, 0, (int64)EventCount - 1);
 						const FRecordingInfoMessage& Event = Recording->GetEvent(ScrubTimeInformation.FrameIndex);
 
-						SetCurrentScrubTime(Event.ElapsedTime);
-
-						TrackCursorDelegate.ExecuteIfBound(false);
+						SetCurrentScrubTime(Event.ElapsedTime, EAutoScrollChange::Disable);
+						TrackCursorDelegate.ExecuteIfBound(/*bReverse*/false);
 					}
 				}
 			}
@@ -743,12 +774,15 @@ void FRewindDebugger::StepBackward()
 	Step(-1);
 }
 
-void FRewindDebugger::ScrubToTime(double ScrubTime, bool bIsScrubbing)
+void FRewindDebugger::ScrubToTime(const double ScrubTime, bool bIsScrubbing)
 {
 	if (CanScrub())
 	{
 		Pause();
-		SetCurrentScrubTime(ScrubTime);
+		SetCurrentScrubTime(ScrubTime, EAutoScrollChange::Disable);
+
+		// No need to broadcast TrackCursorDelegate here since this is the callback
+		// handling the timeline cursor updates
 	}
 }
 
@@ -791,7 +825,7 @@ void FRewindDebugger::SetCurrentViewRange(const TRange<double>& Range)
 	}
 }
 
-void FRewindDebugger::SetCurrentScrubTime(double Time)
+void FRewindDebugger::SetCurrentScrubTime(const double Time, const EAutoScrollChange AutoScrollChange)
 {
 	CurrentScrubTime = Time;
 
@@ -801,6 +835,12 @@ void FRewindDebugger::SetCurrentScrubTime(double Time)
 
 		TraceTime.Set(ScrubTimeInformation.ProfileTime);
 	}
+
+	if (bAutoScrollToLastItem
+		&& AutoScrollChange == EAutoScrollChange::Disable)
+	{
+		bAutoScrollToLastItem = false;
+	}
 }
 
 void FRewindDebugger::GetScrubTimeInformation(double InDebugTime, FScrubTimeInformation& InOutTimeInformation, uint16 InRecordingIndex, const TraceServices::IAnalysisSession* AnalysisSession)
@@ -1172,16 +1212,16 @@ void FRewindDebugger::Tick(float DeltaTime)
 
 			RefreshDebugTracks();
 
-			const bool bRequiresScrubTimeUpdate = bPIESimulating || bNewDataAvailable;
-
-			if (bRequiresScrubTimeUpdate)
+			if (bAutoScrollToLastItem)
 			{
+				// Only need to update scrub time and notify
+				// when recording a PIE session or analysing data from a remote session or trace file.
 				if (IsRecording()
 					|| bNewDataAvailable)
 				{
 					TRACE_CPUPROFILER_EVENT_SCOPE(FRewindDebugger::Tick_UpdateSimulating);
-					SetCurrentScrubTime(RecordingDurationValue);
-					TrackCursorDelegate.ExecuteIfBound(false);
+					SetCurrentScrubTime(RecordingDuration.Get(), EAutoScrollChange::PreserveCurrentState);
+					TrackCursorDelegate.ExecuteIfBound(/*bReverse*/false);
 				}
 
 				// The debug position is only expected to be set outside of PIE
@@ -1196,8 +1236,8 @@ void FRewindDebugger::Tick(float DeltaTime)
 						const float PlaybackRate = URewindDebuggerSettings::Get().PlaybackRate;
 						TRACE_CPUPROFILER_EVENT_SCOPE(FRewindDebugger::Tick_UpdatePlayback);
 						const float Rate = PlaybackRate * (ControlState == EControlState::Play ? 1 : -1);
-						SetCurrentScrubTime(FMath::Clamp(CurrentScrubTime + Rate * DeltaTime, 0.0f, RecordingDuration.Get()));
-						TrackCursorDelegate.ExecuteIfBound(Rate < 0);
+						SetCurrentScrubTime(FMath::Clamp(CurrentScrubTime + Rate * DeltaTime, 0.0f, RecordingDuration.Get()), EAutoScrollChange::Disable);
+						TrackCursorDelegate.ExecuteIfBound(/*bReverse*/Rate < 0);
 
 						if (CurrentScrubTime == 0 || CurrentScrubTime == RecordingDuration.Get())
 						{
@@ -1205,9 +1245,6 @@ void FRewindDebugger::Tick(float DeltaTime)
 							ControlState = EControlState::Pause;
 						}
 					}
-
-					// update trace time
-					SetCurrentScrubTime(CurrentScrubTime);
 				}
 			}
 		}
@@ -1405,6 +1442,11 @@ void FRewindDebugger::MakeWorldsMenu(UToolMenu* Menu)
 		GameplayProvider->EnumerateWorlds([GameplayProvider, &GameWorldsSection, &OtherWorldsSection, &ServerWorldsSection](const FWorldInfo& WorldInfo)
 			{
 				const FObjectInfo* ObjectInfo = GameplayProvider->FindObjectInfo(WorldInfo.Id);
+				if (ObjectInfo == nullptr)
+				{
+					UE_LOG(LogRewindDebugger, Error, TEXT("Unable to find Object information associated to the provided World Id. Might be caused by missing data in the recording."));
+					return;
+				}
 				FString Name = ObjectInfo->Name;
 
 				FToolMenuSection* Section = &OtherWorldsSection;
@@ -1454,7 +1496,7 @@ void FRewindDebugger::RegisterToolBar()
 {
 	UToolMenu* Menu = UToolMenus::Get()->RegisterMenu(FRewindDebuggerModule::MainToolBarName, NAME_None, EMultiBoxType::ToolBar);
 
-	FToolMenuSection& Section = Menu->FindOrAddSection("VCRControls");
+	FToolMenuSection& Section = Menu->FindOrAddSection("PlaybackControls");
 	const FRewindDebuggerCommands& Commands = FRewindDebuggerCommands::Get();
 	Section.AddEntry(FToolMenuEntry::InitToolBarButton(
 		Commands.FirstFrame,
@@ -1657,14 +1699,14 @@ void FRewindDebugger::RegisterToolBar()
 		TAttribute<FText>(),
 		FSlateIcon(FAppStyle::GetAppStyleSetName(), "Icons.Delete")));
 
+	Section.AddSeparator(NAME_None);
+
 	Section.AddEntry(FToolMenuEntry::InitToolBarButton(
 		Commands.AutoEject,
 		FText(),
 		TAttribute<FText>(),
 		FSlateIcon("RewindDebuggerStyle", "RewindDebugger.AutoEject")));
 
-	Section.AddSeparator(NAME_None);
-
 	Section.AddEntry(FToolMenuEntry::InitToolBarButton(
 		Commands.AutoRecord,
 		FText(),
@@ -1684,10 +1726,21 @@ void FRewindDebugger::RegisterToolBar()
 		LOCTEXT("Display World Tooltip", "When loading trace files, only the objects (Such as Skeletal Meshes) from the world selected here will be spawned for preview")
 	));
 
-	Menu->SetStyleSet(&FAppStyle::Get());
-	Menu->StyleName = "PaletteToolBar";
-}
+	Menu->SetStyleSet(&FRewindDebuggerStyle::Get());
+	Menu->StyleName = "RewindDebugger.Toolbar";
 
+	// Right toolbar
+	UToolMenu* RightMenu = UToolMenus::Get()->RegisterMenu(FRewindDebuggerModule::RightToolBarName, NAME_None, EMultiBoxType::ToolBar);
+	FToolMenuSection& RightSection = RightMenu->FindOrAddSection("PlaybackControls");
+	RightSection.AddEntry(FToolMenuEntry::InitToolBarButton(
+		Commands.AutoScroll,
+		FText(),
+		TAttribute<FText>(),
+		FSlateIcon("RewindDebuggerStyle", "RewindDebugger.AutoScroll")));
+
+	RightMenu->SetStyleSet(&FRewindDebuggerStyle::Get());
+	RightMenu->StyleName = "RewindDebugger.Toolbar";
+}
 
 void FRewindDebugger::TrackDoubleClicked(TSharedPtr<RewindDebugger::FRewindDebuggerTrack> InSelectedTrack)
 {
@@ -1728,7 +1781,6 @@ TSharedPtr<SWidget> FRewindDebugger::BuildTrackContextMenu() const
 	return UToolMenus::Get()->GenerateWidget(FRewindDebuggerModule::TrackContextMenuName, FToolMenuContext(MenuContext));
 }
 
-
 TSharedPtr<FDebugObjectInfo> FRewindDebugger::GetSelectedObject() const
 {
 	if (SelectedTrack.IsValid())
diff --git a/Engine/Plugins/Animation/GameplayInsights/Source/RewindDebugger/Private/RewindDebugger.h b/Engine/Plugins/Animation/GameplayInsights/Source/RewindDebugger/Private/RewindDebugger.h
--- a/Engine/Plugins/Animation/GameplayInsights/Source/RewindDebugger/Private/RewindDebugger.h
+++ b/Engine/Plugins/Animation/GameplayInsights/Source/RewindDebugger/Private/RewindDebugger.h
@@ -1,4 +1,4 @@
-﻿// Copyright Epic Games, Inc. All Rights Reserved.
+// Copyright Epic Games, Inc. All Rights Reserved.
 #pragma once
 
 #if UE_ENABLE_INCLUDE_ORDER_DEPRECATED_IN_5_7
@@ -115,6 +115,9 @@ public:
 
 	/** Start a new Recording:  Start tracing Objects, increment the current recording index, and reset the recording elapsed time to 0 */
 	virtual void StartRecording() const override;
+
+	void OnRecordingStarted();
+	void OnRecordingStopped();
 	void OnClearRecording();
 
 	bool CanOpenTrace() const;
@@ -127,6 +130,10 @@ public:
 	void ClearTrace();
 	void ClearTraceAndActiveSession();
 
+	void ToggleAutoScroll();
+	bool CanEnableAutoScroll() const;
+	bool IsAutoScrollEnabled() const;
+
 	bool ShouldAutoRecordOnPIE() const;
 	void SetShouldAutoRecordOnPIE(bool InValue);
 
@@ -267,7 +274,13 @@ private:
 
 	FDelegateHandle ClearRecordingHandle;
 
-	void SetCurrentScrubTime(double Time);
+	enum class EAutoScrollChange : int8
+	{
+		PreserveCurrentState,
+		Disable
+	};
+
+	void SetCurrentScrubTime(double Time, EAutoScrollChange AutoScrollChange);
 
 	TBindableProperty<double> TraceTime;
 	TBindableProperty<double> RecordingDuration;
@@ -287,8 +300,10 @@ private:
 
 	bool bPIEStarted = false;
 	bool bPIESimulating = false;
+	bool bRecordingPIESession = false;
 
-	bool bRecording = false;
+	/** Instructs the debugger to always scrub to the last received event */
+	bool bAutoScrollToLastItem = false;
 
 	double PreviousTraceTime = -1;
 	double CurrentScrubTime = 0;
diff --git a/Engine/Plugins/Animation/GameplayInsights/Source/RewindDebugger/Private/RewindDebuggerCommands.h b/Engine/Plugins/Animation/GameplayInsights/Source/RewindDebugger/Private/RewindDebuggerCommands.h
--- a/Engine/Plugins/Animation/GameplayInsights/Source/RewindDebugger/Private/RewindDebuggerCommands.h
+++ b/Engine/Plugins/Animation/GameplayInsights/Source/RewindDebugger/Private/RewindDebuggerCommands.h
@@ -35,6 +35,7 @@ public:
 		UI_COMMAND(OpenTrace, "Open Recording", "Open a Trace file containing a rewind debugger recording", EUserInterfaceActionType::Button, FInputChord());
 		UI_COMMAND(AttachToSession, "Attach to Session", "Attach to a live session", EUserInterfaceActionType::Button, FInputChord());
 		UI_COMMAND(ClearTrace, "Clear Recording", "Clear current recording data", EUserInterfaceActionType::Button, FInputChord());
+		UI_COMMAND(AutoScroll, "AutoScroll", "Maintains the cursor on the last recorded event", EUserInterfaceActionType::ToggleButton, FInputChord());
 	}
 
 
@@ -53,6 +54,7 @@ public:
 	TSharedPtr<FUICommandInfo> OpenTrace;
 	TSharedPtr<FUICommandInfo> AttachToSession;
 	TSharedPtr<FUICommandInfo> ClearTrace;
+	TSharedPtr<FUICommandInfo> AutoScroll;
 };
 
 
diff --git a/Engine/Plugins/Animation/GameplayInsights/Source/RewindDebugger/Private/RewindDebuggerModule.cpp b/Engine/Plugins/Animation/GameplayInsights/Source/RewindDebugger/Private/RewindDebuggerModule.cpp
--- a/Engine/Plugins/Animation/GameplayInsights/Source/RewindDebugger/Private/RewindDebuggerModule.cpp
+++ b/Engine/Plugins/Animation/GameplayInsights/Source/RewindDebugger/Private/RewindDebuggerModule.cpp
@@ -30,6 +30,7 @@ const FName FRewindDebuggerModule::DetailsTabName("RewindDebuggerDetails2");
 const FName FRewindDebuggerModule::MainMenuName("RewindDebugger.MainMenu");
 const FName FRewindDebuggerModule::TrackContextMenuName("RewindDebugger.TrackContextMenu");
 const FName FRewindDebuggerModule::MainToolBarName = FName("RewindDebugger.ToolBar");
+const FName FRewindDebuggerModule::RightToolBarName = FName("RewindDebugger.RightToolBar");
 const FName FRewindDebuggerModule::MainStatusBarName = FName("RewindDebugger.StatusBar");
 
 #include "TraceServices/ModuleService.h"
@@ -187,6 +188,12 @@ TSharedRef<SDockTab> FRewindDebuggerModule::SpawnRewindDebuggerTab(const FSpawnT
 							FIsActionChecked(),
 							FIsActionButtonVisible());
 
+	CommandList->MapAction(Commands.AutoScroll,
+							FExecuteAction::CreateRaw(DebuggerInstance, &FRewindDebugger::ToggleAutoScroll),
+							FCanExecuteAction::CreateRaw(DebuggerInstance, &FRewindDebugger::CanEnableAutoScroll),
+							FIsActionChecked::CreateRaw(DebuggerInstance, &FRewindDebugger::IsAutoScrollEnabled),
+							FIsActionButtonVisible());
+
 	// Register PIE Rewind Debugger Commands
 	if (GEditor != nullptr)
 	{
diff --git a/Engine/Plugins/Animation/GameplayInsights/Source/RewindDebugger/Private/RewindDebuggerModule.h b/Engine/Plugins/Animation/GameplayInsights/Source/RewindDebugger/Private/RewindDebuggerModule.h
--- a/Engine/Plugins/Animation/GameplayInsights/Source/RewindDebugger/Private/RewindDebuggerModule.h
+++ b/Engine/Plugins/Animation/GameplayInsights/Source/RewindDebugger/Private/RewindDebuggerModule.h
@@ -31,6 +31,7 @@ public:
 
 	static const FName MainMenuName;
 	static const FName MainToolBarName;
+	static const FName RightToolBarName;
 	static const FName MainStatusBarName;
 	static const FName MainTabName;
 	static const FName DetailsTabName;
diff --git a/Engine/Plugins/Animation/GameplayInsights/Source/RewindDebugger/Private/RewindDebuggerRecordingControls.cpp b/Engine/Plugins/Animation/GameplayInsights/Source/RewindDebugger/Private/RewindDebuggerRecordingControls.cpp
--- a/Engine/Plugins/Animation/GameplayInsights/Source/RewindDebugger/Private/RewindDebuggerRecordingControls.cpp
+++ b/Engine/Plugins/Animation/GameplayInsights/Source/RewindDebugger/Private/RewindDebuggerRecordingControls.cpp
@@ -72,7 +72,8 @@ bool FDebuggerRecordingControls::ConnectToLiveSession_DirectInternal(const FGuid
 
 bool FDebuggerRecordingControls::IsRecordingToggleButtonEnabledInternal(const TSharedRef<TraceBasedDebuggers::FSessionInfo>& InSessionInfo) const
 {
-	return (InSessionInfo->BuildTargetType == EBuildTargetType::Editor && FRewindDebugger::Instance()->IsPIESimulating())
+	return (InSessionInfo->BuildTargetType == EBuildTargetType::Editor
+			&& (FRewindDebugger::Instance()->IsPIESimulating() || FRewindDebugger::Instance()->IsRecording()))
 		|| InSessionInfo->BuildTargetType != EBuildTargetType::Editor;
 }
 
@@ -117,4 +118,13 @@ const TraceBasedDebuggers::FTraceSessionDescriptor& FDebuggerRecordingControls::
 #endif
 }
 
+void FDebuggerRecordingControls::OnRecordingStopped()
+{
+	FRewindDebugger::Instance()->OnRecordingStopped();
+}
+
+void FDebuggerRecordingControls::OnRecordingStarted()
+{
+	FRewindDebugger::Instance()->OnRecordingStarted();
+}
 } // UE::RewindDebugger
\ No newline at end of file
diff --git a/Engine/Plugins/Animation/GameplayInsights/Source/RewindDebugger/Private/RewindDebuggerRecordingControls.h b/Engine/Plugins/Animation/GameplayInsights/Source/RewindDebugger/Private/RewindDebuggerRecordingControls.h
--- a/Engine/Plugins/Animation/GameplayInsights/Source/RewindDebugger/Private/RewindDebuggerRecordingControls.h
+++ b/Engine/Plugins/Animation/GameplayInsights/Source/RewindDebugger/Private/RewindDebuggerRecordingControls.h
@@ -33,6 +33,9 @@ protected:
 	virtual void OnToggleRecordingStateInternal(const TSharedRef<TraceBasedDebuggers::FSessionInfo>& InSessionInfo) override;
 	virtual void CloseSessionByRemoteSessionIDInternal(const FGuid& InSessionId) override;
 	virtual const TraceBasedDebuggers::FTraceSessionDescriptor& GetCurrentSessionDescriptorInternal() override;
+
+	virtual void OnRecordingStopped() override;
+	virtual void OnRecordingStarted() override;
 };
 
 }
diff --git a/Engine/Plugins/Animation/GameplayInsights/Source/RewindDebugger/Private/RewindDebuggerStyle.cpp b/Engine/Plugins/Animation/GameplayInsights/Source/RewindDebugger/Private/RewindDebuggerStyle.cpp
--- a/Engine/Plugins/Animation/GameplayInsights/Source/RewindDebugger/Private/RewindDebuggerStyle.cpp
+++ b/Engine/Plugins/Animation/GameplayInsights/Source/RewindDebugger/Private/RewindDebuggerStyle.cpp
@@ -7,6 +7,7 @@
 #include "Styling/SlateStyleMacros.h"
 #include "Styling/SlateStyleRegistry.h"
 #include "Styling/SlateTypes.h"
+#include "Styling/ToolBarStyle.h"
 
 TSharedPtr< FRewindDebuggerStyle > FRewindDebuggerStyle::StyleInstance = nullptr;
 
@@ -32,7 +33,8 @@ FRewindDebuggerStyle::FRewindDebuggerStyle() :
 	Set("RewindDebugger.Play.small", new CORE_IMAGE_BRUSH("Editor/Slate/Sequencer/Transport_Bar/Play_24x", Icon24x24));
 	Set("RewindDebugger.NextFrame.small", new CORE_IMAGE_BRUSH("Editor/Slate/Sequencer/Transport_Bar/Step_Forward_24x", Icon24x24));
 	Set("RewindDebugger.LastFrame.small", new CORE_IMAGE_BRUSH("Editor/Slate/Sequencer/Transport_Bar/Go_To_End_24x", Icon24x24));
-	
+	Set("RewindDebugger.AutoScroll", new CORE_IMAGE_BRUSH_SVG("Slate/Starship/Insights/AutoScrollRight_20", Icon24x24));
+
 	Set("RewindDebugger.AutoEject", new IMAGE_BRUSH("autoeject", {30.f, 24.f}));
 	Set("RewindDebugger.AutoRecord", new IMAGE_BRUSH("autorecord", {30.f, 24.f}));
 
@@ -49,6 +51,24 @@ FRewindDebuggerStyle::FRewindDebuggerStyle() :
 
 	FTableRowStyle TableRowStyle = FTableRowStyle(FAppStyle::Get().GetWidgetStyle<FTableRowStyle>("SceneOutliner.TableViewRow"));
 	Set("RewindDebugger.TableRow", TableRowStyle);
+
+	// Toolbar style
+	{
+		FToolBarStyle CompactToolbarStyle = FAppStyle::Get().GetWidgetStyle<FToolBarStyle>("SlimToolBar");
+
+		// Reduce button padding
+		FButtonStyle CompactButton = CompactToolbarStyle.ButtonStyle;
+		CompactButton.SetNormalPadding(FMargin(2.f, 4.f, 2.f, 4.f)); // Changed from (8.f, 4.f, 8.f, 4.f)
+		CompactButton.SetPressedPadding(FMargin(2.f, 5.f, 2.f, 3.f)); // Changed from (8.f, 5.f, 8.f, 3.f)
+		CompactToolbarStyle.SetButtonStyle(CompactButton);
+
+		// Reduce toggle button padding
+		FCheckBoxStyle CompactToggle = CompactToolbarStyle.ToggleButton;
+		CompactToggle.SetPadding(FMargin(2.f, 4.f, 2.f, 4.f)); // Changed from (8.f, 4.f, 8.f, 4.f)
+		CompactToolbarStyle.SetToggleButtonStyle(CompactToggle);
+
+		Set("RewindDebugger.Toolbar", CompactToolbarStyle);
+	}
 }
 
 void FRewindDebuggerStyle::Initialize()
diff --git a/Engine/Plugins/Animation/GameplayInsights/Source/RewindDebugger/Private/SRewindDebugger.cpp b/Engine/Plugins/Animation/GameplayInsights/Source/RewindDebugger/Private/SRewindDebugger.cpp
--- a/Engine/Plugins/Animation/GameplayInsights/Source/RewindDebugger/Private/SRewindDebugger.cpp
+++ b/Engine/Plugins/Animation/GameplayInsights/Source/RewindDebugger/Private/SRewindDebugger.cpp
@@ -384,6 +384,7 @@ void SRewindDebugger::Construct(const FArguments& InArgs, TSharedRef<FUICommandL
 #endif // WITH_TRACE_BASED_DEBUGGERS
 
 	const TSharedRef<SWidget> ToolBar = UToolMenus::Get()->GenerateWidget(FRewindDebuggerModule::MainToolBarName, ToolMenuContext);
+	const TSharedRef<SWidget> RightToolBar = UToolMenus::Get()->GenerateWidget(FRewindDebuggerModule::RightToolBarName, ToolMenuContext);
 
 	ChildSlot
 		[
@@ -394,7 +395,8 @@ void SRewindDebugger::Construct(const FArguments& InArgs, TSharedRef<FUICommandL
 						+ SVerticalBox::Slot().MaxHeight(48)
 						[
 							SNew(SHorizontalBox)
-								+ SHorizontalBox::Slot().AutoWidth()
+								+ SHorizontalBox::Slot()
+								.AutoWidth()
 								[
 									SNew(SComboButton)
 										.ComboButtonStyle(FAppStyle::Get(), "SimpleComboButton")
@@ -405,10 +407,18 @@ void SRewindDebugger::Construct(const FArguments& InArgs, TSharedRef<FUICommandL
 												.Image(FAppStyle::Get().GetBrush("ClassIcon.CameraComponent"))
 										]
 								]
-								+ SHorizontalBox::Slot().FillWidth(1.0)
+								+ SHorizontalBox::Slot()
+								.HAlign(HAlign_Left)
+								.AutoWidth()
 								[
 									ToolBar
 								]
+								+ SHorizontalBox::Slot()
+								.HAlign(HAlign_Right)
+								.FillContentWidth(1.f)
+								[
+									RightToolBar
+								]
 						]
 				]
 				+ SVerticalBox::Slot().FillHeight(1.0)
@@ -546,7 +556,7 @@ void SRewindDebugger::Construct(const FArguments& InArgs, TSharedRef<FUICommandL
 		];
 }
 
-bool FilterTrack(TSharedPtr<RewindDebugger::FRewindDebuggerTrack>& Track, const FString& FilterString, bool bRemoveNoData, const TArray<FName>& FilteredTrackTypes, bool bParentFilterPassed = false)
+bool FilterTrack(const TSharedPtr<RewindDebugger::FRewindDebuggerTrack>& Track, const FString& FilterString, bool bRemoveNoData, const TArray<FName>& FilteredTrackTypes, bool bParentFilterPassed = false)
 {
 	using namespace RewindDebugger;
 	if (FilteredTrackTypes.Contains(Track->GetName()))
@@ -567,7 +577,7 @@ bool FilterTrack(TSharedPtr<RewindDebugger::FRewindDebuggerTrack>& Track, const
 	const bool bThisFilterPassed = (!bStringFilterEmpty && bStringFilterPassed);
 
 	bool bAnyChildVisible = false;
-	Track->IterateSubTracks([&bAnyChildVisible, bThisFilterPassed, FilterString, bRemoveNoData, FilteredTrackTypes](TSharedPtr<FRewindDebuggerTrack> ChildTrack)
+	Track->IterateSubTracks([&bAnyChildVisible, bThisFilterPassed, FilterString, bRemoveNoData, FilteredTrackTypes](const TSharedPtr<FRewindDebuggerTrack>& ChildTrack)
 		{
 			const bool bChildIsVisible = FilterTrack(ChildTrack, FilterString, bRemoveNoData, FilteredTrackTypes, bThisFilterPassed);
 			bAnyChildVisible |= bChildIsVisible;
diff --git a/Engine/Source/Runtime/Experimental/TraceBasedDebuggersAnalysis/Private/Widgets/RecordingControls.cpp b/Engine/Source/Runtime/Experimental/TraceBasedDebuggersAnalysis/Private/Widgets/RecordingControls.cpp
--- a/Engine/Source/Runtime/Experimental/TraceBasedDebuggersAnalysis/Private/Widgets/RecordingControls.cpp
+++ b/Engine/Source/Runtime/Experimental/TraceBasedDebuggersAnalysis/Private/Widgets/RecordingControls.cpp
@@ -1,4 +1,4 @@
-﻿// Copyright Epic Games, Inc. All Rights Reserved.
+// Copyright Epic Games, Inc. All Rights Reserved.
 
 #if WITH_EDITOR
 
@@ -283,6 +283,9 @@ void FRecordingControls::HandleRecordingStop(const TSharedPtr<FSessionInfo>& InS
 		const FText LiveSessionEnded = LOCTEXT("RecordingSessionEndedMessage", "Recording session has ended");
 		LiveSessionEndedMessageHandle = StatusBarSubsystem->PushStatusBarMessage(StatusBarID, LiveSessionEnded);
 	}
+
+	// Allow derived classes to perform extra operations
+	OnRecordingStopped();
 }
 
 void FRecordingControls::HandleRecordingStart(const TSharedPtr<FSessionInfo>& InSessionInfo)
@@ -312,6 +315,9 @@ void FRecordingControls::HandleRecordingStart(const TSharedPtr<FSessionInfo>& In
 	}
 
 	RecordingMessageHandle = StatusBarSubsystem->PushStatusBarMessage(StatusBarID, LOCTEXT("RecordingMessage", "Recording..."));
+
+	// Allow derived classes to perform extra operations
+	OnRecordingStarted();
 }
 
 void FRecordingControls::ExecuteAsyncConnectionAttemptTaskWithRetry(
diff --git a/Engine/Source/Runtime/Experimental/TraceBasedDebuggersAnalysis/Public/Widgets/RecordingControls.h b/Engine/Source/Runtime/Experimental/TraceBasedDebuggersAnalysis/Public/Widgets/RecordingControls.h
--- a/Engine/Source/Runtime/Experimental/TraceBasedDebuggersAnalysis/Public/Widgets/RecordingControls.h
+++ b/Engine/Source/Runtime/Experimental/TraceBasedDebuggersAnalysis/Public/Widgets/RecordingControls.h
@@ -92,6 +92,14 @@ protected:
 
 	virtual const FTraceSessionDescriptor& GetCurrentSessionDescriptorInternal() = 0;
 
+	virtual void OnRecordingStopped()
+	{
+	}
+
+	virtual void OnRecordingStarted()
+	{
+	}
+
 	UE_API void SelectTargetSession(FGuid SessionId);
 	UE_API TSharedPtr<FSessionInfo> GetCurrentSessionInfo() const;
 
