[RewindDebugger] detect recording started while PIE is running and enable auto-scroll
#Synced-CL 49574741

diff --git a/Engine/Plugins/Animation/GameplayInsights/Source/RewindDebugger/Private/RewindDebugger.cpp b/Engine/Plugins/Animation/GameplayInsights/Source/RewindDebugger/Private/RewindDebugger.cpp
--- a/Engine/Plugins/Animation/GameplayInsights/Source/RewindDebugger/Private/RewindDebugger.cpp
+++ b/Engine/Plugins/Animation/GameplayInsights/Source/RewindDebugger/Private/RewindDebugger.cpp
@@ -417,6 +417,12 @@ void FRewindDebugger::StartRecording() const
 
 void FRewindDebugger::OnRecordingStarted()
 {
+	if (RecordingControls
+		&& RecordingControls->IsRecordingLocalSession())
+	{
+		bRecordingPIESession = true;
+		bAutoScrollToLastItem = true;
+	}
 }
 
 void FRewindDebugger::OnRecordingStopped()
diff --git a/Engine/Plugins/Animation/GameplayInsights/Source/RewindDebugger/Private/RewindDebuggerRecordingControls.cpp b/Engine/Plugins/Animation/GameplayInsights/Source/RewindDebugger/Private/RewindDebuggerRecordingControls.cpp
--- a/Engine/Plugins/Animation/GameplayInsights/Source/RewindDebugger/Private/RewindDebuggerRecordingControls.cpp
+++ b/Engine/Plugins/Animation/GameplayInsights/Source/RewindDebugger/Private/RewindDebuggerRecordingControls.cpp
@@ -26,6 +26,18 @@ FDebuggerRecordingControls::FDebuggerRecordingControls(const FName& InStatusBarI
 #endif
 }
 
+bool FDebuggerRecordingControls::IsRecordingLocalSession() const
+{
+	using namespace UE::TraceBasedDebuggers;
+	if (const TSharedPtr<FSessionInfo> SessionInfo = GetCurrentSessionInfo())
+	{
+		return (SessionInfo->SessionName == FRemoteSessionsManager::LocalEditorSessionName
+			&& SessionInfo->IsRecording(DebuggerGuid));
+	}
+
+	return false;
+}
+
 bool FDebuggerRecordingControls::CanStartRecording() const
 {
 	return !IsRecording() && IsRecordingToggleButtonEnabled();
diff --git a/Engine/Plugins/Animation/GameplayInsights/Source/RewindDebugger/Private/RewindDebuggerRecordingControls.h b/Engine/Plugins/Animation/GameplayInsights/Source/RewindDebugger/Private/RewindDebuggerRecordingControls.h
--- a/Engine/Plugins/Animation/GameplayInsights/Source/RewindDebugger/Private/RewindDebuggerRecordingControls.h
+++ b/Engine/Plugins/Animation/GameplayInsights/Source/RewindDebugger/Private/RewindDebuggerRecordingControls.h
@@ -15,6 +15,7 @@ struct FDebuggerRecordingControls : TraceBasedDebuggers::FRecordingControls
 		, const FLogCategoryBase& InLogCategory);
 
 	using FRecordingControls::IsRecording;
+	bool IsRecordingLocalSession() const;
 
 	bool CanStartRecording() const;
 	void StartRecording();
