[RewindDebugger] added explicit methods to handle the local session recording and made sure that this functionality doesn't interfere with the other active sessions.
- when auto-record on PIE is used we make sure to select the local session in case multiple remote sessions are available. Same when resuming or advancing one frame in PIE
- fixed issue where users are not able to start recording while PIE is paused

#Synced-CL 49721519

diff --git a/Engine/Plugins/Animation/GameplayInsights/Source/RewindDebugger/Private/RewindDebugger.cpp b/Engine/Plugins/Animation/GameplayInsights/Source/RewindDebugger/Private/RewindDebugger.cpp
--- a/Engine/Plugins/Animation/GameplayInsights/Source/RewindDebugger/Private/RewindDebugger.cpp
+++ b/Engine/Plugins/Animation/GameplayInsights/Source/RewindDebugger/Private/RewindDebugger.cpp
@@ -1,4 +1,4 @@
-// Copyright Epic Games, Inc. All Rights Reserved.
+﻿// Copyright Epic Games, Inc. All Rights Reserved.
 
 #include "RewindDebugger.h"
 
@@ -36,6 +36,10 @@
 
 #define LOCTEXT_NAMESPACE "RewindDebugger"
 
+namespace UE::RewindDebugger
+{
+constexpr uint32 NumDefaultDebugTracks = 2;
+
 static void IterateExtensions(TFunction<void(IRewindDebuggerExtension* Extension)> IteratorFunction)
 {
 	// update extensions
@@ -58,6 +62,7 @@ static void TraceSubobjects(const UObject* OuterObject)
 		TRACE_OBJECT_LIFETIME_BEGIN(Subobject);
 	}
 }
+} // UE::RewindDebugger
 
 FRewindDebugger::FRewindDebugger()
 {
@@ -101,7 +106,6 @@ FRewindDebugger::FRewindDebugger()
 				Settings.SaveConfig();
 			}
 
-			CandidateIds.SetNum(0);
 			GetTargetObjectIds(CandidateIds);
 			// make sure all the SubObjects of the root object have been traced
 #if OBJECT_TRACE_ENABLED
@@ -109,7 +113,7 @@ FRewindDebugger::FRewindDebugger()
 			{
 				if (const UObject* TargetObject = FObjectTrace::GetObjectFromId(TargetObjectId.GetMainId()))
 				{
-					TraceSubobjects(TargetObject);
+					UE::RewindDebugger::TraceSubobjects(TargetObject);
 				}
 			}
 #endif
@@ -137,6 +141,10 @@ FRewindDebugger::FRewindDebugger()
 			, LogRewindDebugger);
 		RecordingControls->Initialize();
 		RecordingControls->AddToMenu(FRewindDebuggerModule::MainToolBarName);
+
+		RecordingControls->OnRecordingStarted.BindRaw(this, &FRewindDebugger::HandleRecordingStarted);
+		RecordingControls->OnRecordingStopped.BindRaw(this, &FRewindDebugger::HandleRecordingStopped);
+		RecordingControls->OnSessionSelected.BindRaw(this, &FRewindDebugger::HandleSessionSelected);
 	}
 #endif // WITH_TRACE_BASED_DEBUGGERS
 }
@@ -156,6 +164,29 @@ FRewindDebugger::~FRewindDebugger()
 	{
 		Runtime->ClearRecording.Remove(ClearRecordingHandle);
 	}
+
+	if (RecordingControls)
+	{
+		RecordingControls->OnRecordingStarted.Unbind();
+		RecordingControls->OnRecordingStopped.Unbind();
+		RecordingControls->OnSessionSelected.Unbind();
+	}
+
+	// Close any active analysis sessions
+	using namespace UE::TraceBasedDebuggers;
+	if (ActiveAnalysisDescriptors.Num())
+	{
+		if (const TSharedPtr<FTraceSessionsManager> TraceManager = FRewindDebuggerModule::Get().GetTraceSessionsManager())
+		{
+			for (const FTraceSessionDescriptor& ActiveAnalysisDescriptor : ActiveAnalysisDescriptors)
+			{
+				if (ActiveAnalysisDescriptor.IsValid())
+				{
+					TraceManager->CloseSession(ActiveAnalysisDescriptor.SessionName);
+				}
+			}
+		}
+	}
 #endif // WITH_TRACE_BASED_DEBUGGERS
 }
 
@@ -194,12 +225,7 @@ void FRewindDebugger::OnPIEStarted(bool bSimulating)
 
 	if (ShouldAutoRecordOnPIE())
 	{
-		StartRecording();
-		bRecordingPIESession = true;
-
-		// Reactivate auto-scroll assuming user wants to visualize the newest data until
-		// he manually scrubs or uses the playback controls.
-		bAutoScrollToLastItem = true;
+		StartRecordingLocalSession();
 	}
 }
 
@@ -227,8 +253,7 @@ void FRewindDebugger::OnPIEResumed(bool bSimulating)
 {
 	bPIESimulating = true;
 
-	if (IsRecording()
-		&& bRecordingPIESession)
+	if (IsRecordingLocalSession())
 	{
 		// Reactivate auto-scroll assuming user is done inspecting recorded data
 		bAutoScrollToLastItem = true;
@@ -242,8 +267,7 @@ void FRewindDebugger::OnPIEResumed(bool bSimulating)
 
 void FRewindDebugger::OnPIESingleStepped(bool bSimulating)
 {
-	if (IsRecording()
-		&& bRecordingPIESession)
+	if (IsRecordingLocalSession())
 	{
 		// Reactivate auto-scroll assuming user wants to visualize data for the next frame
 		bAutoScrollToLastItem = true;
@@ -252,10 +276,9 @@ void FRewindDebugger::OnPIESingleStepped(bool bSimulating)
 
 void FRewindDebugger::OnPIEStopped(bool bSimulating)
 {
-	if (IsRecording()
-		&& bRecordingPIESession)
+	if (IsRecordingLocalSession())
 	{
-		StopRecording();
+		StopRecordingLocalSession();
 	}
 
 	bPIEStarted = false;
@@ -276,7 +299,8 @@ void FRewindDebugger::SetRootObjectPosition(const TOptional<FVector>& InPosition
 
 void FRewindDebugger::GetTargetObjectIds(TArray<RewindDebugger::FObjectId>& OutTargetObjectIds) const
 {
-	OutTargetObjectIds.Empty(2);
+	// We are looking for the object ID matching RootObjectName and currently supporting a single root track.
+	OutTargetObjectIds.Empty(1);
 
 	if (RootObjectName.Get() == "")
 	{
@@ -301,13 +325,13 @@ void FRewindDebugger::GetTargetObjectIds(TArray<RewindDebugger::FObjectId>& OutT
 
 	// make sure all the SubObjects of the root object have been traced
 #if OBJECT_TRACE_ENABLED
-	if (IsRecording())
+	if (IsRecordingLocalSession())
 	{
 		for (const RewindDebugger::FObjectId& CandidateId : CandidateIds)
 		{
 			if (const UObject* TargetObject = FObjectTrace::GetObjectFromId(CandidateId.GetMainId()))
 			{
-				TraceSubobjects(TargetObject);
+				UE::RewindDebugger::TraceSubobjects(TargetObject);
 			}
 		}
 	}
@@ -330,11 +354,11 @@ void FRewindDebugger::RefreshDebugTracks()
 	{
 		// fallback code path for when the target object is not found
 
-		if (DebugTracks.Num() != 2)
+		if (DebugTracks.Num() != UE::RewindDebugger::NumDefaultDebugTracks)
 		{
 			// clear tracks so we don't show data from previous recordings
 			DebugTracks.SetNum(0);
-			DebugTracks.SetNum(2);
+			DebugTracks.SetNum(UE::RewindDebugger::NumDefaultDebugTracks);
 		}
 
 		if (DebugTracks[1] == nullptr || DebugTracks[0] == nullptr || DebugTracks[0]->GetName().ToString() != SelectionName)
@@ -392,56 +416,99 @@ void FRewindDebugger::RefreshDebugTracks()
 
 bool FRewindDebugger::CanStartRecording() const
 {
-	// Relying on the recording controls since they rely on the
-	// sessions manager which also handles active sessions from other debuggers.
-	if (RecordingControls)
-	{
-		return RecordingControls->CanStartRecording();
-	}
-
-	return false;
+	// Preserving legacy behavior of IRewindDebugger interface by monitoring the local session only.
+	return CanStartRecordingLocalSession();
 }
 
 void FRewindDebugger::StartRecording() const
 {
-	if (!CanStartRecording())
+	// Preserving legacy behavior of IRewindDebugger interface by recording the local session only.
+	StartRecordingLocalSession();
+}
+
+bool FRewindDebugger::CanStartRecordingLocalSession() const
+{
+	// Relying on the recording controls using the
+	// sessions manager which also handles active sessions from other debuggers.
+	return RecordingControls
+		&& RecordingControls->CanStartRecordingLocalSession();
+}
+
+void FRewindDebugger::StartRecordingLocalSession() const
+{
+	if (!CanStartRecordingLocalSession())
 	{
 		return;
 	}
 
 	if (RecordingControls)
 	{
-		RecordingControls->StartRecording();
+		RecordingControls->StartRecordingLocalSession();
 	}
 }
 
-void FRewindDebugger::OnRecordingStarted()
+#if UE_WITH_TRACE_BASED_DEBUGGERS_ANALYSIS
+void FRewindDebugger::HandleRecordingStarted(const TSharedPtr<UE::TraceBasedDebuggers::FSessionInfo>& InSessionInfo)
 {
 	if (RecordingControls
 		&& RecordingControls->IsRecordingLocalSession())
 	{
-		bRecordingPIESession = true;
 		bAutoScrollToLastItem = true;
 	}
 }
 
-void FRewindDebugger::OnRecordingStopped()
+void FRewindDebugger::HandleRecordingStopped(const TSharedPtr<UE::TraceBasedDebuggers::FSessionInfo>& InSessionInfo)
 {
-	if (bRecordingPIESession)
+	using namespace UE::TraceBasedDebuggers;
+	const FTraceSessionDescriptor* ExistingDescriptor = ActiveAnalysisDescriptors.FindByPredicate(
+		[SessionInfo = InSessionInfo](const FTraceSessionDescriptor& Descriptor)
+		{
+			return Descriptor.RemoteSessionID == SessionInfo->InstanceId;
+		});
+
+	if (ExistingDescriptor)
 	{
-		bRecordingPIESession = false;
-		StopActiveTraceSession();
+		StopAnalysisSession(*ExistingDescriptor);
 	}
 }
 
-void FRewindDebugger::StopRecording() const
+void FRewindDebugger::HandleSessionSelected(const TSharedPtr<UE::TraceBasedDebuggers::FSessionInfo>& InSessionInfo)
+{
+	using namespace UE::TraceBasedDebuggers;
+
+	const FTraceSessionDescriptor* ExistingDescriptor = ActiveAnalysisDescriptors.FindByPredicate(
+		[SessionInfo = InSessionInfo](const FTraceSessionDescriptor& Descriptor)
+		{
+			return Descriptor.RemoteSessionID == SessionInfo->InstanceId;
+		});
+
+	ClearTrackData();
+
+	if (ExistingDescriptor)
+	{
+		SetSelectedAnalysisSession(*ExistingDescriptor);
+	}
+	else
+	{
+		SetSelectedAnalysisSession(FTraceSessionDescriptor{});
+	}
+}
+#endif // UE_WITH_TRACE_BASED_DEBUGGERS_ANALYSIS
+
+void FRewindDebugger::StopRecordingLocalSession() const
 {
 	if (RecordingControls)
 	{
-		RecordingControls->StopRecording();
+		RecordingControls->StopRecordingLocalSession();
 	}
 }
 
+bool FRewindDebugger::IsRecordingLocalSession() const
+{
+	return RecordingControls
+		&& RecordingControls->IsRecordingLocalSession();
+}
+
 #if WITH_TRACE_BASED_DEBUGGERS
 void FRewindDebugger::StopLocalRecording() const
 {
@@ -453,28 +520,14 @@ void FRewindDebugger::StopLocalRecording() const
 }
 #endif // WITH_TRACE_BASED_DEBUGGERS
 
-void FRewindDebugger::OnClearRecording()
-{
-	ClearTrace();
-	RecordingDuration.Set(0);
-	CandidateIds.Empty(2);
-	RootObjectPosition.Reset();
-
-	IterateExtensions([this](IRewindDebuggerExtension* Extension)
-		{
-			Extension->Clear(this);
-		}
-	);
-}
-
 bool FRewindDebugger::CanOpenTrace() const
 {
-	return !IsRecording();
+	return !IsRecordingLocalSession();
 }
 
 void FRewindDebugger::OpenTrace(const FString& FilePath)
 {
-	ClearTraceAndActiveSession();
+	ClearAnalysis();
 
 	bDisplayWorldIdValid = false;
 
@@ -488,7 +541,7 @@ void FRewindDebugger::OpenTrace(const FString& FilePath)
 
 		if (NewSessionFromFileDescriptor.IsValid())
 		{
-			OpenSession(NewSessionFromFileDescriptor);
+			OpenAnalysisSession(NewSessionFromFileDescriptor);
 		}
 	}
 #endif // UE_WITH_TRACE_BASED_DEBUGGERS_ANALYSIS
@@ -533,7 +586,7 @@ void FRewindDebugger::AttachToSession()
 
 	if (SessionBrowserModal->ShowModal() != EAppReturnType::Cancel)
 	{
-		ClearTrace();
+		ClearTrackData();
 
 		bool bSuccess = false;
 		const SModalSessionBrowser::FTraceSessionInfo SessionInfo = SessionBrowserModal->GetSelectedTraceInfo();
@@ -549,17 +602,19 @@ void FRewindDebugger::AttachToSession()
 	}
 }
 
-bool FRewindDebugger::CanClearTrace() const
+bool FRewindDebugger::CanClearAnalysis() const
 {
-	return GetAnalysisSession() != nullptr;
+	return GetAnalysisSession() != nullptr || RecordingDuration.Get() > 0.f;
 }
 
-void FRewindDebugger::ClearTrace()
+void FRewindDebugger::ClearAnalysis()
 {
-#if WITH_TRACE_BASED_DEBUGGERS
-	StopLocalRecording();
-#endif
+	ClearActiveAnalysisSession();
+	ClearTrackData();
+}
 
+void FRewindDebugger::ClearTrackData()
+{
 	RecordingDuration.Set(0);
 
 	CandidateIds.Empty();
@@ -571,7 +626,7 @@ void FRewindDebugger::ClearTrace()
 	TrackSelectionChanged(nullptr);
 
 	// update extensions
-	IterateExtensions([this](IRewindDebuggerExtension* Extension)
+	UE::RewindDebugger::IterateExtensions([this](IRewindDebuggerExtension* Extension)
 		{
 			Extension->Clear(this);
 		}
@@ -580,10 +635,11 @@ void FRewindDebugger::ClearTrace()
 	RefreshDebugTracks();
 }
 
-void FRewindDebugger::ClearTraceAndActiveSession()
+void FRewindDebugger::OnClearRecording()
 {
-	CloseActiveTraceSession();
-	ClearTrace();
+	RootObjectPosition.Reset();
+
+	ClearTrackData();
 }
 
 void FRewindDebugger::ToggleAutoScroll()
@@ -810,12 +866,28 @@ UWorld* FRewindDebugger::GetWorldToVisualize() const
 
 bool FRewindDebugger::IsRecording() const
 {
-	return RecordingControls ? RecordingControls->IsRecording() : false;
+	// Preserving legacy behavior of IRewindDebugger interface by monitoring the local session only.
+	return RecordingControls ? RecordingControls->IsRecordingLocalSession() : false;
+}
+
+bool FRewindDebugger::IsAnalyzingLocalSession() const
+{
+	return SelectedDescriptor.IsValid()
+		&& SelectedDescriptor.RemoteSessionID.IsValid()
+		&& SelectedDescriptor.RemoteSessionID == UE::TraceBasedDebuggers::FRemoteSessionsManager::LocalEditorSessionID;
+}
+
+bool FRewindDebugger::IsAnalyzingRemoteSession() const
+{
+	return SelectedDescriptor.IsValid()
+		&& SelectedDescriptor.RemoteSessionID.IsValid()
+		&& SelectedDescriptor.RemoteSessionID != UE::TraceBasedDebuggers::FRemoteSessionsManager::LocalEditorSessionID;
 }
 
-bool FRewindDebugger::IsTraceFileLoaded() const
+bool FRewindDebugger::IsAnalyzingTraceFile() const
 {
-	return GetAnalysisSession() != nullptr && !bPIEStarted;
+	return SelectedDescriptor.IsValid()
+		&& !SelectedDescriptor.RemoteSessionID.IsValid();
 }
 
 void FRewindDebugger::SetCurrentViewRange(const TRange<double>& Range)
@@ -986,10 +1058,13 @@ TSharedPtr<const TraceServices::IAnalysisSession> FRewindDebugger::GetAnalysisSe
 	TSharedPtr<const TraceServices::IAnalysisSession> Session = nullptr;
 
 #if UE_WITH_TRACE_BASED_DEBUGGERS_ANALYSIS
-	using namespace UE::TraceBasedDebuggers;
-	if (const TSharedPtr<FTraceSessionsManager> TraceManager = FRewindDebuggerModule::Get().GetTraceSessionsManager())
+	if (SelectedDescriptor.IsValid())
 	{
-		Session = TraceManager->GetSession(CurrentSessionDescriptor.SessionName);
+		using namespace UE::TraceBasedDebuggers;
+		if (const TSharedPtr<FTraceSessionsManager> TraceManager = FRewindDebuggerModule::Get().GetTraceSessionsManager())
+		{
+			Session = TraceManager->GetSession(SelectedDescriptor.SessionName);
+		}
 	}
 #endif // UE_WITH_TRACE_BASED_DEBUGGERS_ANALYSIS
 
@@ -1027,8 +1102,6 @@ const FObjectInfo* FRewindDebugger::FindTypedOuterInfo(TNotNull<const UStruct*>
 
 bool FRewindDebugger::ConnectToLiveSession(const uint32 SessionID, const FStringView InSessionAddress)
 {
-	CloseActiveTraceSession();
-
 	using namespace UE::TraceBasedDebuggers;
 	if (const TSharedPtr<FTraceSessionsManager> TraceSessionsManager = FRewindDebuggerModule::Get().GetTraceSessionsManager())
 	{
@@ -1038,7 +1111,7 @@ bool FRewindDebugger::ConnectToLiveSession(const uint32 SessionID, const FString
 
 		if (NewSessionDescriptor.IsValid())
 		{
-			OpenSession(NewSessionDescriptor);
+			OpenAnalysisSession(NewSessionDescriptor);
 			return true;
 		}
 	}
@@ -1046,21 +1119,19 @@ bool FRewindDebugger::ConnectToLiveSession(const uint32 SessionID, const FString
 	return false;
 }
 
-bool FRewindDebugger::ConnectToLiveSession_Direct(const FGuid& RemoteSessionID)
+bool FRewindDebugger::ConnectToLiveSession_Direct(const FGuid& InRemoteSessionID)
 {
-	CloseActiveTraceSession();
-
 	using namespace UE::TraceBasedDebuggers;
 	if (const TSharedPtr<FTraceSessionsManager> TraceSessionsManager = FRewindDebuggerModule::Get().GetTraceSessionsManager())
 	{
 		FTraceSessionDescriptor NewSessionDescriptor;
-		NewSessionDescriptor.RemoteSessionID = RemoteSessionID;
+		NewSessionDescriptor.RemoteSessionID = InRemoteSessionID;
 		NewSessionDescriptor.bIsLiveSession = true;
-		NewSessionDescriptor.SessionName = TraceSessionsManager->ConnectToLiveSession_Direct(RemoteSessionID, NewSessionDescriptor.SessionPort);
+		NewSessionDescriptor.SessionName = TraceSessionsManager->ConnectToLiveSession_Direct(InRemoteSessionID, NewSessionDescriptor.SessionPort);
 
 		if (NewSessionDescriptor.IsValid())
 		{
-			OpenSession(NewSessionDescriptor);
+			OpenAnalysisSession(NewSessionDescriptor);
 			return true;
 		}
 	}
@@ -1068,21 +1139,19 @@ bool FRewindDebugger::ConnectToLiveSession_Direct(const FGuid& RemoteSessionID)
 	return false;
 }
 
-bool FRewindDebugger::ConnectToLiveSession_Relay(const FGuid& RemoteSessionID)
+bool FRewindDebugger::ConnectToLiveSession_Relay(const FGuid& InRemoteSessionID)
 {
-	CloseActiveTraceSession();
-
 	using namespace UE::TraceBasedDebuggers;
 	if (const TSharedPtr<FTraceSessionsManager> TraceSessionsManager = FRewindDebuggerModule::Get().GetTraceSessionsManager())
 	{
 		FTraceSessionDescriptor NewSessionDescriptor;
 		NewSessionDescriptor.bIsLiveSession = true;
-		NewSessionDescriptor.RemoteSessionID = RemoteSessionID;
-		NewSessionDescriptor.SessionName = TraceSessionsManager->ConnectToLiveSession_Relay(RemoteSessionID);
+		NewSessionDescriptor.RemoteSessionID = InRemoteSessionID;
+		NewSessionDescriptor.SessionName = TraceSessionsManager->ConnectToLiveSession_Relay(InRemoteSessionID);
 
 		if (NewSessionDescriptor.IsValid())
 		{
-			OpenSession(NewSessionDescriptor);
+			OpenAnalysisSession(NewSessionDescriptor);
 			return true;
 		}
 	}
@@ -1090,13 +1159,41 @@ bool FRewindDebugger::ConnectToLiveSession_Relay(const FGuid& RemoteSessionID)
 	return false;
 }
 
-void FRewindDebugger::OpenSession(const UE::TraceBasedDebuggers::FTraceSessionDescriptor& SessionDescriptor)
+void FRewindDebugger::SetSelectedAnalysisSession(const UE::TraceBasedDebuggers::FTraceSessionDescriptor& InSessionDescriptor)
+{
+	SelectedDescriptor = InSessionDescriptor;
+
+	// A new analysis session was selected so make sure to read all the data
+	bAutoScrollToLastItem = true;
+}
+
+void FRewindDebugger::OpenAnalysisSession(const UE::TraceBasedDebuggers::FTraceSessionDescriptor& InSessionDescriptor)
 {
-	ensureMsgf(!CurrentSessionDescriptor.IsValid(), TEXT("Expecting current session to be closed before opening a new one"));
-	CurrentSessionDescriptor = SessionDescriptor;
+	using namespace UE::TraceBasedDebuggers;
+
+	const TSharedPtr<FTraceSessionsManager> TraceManager = FRewindDebuggerModule::Get().GetTraceSessionsManager();
+	ActiveAnalysisDescriptors.RemoveAll([&InSessionDescriptor, TraceManager](const FTraceSessionDescriptor& Element)
+		{
+			if (Element.RemoteSessionID == InSessionDescriptor.RemoteSessionID)
+			{
+				if (TraceManager)
+				{
+					TraceManager->CloseSession(Element.SessionName);
+				}
+				return true;
+			}
+
+			return false;
+		});
+
+	ActiveAnalysisDescriptors.Emplace(InSessionDescriptor);
+
+	ClearTrackData();
+
+	SetSelectedAnalysisSession(InSessionDescriptor);
 
 	// Notify extensions
-	IterateExtensions([this](IRewindDebuggerExtension* Extension)
+	UE::RewindDebugger::IterateExtensions([this](IRewindDebuggerExtension* Extension)
 		{
 			Extension->AnalysisSessionOpened(this);
 		});
@@ -1105,68 +1202,79 @@ void FRewindDebugger::OpenSession(const UE::TraceBasedDebuggers::FTraceSessionDe
 	bDisplayWorldIdValid = false;
 }
 
-const UE::TraceBasedDebuggers::FTraceSessionDescriptor& FRewindDebugger::GetSessionDescriptor(const FGuid& RemoteSessionID) const
+void FRewindDebugger::ClearAnalysisSessionLinkedToRemoteSessionID(const FGuid& InRemoteSessionID)
 {
-	if (CurrentSessionDescriptor.IsValid()
-		&& CurrentSessionDescriptor.RemoteSessionID == RemoteSessionID)
+	using namespace UE::TraceBasedDebuggers;
+
+	if (const FTraceSessionDescriptor* ExistingDescriptor = ActiveAnalysisDescriptors.FindByPredicate(
+		[InRemoteSessionID](const FTraceSessionDescriptor& Descriptor)
+		{
+			return Descriptor.RemoteSessionID == InRemoteSessionID;
+		}))
 	{
-		return CurrentSessionDescriptor;
+		ClearAnalysisSession(*ExistingDescriptor);
 	}
-
-	static UE::TraceBasedDebuggers::FTraceSessionDescriptor InvalidSessionDescriptor;
-	return InvalidSessionDescriptor;
 }
 
-void FRewindDebugger::CloseSessionByRemoteSessionID(const FGuid& RemoteSessionID)
+void FRewindDebugger::StopAnalysisSession(const UE::TraceBasedDebuggers::FTraceSessionDescriptor& InSessionDescriptor)
 {
 	using namespace UE::TraceBasedDebuggers;
 
-	if (!CurrentSessionDescriptor.IsValid()
-		|| CurrentSessionDescriptor.RemoteSessionID != RemoteSessionID)
+	if (InSessionDescriptor.IsValid())
 	{
-		return;
-	}
+		if (const TSharedPtr<FTraceSessionsManager> TraceManager = FRewindDebuggerModule::Get().GetTraceSessionsManager())
+		{
+			TraceManager->StopSession(InSessionDescriptor.SessionName);
+		}
 
-	CloseActiveTraceSession();
+		// Notify extensions
+		UE::RewindDebugger::IterateExtensions([this](IRewindDebuggerExtension* Extension)
+			{
+				Extension->AnalysisSessionClosed(this);
+			});
+	}
 }
 
-#endif // UE_WITH_TRACE_BASED_DEBUGGERS_ANALYSIS
-
-void FRewindDebugger::StopActiveTraceSession() const
+void FRewindDebugger::ClearAnalysisSession(const UE::TraceBasedDebuggers::FTraceSessionDescriptor& InSessionDescriptor)
 {
-#if UE_WITH_TRACE_BASED_DEBUGGERS_ANALYSIS
 	using namespace UE::TraceBasedDebuggers;
 
-	if (CurrentSessionDescriptor.IsValid())
+	if (InSessionDescriptor.IsValid())
 	{
 		if (const TSharedPtr<FTraceSessionsManager> TraceManager = FRewindDebuggerModule::Get().GetTraceSessionsManager())
 		{
-			TraceManager->StopSession(CurrentSessionDescriptor.SessionName);
+			TraceManager->CloseSession(InSessionDescriptor.SessionName);
+		}
+
+		// Notify extensions
+		UE::RewindDebugger::IterateExtensions([this](IRewindDebuggerExtension* Extension)
+			{
+				Extension->AnalysisSessionClosed(this);
+			});
+
+		ActiveAnalysisDescriptors.RemoveAll([&InSessionDescriptor](const FTraceSessionDescriptor& Element)
+			{
+				return Element.RemoteSessionID == InSessionDescriptor.RemoteSessionID;
+			});
+
+		if (SelectedDescriptor == InSessionDescriptor)
+		{
+			SetSelectedAnalysisSession(FTraceSessionDescriptor{});
 		}
 	}
-#endif // UE_WITH_TRACE_BASED_DEBUGGERS_ANALYSIS
 }
 
-void FRewindDebugger::CloseActiveTraceSession()
+#endif // UE_WITH_TRACE_BASED_DEBUGGERS_ANALYSIS
+
+void FRewindDebugger::ClearActiveAnalysisSession()
 {
 #if UE_WITH_TRACE_BASED_DEBUGGERS_ANALYSIS
 	using namespace UE::TraceBasedDebuggers;
 
-	if (CurrentSessionDescriptor.IsValid())
+	if (SelectedDescriptor.IsValid())
 	{
-		if (const TSharedPtr<FTraceSessionsManager> TraceManager = FRewindDebuggerModule::Get().GetTraceSessionsManager())
-		{
-			TraceManager->CloseSession(CurrentSessionDescriptor.SessionName);
-
-			// Notify extensions
-			IterateExtensions([this](IRewindDebuggerExtension* Extension)
-				{
-					Extension->AnalysisSessionClosed(this);
-				});
-		}
+		ClearAnalysisSession(SelectedDescriptor);
 	}
-
-	CurrentSessionDescriptor = {};
 #endif // UE_WITH_TRACE_BASED_DEBUGGERS_ANALYSIS
 }
 
@@ -1176,12 +1284,17 @@ void FRewindDebugger::Tick(float DeltaTime)
 
 	if (const TraceServices::IAnalysisSession* Session = GetAnalysisSession())
 	{
+		double RecordingDurationValue = 0;
+
 		if (const IGameplayProvider* GameplayProvider = Session->ReadProvider<IGameplayProvider>("GameplayProvider"))
 		{
 			TraceServices::FAnalysisSessionReadScope SessionReadScope(*Session);
 
-			// set a default display world when loading a trace (first client/standalone world)
-			if (IsTraceFileLoaded() && !bDisplayWorldIdValid)
+			RecordingDurationValue = GameplayProvider->GetRecordingDuration();
+
+			// set a default display world when loading a trace or analyzing a remote session (first client/standalone world)
+			if (!bDisplayWorldIdValid
+				&& (IsAnalyzingTraceFile() || IsAnalyzingRemoteSession()))
 			{
 				GameplayProvider->EnumerateWorlds([this](const FWorldInfo& WorldInfo)
 					{
@@ -1205,58 +1318,58 @@ void FRewindDebugger::Tick(float DeltaTime)
 						}
 					});
 			}
+		}
 
-			const double RecordingDurationValue = GameplayProvider->GetRecordingDuration();
-			bool bNewDataAvailable = false;
-			if (IsTraceFileLoaded() && RecordingDurationValue > RecordingDuration.Get())
-			{
-				// while trace file is loading up or analysing a live session, force the trace range to update.
-				SetCurrentViewRange(GetCurrentViewRange());
-				bNewDataAvailable = true;
-			}
-			RecordingDuration.Set(RecordingDurationValue);
+		bool bNewDataAvailable = false;
+		if (RecordingDurationValue > RecordingDuration.Get()
+			&& (IsAnalyzingTraceFile() || IsAnalyzingRemoteSession()))
+		{
+			// while trace file is loading up or analyzing a live session, force the trace range to update.
+			SetCurrentViewRange(GetCurrentViewRange());
+			bNewDataAvailable = true;
+		}
+		RecordingDuration.Set(RecordingDurationValue);
 
-			RefreshDebugTracks();
+		RefreshDebugTracks();
 
-			if (bAutoScrollToLastItem)
+		if (bAutoScrollToLastItem)
+		{
+			// Only need to update scrub time and notify
+			// when recording a PIE session or analyzing data from a remote session or trace file.
+			if (IsRecordingLocalSession()
+				|| bNewDataAvailable)
 			{
-				// Only need to update scrub time and notify
-				// when recording a PIE session or analysing data from a remote session or trace file.
-				if (IsRecording()
-					|| bNewDataAvailable)
-				{
-					TRACE_CPUPROFILER_EVENT_SCOPE(FRewindDebugger::Tick_UpdateSimulating);
-					SetCurrentScrubTime(RecordingDuration.Get(), EAutoScrollChange::PreserveCurrentState);
-					TrackCursorDelegate.ExecuteIfBound(/*bReverse*/false);
-				}
-
-				// The debug position is only expected to be set outside of PIE
-				RootObjectPosition.Reset();
+				TRACE_CPUPROFILER_EVENT_SCOPE(FRewindDebugger::Tick_UpdateSimulating);
+				SetCurrentScrubTime(RecordingDuration.Get(), EAutoScrollChange::PreserveCurrentState);
+				TrackCursorDelegate.ExecuteIfBound(/*bReverse*/false);
 			}
-			else
+
+			// The debug position is only expected to be set outside of PIE
+			RootObjectPosition.Reset();
+		}
+		else
+		{
+			if (RecordingDuration.Get() > 0 && CurrentScrubTime <= RecordingDuration.Get())
 			{
-				if (RecordingDuration.Get() > 0 && CurrentScrubTime <= RecordingDuration.Get())
+				if (ControlState == EControlState::Play || ControlState == EControlState::PlayReverse)
 				{
-					if (ControlState == EControlState::Play || ControlState == EControlState::PlayReverse)
-					{
-						const float PlaybackRate = URewindDebuggerSettings::Get().PlaybackRate;
-						TRACE_CPUPROFILER_EVENT_SCOPE(FRewindDebugger::Tick_UpdatePlayback);
-						const float Rate = PlaybackRate * (ControlState == EControlState::Play ? 1 : -1);
-						SetCurrentScrubTime(FMath::Clamp(CurrentScrubTime + Rate * DeltaTime, 0.0f, RecordingDuration.Get()), EAutoScrollChange::Disable);
-						TrackCursorDelegate.ExecuteIfBound(/*bReverse*/Rate < 0);
+					const float PlaybackRate = URewindDebuggerSettings::Get().PlaybackRate;
+					TRACE_CPUPROFILER_EVENT_SCOPE(FRewindDebugger::Tick_UpdatePlayback);
+					const float Rate = PlaybackRate * (ControlState == EControlState::Play ? 1 : -1);
+					SetCurrentScrubTime(FMath::Clamp(CurrentScrubTime + Rate * DeltaTime, 0.0f, RecordingDuration.Get()), EAutoScrollChange::Disable);
+					TrackCursorDelegate.ExecuteIfBound(/*bReverse*/Rate < 0);
 
-						if (CurrentScrubTime == 0 || CurrentScrubTime == RecordingDuration.Get())
-						{
-							// pause at end.
-							ControlState = EControlState::Pause;
-						}
+					if (CurrentScrubTime == 0 || CurrentScrubTime == RecordingDuration.Get())
+					{
+						// pause at end.
+						ControlState = EControlState::Pause;
 					}
 				}
 			}
 		}
 
 		// update extensions
-		IterateExtensions([DeltaTime, this](IRewindDebuggerExtension* Extension)
+		UE::RewindDebugger::IterateExtensions([DeltaTime, this](IRewindDebuggerExtension* Extension)
 			{
 				TRACE_CPUPROFILER_EVENT_SCOPE_TEXT(*Extension->GetName());
 				Extension->Update(DeltaTime, this);
@@ -1269,7 +1382,7 @@ void FRewindDebugger::OnTrackListChanged()
 {
 	TrackListChangedDelegate.ExecuteIfBound();
 
-	IterateExtensions([this](IRewindDebuggerExtension* Extension)
+	UE::RewindDebugger::IterateExtensions([this](IRewindDebuggerExtension* Extension)
 		{
 			TRACE_CPUPROFILER_EVENT_SCOPE_TEXT(*Extension->GetName());
 			Extension->OnTrackListChanged(this);
@@ -1418,7 +1531,7 @@ void FRewindDebugger::SetDisplayWorld(uint64 WorldId)
 {
 	DisplayWorldId = WorldId;
 
-	IterateExtensions([this](IRewindDebuggerExtension* Extension)
+	UE::RewindDebugger::IterateExtensions([this](IRewindDebuggerExtension* Extension)
 		{
 			Extension->Clear(this);
 			Extension->Update(0.0, this);
@@ -1700,7 +1813,7 @@ void FRewindDebugger::RegisterToolBar()
 		FSlateIcon(FAppStyle::GetAppStyleSetName(), "Icons.FolderOpen")));
 
 	Section.AddEntry(FToolMenuEntry::InitToolBarButton(
-		Commands.ClearTrace,
+		Commands.ClearAnalysis,
 		FText(),
 		TAttribute<FText>(),
 		FSlateIcon(FAppStyle::GetAppStyleSetName(), "Icons.Delete")));
@@ -1725,7 +1838,7 @@ void FRewindDebugger::RegisterToolBar()
 		"Display World",
 		FUIAction(
 			FExecuteAction(),
-			FCanExecuteAction::CreateLambda([]() { return Instance()->IsTraceFileLoaded(); })
+			FCanExecuteAction::CreateLambda([]() { return !Instance()->IsPIESimulating(); })
 		),
 		FNewToolMenuDelegate::CreateStatic(&FRewindDebugger::MakeWorldsMenu),
 		LOCTEXT("Display World", "Display World"),
@@ -1817,8 +1930,7 @@ void FRewindDebugger::SetObjectToDebug(const RewindDebugger::FObjectId ObjectId)
 		return;
 	}
 
-	CandidateIds.Reset(1);
-	CandidateIds.Push(ObjectId);
+	CandidateIds = {ObjectId};
 
 	RefreshDebugTracks();
 }
diff --git a/Engine/Plugins/Animation/GameplayInsights/Source/RewindDebugger/Private/RewindDebugger.h b/Engine/Plugins/Animation/GameplayInsights/Source/RewindDebugger/Private/RewindDebugger.h
--- a/Engine/Plugins/Animation/GameplayInsights/Source/RewindDebugger/Private/RewindDebugger.h
+++ b/Engine/Plugins/Animation/GameplayInsights/Source/RewindDebugger/Private/RewindDebugger.h
@@ -23,6 +23,11 @@ namespace UE::RewindDebugger
 struct FDebuggerRecordingControls;
 }
 
+namespace UE::TraceBasedDebuggers
+{
+struct FSessionInfo;
+}
+
 class FMenuBuilder;
 class IGameplayProvider;
 class SDockTab;
@@ -74,8 +79,11 @@ public:
 	{
 		return bPIESimulating;
 	}
+	bool IsPIEStarted() const
+	{
+		return bPIEStarted;
+	}
 
-	virtual bool IsTraceFileLoaded() const override;
 	virtual double GetRecordingDuration() const override
 	{
 		return RecordingDuration.Get();
@@ -92,6 +100,15 @@ public:
 		return DisplayWorldId == WorldId;
 	}
 
+	/** @return Whether the debugger is currently analyzing a trace session from the current process (e.g., PIE) */
+	bool IsAnalyzingLocalSession() const;
+
+	/** @return Whether the debugger is currently analyzing a remote trace session (e.g., external process like standalone game, client, server) */
+	bool IsAnalyzingRemoteSession() const;
+
+	/** @return Whether the debugger is currently analyzing a trace file */
+	bool IsAnalyzingTraceFile() const;
+
 	void GetTargetObjectIds(TArray<RewindDebugger::FObjectId>& OutTargetObjectIds) const;
 
 	TArray<TSharedPtr<RewindDebugger::FRewindDebuggerTrack>>& GetTracks()
@@ -113,11 +130,13 @@ public:
 
 	virtual bool CanStartRecording() const override;
 
-	/** Start a new Recording:  Start tracing Objects, increment the current recording index, and reset the recording elapsed time to 0 */
+	/** Starts a new recording for the local session */
 	virtual void StartRecording() const override;
 
-	void OnRecordingStarted();
-	void OnRecordingStopped();
+	void StartRecordingLocalSession() const;
+	bool CanStartRecordingLocalSession() const;
+	bool IsRecordingLocalSession() const;
+
 	void OnClearRecording();
 
 	bool CanOpenTrace() const;
@@ -126,9 +145,14 @@ public:
 
 	void AttachToSession();
 
-	bool CanClearTrace() const;
-	void ClearTrace();
-	void ClearTraceAndActiveSession();
+	/** @return Whether there is an active analysis session to close and/or tracks data to clear */
+	bool CanClearAnalysis() const;
+
+	/** Closes the currently selected analysis session and clears debugger and tracks data */
+	void ClearAnalysis();
+
+	/** Clears all debugger and track data associated to the current, or last, analysis session */
+	void ClearTrackData();
 
 	void ToggleAutoScroll();
 	bool CanEnableAutoScroll() const;
@@ -140,14 +164,13 @@ public:
 	bool ShouldAutoEject() const;
 	void SetShouldAutoEject(bool InValue);
 
-	/** Stop recording: Stop tracing Object + Animation Data. */
-	void StopRecording() const;
-	bool CanStopRecording() const
+	void StopRecordingLocalSession() const;
+	bool CanStopRecordingLocalSession() const
 	{
-		return IsRecording();
+		return IsRecordingLocalSession();
 	}
 
-	//~ VCR controls
+	//~ Playback controls
 
 	bool CanPause() const;
 	void Pause();
@@ -230,25 +253,61 @@ public:
 #if UE_WITH_TRACE_BASED_DEBUGGERS_ANALYSIS
 public:
 	bool ConnectToLiveSession(uint32 SessionID, FStringView InSessionAddress);
-	bool ConnectToLiveSession_Direct(const FGuid& RemoteSessionID);
-	bool ConnectToLiveSession_Relay(const FGuid& RemoteSessionID);
+	bool ConnectToLiveSession_Direct(const FGuid& InRemoteSessionID);
+	bool ConnectToLiveSession_Relay(const FGuid& InRemoteSessionID);
 
-	void CloseSessionByRemoteSessionID(const FGuid& RemoteSessionID);
+	/**
+	 * Finds the trace session analysis linked to the provided remote session ID,
+	 * then calls ClearAnalysisSession.
+	 * @see ClearAnalysisSession
+	 */
+	void ClearAnalysisSessionLinkedToRemoteSessionID(const FGuid& InRemoteSessionID);
 
 	const UE::TraceBasedDebuggers::FTraceSessionDescriptor& GetCurrentSessionDescriptor() const
 	{
-		return CurrentSessionDescriptor;
+		return SelectedDescriptor;
 	}
 
 private:
-	void OpenSession(const UE::TraceBasedDebuggers::FTraceSessionDescriptor& SessionDescriptor);
-	const UE::TraceBasedDebuggers::FTraceSessionDescriptor& GetSessionDescriptor(const FGuid& RemoteSessionID) const;
+	/** Sets the analysis session as the one used by the debugger */
+	void SetSelectedAnalysisSession(const UE::TraceBasedDebuggers::FTraceSessionDescriptor& InSessionDescriptor);
+
+	/**
+	 * Adds the provided session descriptor to the list of active analysis sessions,
+	 * sets the selected analysis session that the debugger will represent and notifies extensions.
+	 * @see SetSelectedAnalysisSession
+	 */
+	void OpenAnalysisSession(const UE::TraceBasedDebuggers::FTraceSessionDescriptor& InSessionDescriptor);
+
+	/**
+	 * Stops the trace analysis session associated to the remote session specified in the
+	 * provided session descriptor and notifies the extensions.
+	 * @note Unlike ClearAnalysisSession this method will keep the session available by not removing it
+	 * from the Trace Sessions Manager.
+	 * The current session selection is not affected.
+	 * @see ClearAnalysisSession
+	 */
+	void StopAnalysisSession(const UE::TraceBasedDebuggers::FTraceSessionDescriptor& InSessionDescriptor);
+
+	/**
+	 * Closes the trace analysis session (i.e., stop and remove from the Trace Sessions Manager)
+	 * associated to the remote session specified in the provided session descriptor,
+	 * notifies the extensions and removes the session from the active sessions.
+	 * Clears the current selected session, if applicable.
+	 * @note Consider using StopAnalysisSession if the session information still need to be accessed after stopping the analysis.
+	 * @see StopAnalysisSession
+	 * @see SetSelectedAnalysisSession
+	 */
+	void ClearAnalysisSession(const UE::TraceBasedDebuggers::FTraceSessionDescriptor& InSessionDescriptor);
+
+	void HandleRecordingStarted(const TSharedPtr<UE::TraceBasedDebuggers::FSessionInfo>& InSessionInfo);
+	void HandleRecordingStopped(const TSharedPtr<UE::TraceBasedDebuggers::FSessionInfo>& InSessionInfo);
+	void HandleSessionSelected(const TSharedPtr<UE::TraceBasedDebuggers::FSessionInfo>& InSessionInfo);
 #endif // UE_WITH_TRACE_BASED_DEBUGGERS_ANALYSIS
 
 private:
 
-	void CloseActiveTraceSession();
-	void StopActiveTraceSession() const;
+	void ClearActiveAnalysisSession();
 
 #if WITH_TRACE_BASED_DEBUGGERS
 	void StopLocalRecording() const;
@@ -300,7 +359,6 @@ private:
 
 	bool bPIEStarted = false;
 	bool bPIESimulating = false;
-	bool bRecordingPIESession = false;
 
 	/** Instructs the debugger to always scrub to the last received event */
 	bool bAutoScrollToLastItem = false;
@@ -326,7 +384,15 @@ private:
 	static void GetScrubTimeInformation(double InDebugTime, FScrubTimeInformation& InOutTimeInformation, uint16 InRecordingIndex, const TraceServices::IAnalysisSession* AnalysisSession);
 
 #if UE_WITH_TRACE_BASED_DEBUGGERS_ANALYSIS
-	UE::TraceBasedDebuggers::FTraceSessionDescriptor CurrentSessionDescriptor;
+	/** List of session descriptors associated to active trace analysis sessions */
+	TArray<UE::TraceBasedDebuggers::FTraceSessionDescriptor> ActiveAnalysisDescriptors;
+
+	/**
+	 * Trace analysis descriptor associated to the selected remote session
+	 * An invalid descriptor indicates that no active recording/analysis is active
+	 * for the currently selected session.
+	 */
+	UE::TraceBasedDebuggers::FTraceSessionDescriptor SelectedDescriptor;
 #endif // UE_WITH_TRACE_BASED_DEBUGGERS_ANALYSIS
 	TArray<TSharedPtr<FDebugObjectInfo>> DebuggedObjects;
 	mutable TSharedPtr<FDebugObjectInfo> SelectedObject;
diff --git a/Engine/Plugins/Animation/GameplayInsights/Source/RewindDebugger/Private/RewindDebuggerAnimation.cpp b/Engine/Plugins/Animation/GameplayInsights/Source/RewindDebugger/Private/RewindDebuggerAnimation.cpp
--- a/Engine/Plugins/Animation/GameplayInsights/Source/RewindDebugger/Private/RewindDebuggerAnimation.cpp
+++ b/Engine/Plugins/Animation/GameplayInsights/Source/RewindDebugger/Private/RewindDebuggerAnimation.cpp
@@ -306,8 +306,11 @@ void FRewindDebuggerAnimation::Update(float DeltaTime, IRewindDebugger* RewindDe
 					// - apply the recorded pose for the current Frame
 					{
 						TRACE_CPUPROFILER_EVENT_SCOPE(FRewindDebugger::Tick_UpdatePoses);
+
+						const bool ResetOrSpawnMeshes = !RewindDebugger->IsPIESimulating() && RewindDebugger->GetAnalysisSession() != nullptr;
+
 						AnimationProvider->EnumerateSkeletalMeshPoseTimelines(
-						[this, RewindDebugger, &Frame, AnimationProvider, GameplayProvider, &bSearchForActor, &DebuggedObjects, &bNewActor](uint64 ObjectId, const IAnimationProvider::SkeletalMeshPoseTimeline& TimelineData)
+						[this, ResetOrSpawnMeshes, RewindDebugger, &Frame, AnimationProvider, GameplayProvider, &bSearchForActor, &DebuggedObjects, &bNewActor](uint64 ObjectId, const IAnimationProvider::SkeletalMeshPoseTimeline& TimelineData)
 						{
 							if (bSearchForActor)
 							{
@@ -336,9 +339,8 @@ void FRewindDebuggerAnimation::Update(float DeltaTime, IRewindDebugger* RewindDe
 							USkeletalMeshComponent* MeshComponent = nullptr;
 							bool bQueueForReset = false;
 							bool bLoadMesh = false;
-
 #if OBJECT_TRACE_ENABLED
-							if (!RewindDebugger->IsTraceFileLoaded())
+							if (ResetOrSpawnMeshes)
 							{
 								if(UObject* ObjectInstance = FObjectTrace::GetObjectFromId(ObjectId))
 								{
@@ -366,7 +368,7 @@ void FRewindDebuggerAnimation::Update(float DeltaTime, IRewindDebugger* RewindDe
 									
 								if (MeshComponentInfo == nullptr)
 								{
-									if (RewindDebugger->IsTraceFileLoaded())
+									if (ResetOrSpawnMeshes)
 									{
 										if (const FWorldInfo* ObjectWorldInfo = GameplayProvider->FindWorldInfoFromObject(ObjectId))
 										{
diff --git a/Engine/Plugins/Animation/GameplayInsights/Source/RewindDebugger/Private/RewindDebuggerCommands.h b/Engine/Plugins/Animation/GameplayInsights/Source/RewindDebugger/Private/RewindDebuggerCommands.h
--- a/Engine/Plugins/Animation/GameplayInsights/Source/RewindDebugger/Private/RewindDebuggerCommands.h
+++ b/Engine/Plugins/Animation/GameplayInsights/Source/RewindDebugger/Private/RewindDebuggerCommands.h
@@ -1,4 +1,4 @@
-// Copyright Epic Games, Inc. All Rights Reserved.
+﻿// Copyright Epic Games, Inc. All Rights Reserved.
 
 #pragma once
 
@@ -34,7 +34,7 @@ public:
 		UI_COMMAND(AutoRecord, "Auto Record", "Automatically start recording when PIE is started", EUserInterfaceActionType::ToggleButton, FInputChord());
 		UI_COMMAND(OpenTrace, "Open Recording", "Open a Trace file containing a rewind debugger recording", EUserInterfaceActionType::Button, FInputChord());
 		UI_COMMAND(AttachToSession, "Attach to Session", "Attach to a live session", EUserInterfaceActionType::Button, FInputChord());
-		UI_COMMAND(ClearTrace, "Clear Recording", "Clear current recording data", EUserInterfaceActionType::Button, FInputChord());
+		UI_COMMAND(ClearAnalysis, "Clear Analysis", "Stop current analysis and clear related debugger data", EUserInterfaceActionType::Button, FInputChord());
 		UI_COMMAND(AutoScroll, "AutoScroll", "Maintains the cursor on the last recorded event", EUserInterfaceActionType::ToggleButton, FInputChord());
 	}
 
@@ -53,7 +53,7 @@ public:
 	TSharedPtr<FUICommandInfo> AutoRecord;
 	TSharedPtr<FUICommandInfo> OpenTrace;
 	TSharedPtr<FUICommandInfo> AttachToSession;
-	TSharedPtr<FUICommandInfo> ClearTrace;
+	TSharedPtr<FUICommandInfo> ClearAnalysis;
 	TSharedPtr<FUICommandInfo> AutoScroll;
 };
 
diff --git a/Engine/Plugins/Animation/GameplayInsights/Source/RewindDebugger/Private/RewindDebuggerModule.cpp b/Engine/Plugins/Animation/GameplayInsights/Source/RewindDebugger/Private/RewindDebuggerModule.cpp
--- a/Engine/Plugins/Animation/GameplayInsights/Source/RewindDebugger/Private/RewindDebuggerModule.cpp
+++ b/Engine/Plugins/Animation/GameplayInsights/Source/RewindDebugger/Private/RewindDebuggerModule.cpp
@@ -1,4 +1,4 @@
-// Copyright Epic Games, Inc. All Rights Reserved.
+﻿// Copyright Epic Games, Inc. All Rights Reserved.
 
 #include "RewindDebuggerModule.h"
 #include "AnimInstanceHelpers.h"
@@ -182,9 +182,9 @@ TSharedRef<SDockTab> FRewindDebuggerModule::SpawnRewindDebuggerTab(const FSpawnT
 							FIsActionChecked(),
 							FIsActionButtonVisible());
 
-	CommandList->MapAction(Commands.ClearTrace,
-							FExecuteAction::CreateRaw(DebuggerInstance, &FRewindDebugger::ClearTraceAndActiveSession),
-							FCanExecuteAction::CreateRaw(DebuggerInstance, &FRewindDebugger::CanClearTrace),
+	CommandList->MapAction(Commands.ClearAnalysis,
+							FExecuteAction::CreateRaw(DebuggerInstance, &FRewindDebugger::ClearAnalysis),
+							FCanExecuteAction::CreateRaw(DebuggerInstance, &FRewindDebugger::CanClearAnalysis),
 							FIsActionChecked(),
 							FIsActionButtonVisible());
 
@@ -199,19 +199,23 @@ TSharedRef<SDockTab> FRewindDebuggerModule::SpawnRewindDebuggerTab(const FSpawnT
 	{
 		check(FPlayWorldCommands::GlobalPlayWorldActions.IsValid());
 
+		// Recording buttons have been replaced by the recording controls widget and the remote sessions manager.
+		// The following commands are not added to the toolbar but still registered to allow users to use associated shortcuts.
+		// To preserve legacy behavior they only control recording on the local session (i.e., PIE), regardless of the session
+		// selected in the recording controls.
 		FPlayWorldCommands::GlobalPlayWorldActions->MapAction(Commands.StartRecording, 
-							FExecuteAction::CreateRaw(DebuggerInstance, &FRewindDebugger::StartRecording), 
+							FExecuteAction::CreateRaw(DebuggerInstance, &FRewindDebugger::StartRecording),
 							FCanExecuteAction::CreateRaw(DebuggerInstance, &FRewindDebugger::CanStartRecording),
 							FIsActionChecked(),
-							FIsActionButtonVisible::CreateLambda([]() { return !FRewindDebugger::Instance()->IsRecording();}));
+							FIsActionButtonVisible());
 
 		FPlayWorldCommands::GlobalPlayWorldActions->MapAction(Commands.StopRecording,
-							 FExecuteAction::CreateRaw(DebuggerInstance, &FRewindDebugger::StopRecording),
-							 FCanExecuteAction::CreateRaw(DebuggerInstance, &FRewindDebugger::CanStopRecording),
+							 FExecuteAction::CreateRaw(DebuggerInstance, &FRewindDebugger::StopRecordingLocalSession),
+							 FCanExecuteAction::CreateRaw(DebuggerInstance, &FRewindDebugger::CanStopRecordingLocalSession),
 							 FIsActionChecked(),
-							 FIsActionButtonVisible::CreateRaw(DebuggerInstance, &FRewindDebugger::CanStopRecording));
+							 FIsActionButtonVisible());
 	}
-	
+
 	RewindDebuggerWidget = SNew(SRewindDebugger, CommandList.ToSharedRef(), MajorTab, SpawnTabArgs.GetOwnerWindow())
 								.DebuggedObjectName({ DebuggerInstance->GetRootObjectNameProperty(), URewindDebuggerSettings::Get().DebugTargetActor})
 								.RecordingDuration(DebuggerInstance->GetRecordingDurationProperty())
diff --git a/Engine/Plugins/Animation/GameplayInsights/Source/RewindDebugger/Private/RewindDebuggerRecordingControls.cpp b/Engine/Plugins/Animation/GameplayInsights/Source/RewindDebugger/Private/RewindDebuggerRecordingControls.cpp
--- a/Engine/Plugins/Animation/GameplayInsights/Source/RewindDebugger/Private/RewindDebuggerRecordingControls.cpp
+++ b/Engine/Plugins/Animation/GameplayInsights/Source/RewindDebugger/Private/RewindDebuggerRecordingControls.cpp
@@ -29,38 +29,66 @@ FDebuggerRecordingControls::FDebuggerRecordingControls(const FName& InStatusBarI
 bool FDebuggerRecordingControls::IsRecordingLocalSession() const
 {
 	using namespace UE::TraceBasedDebuggers;
-	if (const TSharedPtr<FSessionInfo> SessionInfo = GetCurrentSessionInfo())
+	if (const TSharedPtr<FSessionInfo> SessionInfo = GetLocalSessionInfo())
 	{
-		return (SessionInfo->SessionName == FRemoteSessionsManager::LocalEditorSessionName
-			&& SessionInfo->IsRecording(DebuggerGuid));
+		return SessionInfo->IsRecording(DebuggerGuid);
 	}
 
 	return false;
 }
 
-bool FDebuggerRecordingControls::CanStartRecording() const
+TSharedPtr<TraceBasedDebuggers::FSessionInfo> FDebuggerRecordingControls::GetLocalSessionInfo() const
 {
-	return !IsRecording() && IsRecordingToggleButtonEnabled();
+	using namespace UE::TraceBasedDebuggers;
+	TSharedPtr<FSessionInfo> LocalSessionInfo;
+	if (const TSharedPtr<FRemoteSessionsManager> RemoteSessionManager = RemoteSessionManagerWeakPtr.Pin())
+	{
+		RemoteSessionManager->EnumerateActiveSessions([&LocalSessionInfo](const TSharedRef<FSessionInfo>& InSessionInfoRef)
+			{
+				bool bKeepIterating = true;
+				if (InSessionInfoRef->SessionName == FRemoteSessionsManager::LocalEditorSessionName)
+				{
+					LocalSessionInfo = InSessionInfoRef;
+					bKeepIterating = false;
+				}
+				return bKeepIterating;
+			});
+	}
+
+	return LocalSessionInfo;
 }
 
-void FDebuggerRecordingControls::StartRecording()
+bool FDebuggerRecordingControls::CanStartRecordingLocalSession() const
 {
-	if (!IsRecording())
-	{
-		ToggleRecordingState();
-	}
+	using namespace UE::TraceBasedDebuggers;
+	const TSharedPtr<FSessionInfo> LocalSessionInfo = GetLocalSessionInfo();
+	return LocalSessionInfo
+		&& !LocalSessionInfo->IsRecording(DebuggerGuid)
+		&& IsRecordingAvailableForSessionInternal(LocalSessionInfo.ToSharedRef());
 }
 
-bool FDebuggerRecordingControls::CanStopRecording() const
+void FDebuggerRecordingControls::StartRecordingLocalSession()
 {
-	return IsRecording() && IsRecordingToggleButtonEnabled();
+	using namespace UE::TraceBasedDebuggers;
+	if (const TSharedPtr<FSessionInfo> LocalSessionInfo = GetLocalSessionInfo())
+	{
+		if (!LocalSessionInfo->IsRecording(DebuggerGuid))
+		{
+			ToggleRecordingState(LocalSessionInfo.ToSharedRef());
+		}
+
+		SelectTargetSession(LocalSessionInfo->InstanceId);
+	}
 }
 
-void FDebuggerRecordingControls::StopRecording()
+void FDebuggerRecordingControls::StopRecordingLocalSession()
 {
-	if (IsRecording())
+	using namespace TraceBasedDebuggers;
+	const TSharedPtr<FSessionInfo> LocalSessionInfo = GetLocalSessionInfo();
+	if (LocalSessionInfo
+		&& LocalSessionInfo->IsRecording(DebuggerGuid))
 	{
-		ToggleRecordingState();
+		ToggleRecordingState(LocalSessionInfo.ToSharedRef());
 	}
 }
 
@@ -82,10 +110,10 @@ bool FDebuggerRecordingControls::ConnectToLiveSession_DirectInternal(const FGuid
 #endif
 }
 
-bool FDebuggerRecordingControls::IsRecordingToggleButtonEnabledInternal(const TSharedRef<TraceBasedDebuggers::FSessionInfo>& InSessionInfo) const
+bool FDebuggerRecordingControls::IsRecordingAvailableForSessionInternal(const TSharedRef<TraceBasedDebuggers::FSessionInfo>& InSessionInfo) const
 {
 	return (InSessionInfo->BuildTargetType == EBuildTargetType::Editor
-			&& (FRewindDebugger::Instance()->IsPIESimulating() || FRewindDebugger::Instance()->IsRecording()))
+			&& (FRewindDebugger::Instance()->IsPIEStarted()))
 		|| InSessionInfo->BuildTargetType != EBuildTargetType::Editor;
 }
 
@@ -119,7 +147,7 @@ void FDebuggerRecordingControls::OnToggleRecordingStateInternal(const TSharedRef
 void FDebuggerRecordingControls::CloseSessionByRemoteSessionIDInternal(const FGuid& InSessionId)
 {
 #if UE_WITH_TRACE_BASED_DEBUGGERS_ANALYSIS
-	FRewindDebugger::Instance()->CloseSessionByRemoteSessionID(InSessionId);
+	FRewindDebugger::Instance()->ClearAnalysisSessionLinkedToRemoteSessionID(InSessionId);
 #endif
 }
 
@@ -130,13 +158,4 @@ const TraceBasedDebuggers::FTraceSessionDescriptor& FDebuggerRecordingControls::
 #endif
 }
 
-void FDebuggerRecordingControls::OnRecordingStopped()
-{
-	FRewindDebugger::Instance()->OnRecordingStopped();
-}
-
-void FDebuggerRecordingControls::OnRecordingStarted()
-{
-	FRewindDebugger::Instance()->OnRecordingStarted();
-}
 } // UE::RewindDebugger
\ No newline at end of file
diff --git a/Engine/Plugins/Animation/GameplayInsights/Source/RewindDebugger/Private/RewindDebuggerRecordingControls.h b/Engine/Plugins/Animation/GameplayInsights/Source/RewindDebugger/Private/RewindDebuggerRecordingControls.h
--- a/Engine/Plugins/Animation/GameplayInsights/Source/RewindDebugger/Private/RewindDebuggerRecordingControls.h
+++ b/Engine/Plugins/Animation/GameplayInsights/Source/RewindDebugger/Private/RewindDebuggerRecordingControls.h
@@ -1,4 +1,4 @@
-// Copyright Epic Games, Inc. All Rights Reserved.
+﻿// Copyright Epic Games, Inc. All Rights Reserved.
 
 #pragma once
 
@@ -17,15 +17,19 @@ struct FDebuggerRecordingControls : TraceBasedDebuggers::FRecordingControls
 	using FRecordingControls::IsRecording;
 	bool IsRecordingLocalSession() const;
 
-	bool CanStartRecording() const;
-	void StartRecording();
+	bool CanStartRecordingLocalSession() const;
+
+	/** Starts recording in the local session and makes it the selected session. */
+	void StartRecordingLocalSession();
+
+	void StopRecordingLocalSession();
 
-	bool CanStopRecording() const;
-	void StopRecording();
 protected:
+	TSharedPtr<TraceBasedDebuggers::FSessionInfo> GetLocalSessionInfo() const;
+
 	virtual bool ConnectToLiveSession_RelayInternal(const FGuid& InSessionId) override;
 	virtual bool ConnectToLiveSession_DirectInternal(const FGuid& InSessionId) override;
-	virtual bool IsRecordingToggleButtonEnabledInternal(const TSharedRef<TraceBasedDebuggers::FSessionInfo>& InSessionInfo) const override;
+	virtual bool IsRecordingAvailableForSessionInternal(const TSharedRef<TraceBasedDebuggers::FSessionInfo>& InSessionInfo) const override;
 
 	virtual TNotNull<TUniquePtr<TraceBasedDebuggers::FStartRecordingCommandMessage>> BuildStartRecordingParamsInternal() override;
 	virtual void SendStartRecordingCommandInternal(const FMessageAddress& InAddress, TUniquePtr<TraceBasedDebuggers::FStartRecordingCommandMessage> InStartRecordingCommand) override;
@@ -34,9 +38,6 @@ protected:
 	virtual void OnToggleRecordingStateInternal(const TSharedRef<TraceBasedDebuggers::FSessionInfo>& InSessionInfo) override;
 	virtual void CloseSessionByRemoteSessionIDInternal(const FGuid& InSessionId) override;
 	virtual const TraceBasedDebuggers::FTraceSessionDescriptor& GetCurrentSessionDescriptorInternal() override;
-
-	virtual void OnRecordingStopped() override;
-	virtual void OnRecordingStarted() override;
 };
 
 }
diff --git a/Engine/Plugins/ChaosVD/Source/ChaosVD/Private/Widgets/ChaosVDRecordingControls.cpp b/Engine/Plugins/ChaosVD/Source/ChaosVD/Private/Widgets/ChaosVDRecordingControls.cpp
--- a/Engine/Plugins/ChaosVD/Source/ChaosVD/Private/Widgets/ChaosVDRecordingControls.cpp
+++ b/Engine/Plugins/ChaosVD/Source/ChaosVD/Private/Widgets/ChaosVDRecordingControls.cpp
@@ -116,7 +116,7 @@ void FChaosVDRecordingControls::OnGenerateTargetSessionSelectorInternal(FMenuBui
 	MenuBuilder.EndSection();
 }
 
-int32 FChaosVDRecordingControls::GetMaxConnectionRetriesInternal()
+int32 FChaosVDRecordingControls::GetMaxConnectionRetriesInternal() const
 {
 	if (const UChaosVDGeneralSettings* Settings = FChaosVDSettingsManager::Get().GetSettingsObject<UChaosVDGeneralSettings>())
 	{
@@ -128,7 +128,7 @@ int32 FChaosVDRecordingControls::GetMaxConnectionRetriesInternal()
 	return DefaultValue;
 }
 
-UE::TraceBasedDebuggers::ETraceTransportMode FChaosVDRecordingControls::GetTransportModeOverrideForTargetTypeInternal(const EBuildTargetType InBuildTarget)
+UE::TraceBasedDebuggers::ETraceTransportMode FChaosVDRecordingControls::GetTransportModeOverrideForTargetTypeInternal(const EBuildTargetType InBuildTarget) const
 {
 	if (const UChaosVDGeneralSettings* Settings = FChaosVDSettingsManager::Get().GetSettingsObject<UChaosVDGeneralSettings>())
 	{
@@ -199,7 +199,7 @@ void FChaosVDRecordingControls::SendStopRecordingCommandInternal(const FMessageA
 	}
 }
 
-bool FChaosVDRecordingControls::IsRecordingToggleButtonEnabledInternal(const TSharedRef<UE::TraceBasedDebuggers::FSessionInfo>& InSessionInfo) const
+bool FChaosVDRecordingControls::IsRecordingAvailableForSessionInternal(const TSharedRef<UE::TraceBasedDebuggers::FSessionInfo>& InSessionInfo) const
 {
 	// If we are recording, don't show the stop button for the mode that is disabled
 	if (EnumHasAnyFlags(InSessionInfo->GetSessionTypeAttributes(), UE::TraceBasedDebuggers::ERemoteSessionAttributes::IsMultiSessionWrapper))
diff --git a/Engine/Plugins/ChaosVD/Source/ChaosVD/Private/Widgets/ChaosVDRecordingControls.h b/Engine/Plugins/ChaosVD/Source/ChaosVD/Private/Widgets/ChaosVDRecordingControls.h
--- a/Engine/Plugins/ChaosVD/Source/ChaosVD/Private/Widgets/ChaosVDRecordingControls.h
+++ b/Engine/Plugins/ChaosVD/Source/ChaosVD/Private/Widgets/ChaosVDRecordingControls.h
@@ -32,8 +32,8 @@ protected:
 	virtual void AddToMenuInternal(FToolMenuSection& InSection) override;
 	virtual void OnGenerateTargetSessionSelectorInternal(FMenuBuilder& InMenuBuilder) override;
 
-	virtual int32 GetMaxConnectionRetriesInternal() override;
-	virtual UE::TraceBasedDebuggers::ETraceTransportMode GetTransportModeOverrideForTargetTypeInternal(EBuildTargetType InBuildTarget) override;
+	virtual int32 GetMaxConnectionRetriesInternal() const override;
+	virtual UE::TraceBasedDebuggers::ETraceTransportMode GetTransportModeOverrideForTargetTypeInternal(EBuildTargetType InBuildTarget) const override;
 
 	virtual bool ConnectToLiveSession_RelayInternal(const FGuid& InSessionId) override;
 	virtual bool ConnectToLiveSession_DirectInternal(const FGuid& InSessionId) override;
@@ -42,7 +42,7 @@ protected:
 	virtual void SendStartRecordingCommandInternal(const FMessageAddress& InAddress, TUniquePtr<UE::TraceBasedDebuggers::FStartRecordingCommandMessage> InStartRecordingCommand) override;
 	virtual void SendStopRecordingCommandInternal(const FMessageAddress& InAddress) override;
 
-	virtual bool IsRecordingToggleButtonEnabledInternal(const TSharedRef<UE::TraceBasedDebuggers::FSessionInfo>& InSessionInfo) const override;
+	virtual bool IsRecordingAvailableForSessionInternal(const TSharedRef<UE::TraceBasedDebuggers::FSessionInfo>& InSessionInfo) const override;
 	virtual void OnToggleRecordingStateInternal(const TSharedRef<UE::TraceBasedDebuggers::FSessionInfo>& InSessionInfo) override;
 
 	virtual void CloseSessionByRemoteSessionIDInternal(const FGuid& InSessionId) override;
diff --git a/Engine/Source/Editor/RewindDebuggerInterface/Public/IRewindDebugger.h b/Engine/Source/Editor/RewindDebuggerInterface/Public/IRewindDebugger.h
--- a/Engine/Source/Editor/RewindDebuggerInterface/Public/IRewindDebugger.h
+++ b/Engine/Source/Editor/RewindDebuggerInterface/Public/IRewindDebugger.h
@@ -166,7 +166,11 @@ public:
 	virtual bool IsRecording() const = 0;
 
 	/** @return whether a trace file is loaded from disk */
-	virtual bool IsTraceFileLoaded() const = 0;
+	UE_DEPRECATED(5.8, "Method is no longer used and will be removed.")
+	virtual bool IsTraceFileLoaded() const final
+	{
+		return false;
+	}
 
 	/** @return PIE is running and not paused */
 	virtual bool IsPIESimulating() const = 0;
diff --git a/Engine/Source/Runtime/Experimental/TraceBasedDebuggers/Private/RemoteSessionsManager.cpp b/Engine/Source/Runtime/Experimental/TraceBasedDebuggers/Private/RemoteSessionsManager.cpp
--- a/Engine/Source/Runtime/Experimental/TraceBasedDebuggers/Private/RemoteSessionsManager.cpp
+++ b/Engine/Source/Runtime/Experimental/TraceBasedDebuggers/Private/RemoteSessionsManager.cpp
@@ -1,4 +1,4 @@
-﻿// Copyright Epic Games, Inc. All Rights Reserved.
+// Copyright Epic Games, Inc. All Rights Reserved.
 
 #include "RemoteSessionsManager.h"
 
@@ -18,6 +18,14 @@
 
 namespace UE::TraceBasedDebuggers
 {
+
+bool bIgnoreExpiredSessionsForDebugging = false;
+static FAutoConsoleVariableRef CVarIgnoreExpiredSessionsForDebugging(
+	TEXT("RemoteSessionsManager.IgnoreExpiredSessionsForDebugging"),
+	bIgnoreExpiredSessionsForDebugging,
+	TEXT("Prevent sessions to expire due to breakpoints."),
+	ECVF_Cheat);
+
 const FGuid FRemoteSessionsManager::AllRemoteSessionsWrapperGUID = FGuid::NewGuid();
 const FGuid FRemoteSessionsManager::AllRemoteServersWrapperGUID = FGuid::NewGuid();
 const FGuid FRemoteSessionsManager::AllRemoteClientsWrapperGUID = FGuid::NewGuid();
@@ -597,6 +605,12 @@ void FRemoteSessionsManager::HandleFullSessionStateResponseMessage(const FFullSe
 
 void FRemoteSessionsManager::RemoveExpiredSessions(const ERemoveSessionOptions Options)
 {
+	if (bIgnoreExpiredSessionsForDebugging
+		&& !EnumHasAnyFlags(Options, ERemoveSessionOptions::ForceRemoveAll))
+	{
+		return;
+	}
+
 	bool bAnySessionRemoved = false;
 	const FDateTime CurrentTime = FDateTime::UtcNow();
 	for (TMap<FGuid, TSharedPtr<FSessionInfo>>::TIterator RemoveIterator = ActiveSessionsById.CreateIterator(); RemoveIterator; ++RemoveIterator)
diff --git a/Engine/Source/Runtime/Experimental/TraceBasedDebuggersAnalysis/Private/Widgets/RecordingControls.cpp b/Engine/Source/Runtime/Experimental/TraceBasedDebuggersAnalysis/Private/Widgets/RecordingControls.cpp
--- a/Engine/Source/Runtime/Experimental/TraceBasedDebuggersAnalysis/Private/Widgets/RecordingControls.cpp
+++ b/Engine/Source/Runtime/Experimental/TraceBasedDebuggersAnalysis/Private/Widgets/RecordingControls.cpp
@@ -230,6 +230,11 @@ EVisibility FRecordingControls::GetRecordingTimeTextBlockVisibility() const
 void FRecordingControls::SelectTargetSession(const FGuid SessionId)
 {
 	CurrentSelectedSessionId = SessionId;
+
+	if (const TSharedPtr<FSessionInfo> SessionInfo = GetSessionInfo(SessionId))
+	{
+		OnSessionSelected.ExecuteIfBound(SessionInfo);
+	}
 }
 
 FSlateIcon FRecordingControls::GetIconForSession(const FGuid SessionId) const
@@ -270,7 +275,8 @@ const FSlateBrush* FRecordingControls::GetRecordOrStopButton() const
 
 void FRecordingControls::HandleRecordingStop(const TSharedPtr<FSessionInfo>& InSessionInfo)
 {
-	if (DebuggerGuid != InSessionInfo->LastKnownRecordingState.DebuggerId)
+	if (!InSessionInfo.IsValid()
+		|| DebuggerGuid != InSessionInfo->LastKnownRecordingState.DebuggerId)
 	{
 		// Recording was stopped by another debugger, ignore it
 		return;
@@ -284,13 +290,13 @@ void FRecordingControls::HandleRecordingStop(const TSharedPtr<FSessionInfo>& InS
 		LiveSessionEndedMessageHandle = StatusBarSubsystem->PushStatusBarMessage(StatusBarID, LiveSessionEnded);
 	}
 
-	// Allow derived classes to perform extra operations
-	OnRecordingStopped();
+	OnRecordingStopped.ExecuteIfBound(InSessionInfo);
 }
 
 void FRecordingControls::HandleRecordingStart(const TSharedPtr<FSessionInfo>& InSessionInfo)
 {
-	if (DebuggerGuid != InSessionInfo->LastKnownRecordingState.DebuggerId)
+	if (!InSessionInfo.IsValid()
+		|| DebuggerGuid != InSessionInfo->LastKnownRecordingState.DebuggerId)
 	{
 		// Recording was started by another debugger, ignore it
 		return;
@@ -316,8 +322,7 @@ void FRecordingControls::HandleRecordingStart(const TSharedPtr<FSessionInfo>& In
 
 	RecordingMessageHandle = StatusBarSubsystem->PushStatusBarMessage(StatusBarID, LOCTEXT("RecordingMessage", "Recording..."));
 
-	// Allow derived classes to perform extra operations
-	OnRecordingStarted();
+	OnRecordingStarted.ExecuteIfBound(InSessionInfo);
 }
 
 void FRecordingControls::ExecuteAsyncConnectionAttemptTaskWithRetry(
@@ -566,7 +571,7 @@ void FRecordingControls::SetSessionRecordingState(const bool bIsRecording, const
 	if (bIsRecording)
 	{
 		const int32 RemainingRetries = GetMaxConnectionRetriesInternal();
-		const ETraceTransportMode DataTransportOverrideMode = GetTransportModeOverrideForTargetTypeInternal(SessionInfoRef->BuildTargetType);
+		const ETraceTransportMode DataTransportOverrideMode = GetTransportMode(SessionInfoRef->BuildTargetType);
 
 		// Set the session in a busy state so the auto-expire policy for it is relaxed. Starting a recording might cause the remote instance to stall for a few seconds
 		// if the map loaded is too complex
@@ -706,21 +711,18 @@ void FRecordingControls::SetSessionRecordingState(const bool bIsRecording, const
 #endif // WITH_TRACE_BASED_DEBUGGERS
 }
 
-FReply FRecordingControls::ToggleRecordingState()
+void FRecordingControls::ToggleRecordingState(const TSharedPtr<FSessionInfo>& InSessionInfo)
 {
 	using namespace UE::TraceBasedDebuggers;
-	const TSharedPtr<FSessionInfo> SessionInfoPtr = GetCurrentSessionInfo();
-	if (!SessionInfoPtr)
+	if (!InSessionInfo)
 	{
-		return FReply::Handled();
+		return;
 	}
 
-	OnToggleRecordingStateInternal(SessionInfoPtr.ToSharedRef());
-
-	return FReply::Handled();
+	OnToggleRecordingStateInternal(InSessionInfo.ToSharedRef());
 }
 
-bool FRecordingControls::IsRecordingToggleButtonEnabled() const
+bool FRecordingControls::IsRecordingPossibleForSession(const TSharedPtr<FSessionInfo>& InSessionInfo) const
 {
 	using namespace UE::TraceBasedDebuggers;
 	if (CurrentSelectedSessionId == FRemoteSessionsManager::InvalidSessionGUID)
@@ -728,22 +730,21 @@ bool FRecordingControls::IsRecordingToggleButtonEnabled() const
 		return false;
 	}
 
-	const TSharedPtr<FSessionInfo> SessionInfo = GetCurrentSessionInfo();
-	if (!SessionInfo)
+	if (!InSessionInfo)
 	{
 		return false;
 	}
 
 	// Do not allow recording if another trace-based debuggers is already recording in the current session
-	if (SessionInfo->LastKnownRecordingState.DebuggerId.IsValid()
-		&& SessionInfo->LastKnownRecordingState.DebuggerId != DebuggerGuid
-		&& SessionInfo->LastKnownRecordingState.bIsRecording)
+	if (InSessionInfo->LastKnownRecordingState.DebuggerId.IsValid()
+		&& InSessionInfo->LastKnownRecordingState.DebuggerId != DebuggerGuid
+		&& InSessionInfo->LastKnownRecordingState.bIsRecording)
 	{
 		return false;
 	}
 
 	// Allow derived classes to deny
-	if (!IsRecordingToggleButtonEnabledInternal(SessionInfo.ToSharedRef()))
+	if (!IsRecordingAvailableForSessionInternal(InSessionInfo.ToSharedRef()))
 	{
 		return false;
 	}
@@ -812,17 +813,23 @@ void FRecordingControls::AddToMenu(const FName ExistingMenu)
 		}));
 }
 
-int32 FRecordingControls::GetMaxConnectionRetriesInternal()
+int32 FRecordingControls::GetMaxConnectionRetriesInternal() const
 {
 	return DefaultNumberOfConnectionRetries;
 }
 
-ETraceTransportMode FRecordingControls::GetTransportModeOverrideForTargetTypeInternal(EBuildTargetType InBuildTarget)
+ETraceTransportMode FRecordingControls::GetTransportMode(const EBuildTargetType InBuildTarget) const
+{
+	const ETraceTransportMode Override = GetTransportModeOverrideForTargetTypeInternal(InBuildTarget);
+	return Override != ETraceTransportMode::Invalid ? Override : ETraceTransportMode::Direct;
+}
+
+ETraceTransportMode FRecordingControls::GetTransportModeOverrideForTargetTypeInternal(EBuildTargetType InBuildTarget) const
 {
 	return ETraceTransportMode::Invalid;
 }
 
-bool FRecordingControls::IsRecordingToggleButtonEnabledInternal(const TSharedRef<FSessionInfo>& InSessionInfo) const
+bool FRecordingControls::IsRecordingAvailableForSessionInternal(const TSharedRef<FSessionInfo>& InSessionInfo) const
 {
 	return true;
 }
diff --git a/Engine/Source/Runtime/Experimental/TraceBasedDebuggersAnalysis/Public/Widgets/RecordingControls.h b/Engine/Source/Runtime/Experimental/TraceBasedDebuggersAnalysis/Public/Widgets/RecordingControls.h
--- a/Engine/Source/Runtime/Experimental/TraceBasedDebuggersAnalysis/Public/Widgets/RecordingControls.h
+++ b/Engine/Source/Runtime/Experimental/TraceBasedDebuggersAnalysis/Public/Widgets/RecordingControls.h
@@ -1,4 +1,4 @@
-// Copyright Epic Games, Inc. All Rights Reserved.
+﻿// Copyright Epic Games, Inc. All Rights Reserved.
 
 #pragma once
 
@@ -22,6 +22,8 @@ struct FStartRecordingCommandMessage;
 struct FRemoteSessionsManager;
 struct FSessionInfo;
 
+DECLARE_DELEGATE_OneParam(FOnSessionEvent, const TSharedPtr<FSessionInfo>&)
+
 UENUM()
 enum class ELiveConnectionAttemptResult
 {
@@ -43,6 +45,15 @@ struct FRecordingControls : TSharedFromThis<FRecordingControls>
 	UE_API void Deinitialize();
 	UE_API void AddToMenu(FName ExistingMenu);
 
+	/** Delegate executed when recording got started in a given session */
+	FOnSessionEvent OnRecordingStarted;
+
+	/** Delegate executed when recording got stopped in a given session */
+	FOnSessionEvent OnRecordingStopped;
+
+	/** Delegate executed when a given session got selected */
+	FOnSessionEvent OnSessionSelected;
+
 protected:
 
 	virtual void AddToMenuInternal(FToolMenuSection& InSection)
@@ -52,8 +63,9 @@ protected:
 	{
 	}
 
-	UE_API virtual int32 GetMaxConnectionRetriesInternal();
-	UE_API virtual ETraceTransportMode GetTransportModeOverrideForTargetTypeInternal(EBuildTargetType InBuildTarget);
+	UE_API virtual int32 GetMaxConnectionRetriesInternal() const;
+	UE_API ETraceTransportMode GetTransportMode(EBuildTargetType InBuildTarget) const;
+	UE_API virtual ETraceTransportMode GetTransportModeOverrideForTargetTypeInternal(EBuildTargetType InBuildTarget) const;
 
 	virtual bool ConnectToLiveSession_RelayInternal(const FGuid& InSessionId) = 0;
 	virtual bool ConnectToLiveSession_DirectInternal(const FGuid& InSessionId) = 0;
@@ -85,30 +97,32 @@ protected:
 
 	virtual void SendStopRecordingCommandInternal(const FMessageAddress& InAddress) = 0;
 
-	UE_API virtual bool IsRecordingToggleButtonEnabledInternal(const TSharedRef<FSessionInfo>& InSessionInfo) const;
+	UE_API virtual bool IsRecordingAvailableForSessionInternal(const TSharedRef<FSessionInfo>& InSessionInfo) const;
 
 	virtual void OnToggleRecordingStateInternal(const TSharedRef<FSessionInfo>& InSessionInfo) = 0;
 	virtual void CloseSessionByRemoteSessionIDInternal(const FGuid& InSessionId) = 0;
 
 	virtual const FTraceSessionDescriptor& GetCurrentSessionDescriptorInternal() = 0;
 
-	virtual void OnRecordingStopped()
-	{
-	}
+	UE_API void SelectTargetSession(FGuid SessionId);
+	UE_API TSharedPtr<FSessionInfo> GetCurrentSessionInfo() const;
 
-	virtual void OnRecordingStarted()
+	UE_API void ToggleRecordingState(const TSharedPtr<FSessionInfo>& InSessionInfo);
+	FReply ToggleRecordingState()
 	{
+		ToggleRecordingState(GetCurrentSessionInfo());
+		return FReply::Handled();
 	}
 
-	UE_API void SelectTargetSession(FGuid SessionId);
-	UE_API TSharedPtr<FSessionInfo> GetCurrentSessionInfo() const;
-
-	UE_API FReply ToggleRecordingState();
 	UE_API void ToggleSingleSessionRecordingState(const TSharedRef<FSessionInfo>& InSessionInfo);
 
 	UE_API void SetSessionRecordingState(bool bIsRecording, const TSharedRef<FSessionInfo>& SessionInfoRef);
 	UE_API bool IsRecording() const;
-	UE_API bool IsRecordingToggleButtonEnabled() const;
+	UE_API bool IsRecordingPossibleForSession(const TSharedPtr<FSessionInfo>& InSessionInfo) const;
+	bool IsRecordingToggleButtonEnabled() const
+	{
+		return IsRecordingPossibleForSession(GetCurrentSessionInfo());
+	}
 
 	UE_API FSlateIcon GetIconForSession(FGuid SessionId) const;
 
