[RewindDebugger] Replaced included header by forward declaration for TraceService::Frame

#Synced-CL 49884344

diff --git a/Engine/Plugins/Animation/GameplayInsights/Source/GameplayInsights/Private/PropertyTrack.h b/Engine/Plugins/Animation/GameplayInsights/Source/GameplayInsights/Private/PropertyTrack.h
--- a/Engine/Plugins/Animation/GameplayInsights/Source/GameplayInsights/Private/PropertyTrack.h
+++ b/Engine/Plugins/Animation/GameplayInsights/Source/GameplayInsights/Private/PropertyTrack.h
@@ -5,6 +5,11 @@
 #include "RewindDebuggerTrack.h"
 #include "IGameplayProvider.h"
 
+namespace TraceServices
+{
+struct FFrame;
+}
+
 namespace RewindDebugger
 {
 	class FPropertiesTrack;
diff --git a/Engine/Plugins/Animation/GameplayInsights/Source/RewindDebuggerVLog/Private/VisualLogTrack.cpp b/Engine/Plugins/Animation/GameplayInsights/Source/RewindDebuggerVLog/Private/VisualLogTrack.cpp
--- a/Engine/Plugins/Animation/GameplayInsights/Source/RewindDebuggerVLog/Private/VisualLogTrack.cpp
+++ b/Engine/Plugins/Animation/GameplayInsights/Source/RewindDebuggerVLog/Private/VisualLogTrack.cpp
@@ -6,6 +6,7 @@
 #include "PropertyEditorModule.h"
 #include "VisualLoggerProvider.h"
 #include "SVisualLoggerTableRow.h"
+#include "TraceServices/Model/Frames.h"
 #include "UObject/Package.h"
 
 #define LOCTEXT_NAMESPACE "VisualLogTrack"
diff --git a/Engine/Plugins/Animation/GameplayInsights/Source/RewindDebuggerVLog/Public/IVisualLoggerProvider.h b/Engine/Plugins/Animation/GameplayInsights/Source/RewindDebuggerVLog/Public/IVisualLoggerProvider.h
--- a/Engine/Plugins/Animation/GameplayInsights/Source/RewindDebuggerVLog/Public/IVisualLoggerProvider.h
+++ b/Engine/Plugins/Animation/GameplayInsights/Source/RewindDebuggerVLog/Public/IVisualLoggerProvider.h
@@ -1,11 +1,19 @@
-// Copyright Epic Games, Inc. All Rights Reserved.
+﻿// Copyright Epic Games, Inc. All Rights Reserved.
 
 #pragma once
 
+#if UE_ENABLE_INCLUDE_ORDER_DEPRECATED_IN_5_8
 #include "CoreMinimal.h"
-#include "TraceServices/Model/AnalysisSession.h"
+#endif // UE_ENABLE_INCLUDE_ORDER_DEPRECATED_IN_5_8
+
 #include "TraceServices/Containers/Timelines.h"
+#include "TraceServices/Model/AnalysisSession.h"
+
+#if UE_ENABLE_INCLUDE_ORDER_DEPRECATED_IN_5_8
 #include "VisualLogger/VisualLoggerTypes.h"
+#endif // UE_ENABLE_INCLUDE_ORDER_DEPRECATED_IN_5_8
+
+struct FVisualLogEntry;
 
 class IVisualLoggerProvider : public TraceServices::IProvider
 {
diff --git a/Engine/Plugins/Animation/PoseSearch/Source/Editor/Private/PoseSearchDebuggerViewModel.cpp b/Engine/Plugins/Animation/PoseSearch/Source/Editor/Private/PoseSearchDebuggerViewModel.cpp
--- a/Engine/Plugins/Animation/PoseSearch/Source/Editor/Private/PoseSearchDebuggerViewModel.cpp
+++ b/Engine/Plugins/Animation/PoseSearch/Source/Editor/Private/PoseSearchDebuggerViewModel.cpp
@@ -14,6 +14,7 @@
 #include "PoseSearch/PoseSearchDerivedData.h"
 #include "PoseSearch/PoseSearchSchema.h"
 #include "Trace/PoseSearchTraceProvider.h"
+#include "TraceServices/Model/Frames.h"
 
 namespace UE::PoseSearch
 {
diff --git a/Engine/Plugins/Cameras/GameplayCameras/Source/GameplayCamerasEditor/Private/Trace/CameraSystemRewindDebuggerExtension.cpp b/Engine/Plugins/Cameras/GameplayCameras/Source/GameplayCamerasEditor/Private/Trace/CameraSystemRewindDebuggerExtension.cpp
--- a/Engine/Plugins/Cameras/GameplayCameras/Source/GameplayCamerasEditor/Private/Trace/CameraSystemRewindDebuggerExtension.cpp
+++ b/Engine/Plugins/Cameras/GameplayCameras/Source/GameplayCamerasEditor/Private/Trace/CameraSystemRewindDebuggerExtension.cpp
@@ -13,6 +13,7 @@
 #include "Modules/ModuleManager.h"
 #include "SLevelViewport.h"
 #include "Trace/CameraSystemTraceProvider.h"
+#include "TraceServices/Model/Frames.h"
 
 #if UE_GAMEPLAY_CAMERAS_TRACE
 
diff --git a/Engine/Plugins/Chooser/Source/ChooserEditor/Private/RewindDebuggerChooser.cpp b/Engine/Plugins/Chooser/Source/ChooserEditor/Private/RewindDebuggerChooser.cpp
--- a/Engine/Plugins/Chooser/Source/ChooserEditor/Private/RewindDebuggerChooser.cpp
+++ b/Engine/Plugins/Chooser/Source/ChooserEditor/Private/RewindDebuggerChooser.cpp
@@ -6,6 +6,7 @@
 #include "ChooserProvider.h"
 #include "GameFramework/Actor.h"
 #include "IGameplayProvider.h"
+#include "TraceServices/Model/Frames.h"
 
 FRewindDebuggerChooser::FRewindDebuggerChooser()
 {
diff --git a/Engine/Plugins/Developer/FieldNotificationTrace/Source/FieldNotificationTraceEditor/Private/FieldNotificationRewindDebugger.cpp b/Engine/Plugins/Developer/FieldNotificationTrace/Source/FieldNotificationTraceEditor/Private/FieldNotificationRewindDebugger.cpp
--- a/Engine/Plugins/Developer/FieldNotificationTrace/Source/FieldNotificationTraceEditor/Private/FieldNotificationRewindDebugger.cpp
+++ b/Engine/Plugins/Developer/FieldNotificationTrace/Source/FieldNotificationTraceEditor/Private/FieldNotificationRewindDebugger.cpp
@@ -3,7 +3,7 @@
 #include "FieldNotificationRewindDebugger.h"
 #include "FieldNotificationTraceProvider.h"
 #include "IRewindDebugger.h"
-//#include "IGameplayProvider.h"
+#include "TraceServices/Model/Frames.h"
 
 namespace UE::FieldNotification
 {
diff --git a/Engine/Plugins/Experimental/UAF/UAF/Source/UAFEditor/Private/RewindDebugger/AnimNextModuleTrack.cpp b/Engine/Plugins/Experimental/UAF/UAF/Source/UAFEditor/Private/RewindDebugger/AnimNextModuleTrack.cpp
--- a/Engine/Plugins/Experimental/UAF/UAF/Source/UAFEditor/Private/RewindDebugger/AnimNextModuleTrack.cpp
+++ b/Engine/Plugins/Experimental/UAF/UAF/Source/UAFEditor/Private/RewindDebugger/AnimNextModuleTrack.cpp
@@ -18,6 +18,7 @@
 #include "Serialization/MemoryReader.h"
 #include "StructUtils/PropertyBag.h"
 #include "Subsystems/AssetEditorSubsystem.h"
+#include "TraceServices/Model/Frames.h"
 
 #define LOCTEXT_NAMESPACE "AnimNextModuleTrack"
 
diff --git a/Engine/Plugins/Experimental/UAF/UAF/Source/UAFEditor/Private/UAF/Debug/SVariableWatch.cpp b/Engine/Plugins/Experimental/UAF/UAF/Source/UAFEditor/Private/UAF/Debug/SVariableWatch.cpp
--- a/Engine/Plugins/Experimental/UAF/UAF/Source/UAFEditor/Private/UAF/Debug/SVariableWatch.cpp
+++ b/Engine/Plugins/Experimental/UAF/UAF/Source/UAFEditor/Private/UAF/Debug/SVariableWatch.cpp
@@ -19,6 +19,7 @@
 #include "SceneOutlinerPublicTypes.h"
 #include "SSceneOutliner.h"
 #include "TraceServices/Model/AnalysisSession.h"
+#include "TraceServices/Model/Frames.h"
 #include "UAF/Workspace/IDebugObjectSelector.h"
 #include "UAF/Workspace/RewindDebuggerExtension.h"
 #include "Variables/Outliner/VariablesOutlinerColumns.h"
diff --git a/Engine/Plugins/Experimental/UAF/UAFAnimGraph/Source/UAFAnimGraphEditor/Private/RewindDebugger/EvaluationProgramTrack.cpp b/Engine/Plugins/Experimental/UAF/UAFAnimGraph/Source/UAFAnimGraphEditor/Private/RewindDebugger/EvaluationProgramTrack.cpp
--- a/Engine/Plugins/Experimental/UAF/UAFAnimGraph/Source/UAFAnimGraphEditor/Private/RewindDebugger/EvaluationProgramTrack.cpp
+++ b/Engine/Plugins/Experimental/UAF/UAFAnimGraph/Source/UAFAnimGraphEditor/Private/RewindDebugger/EvaluationProgramTrack.cpp
@@ -12,6 +12,7 @@
 #include "ObjectAsTraceIdProxyArchiveReader.h"
 #include "Component/AnimNextComponent.h"
 #include "Serialization/ObjectReader.h"
+#include "TraceServices/Model/Frames.h"
 
 #include UE_INLINE_GENERATED_CPP_BY_NAME(EvaluationProgramTrack)
 
diff --git a/Engine/Plugins/Experimental/UAF/UAFAnimGraph/Source/UAFAnimGraphEditor/Private/RewindDebugger/SequenceInfoTrack.cpp b/Engine/Plugins/Experimental/UAF/UAFAnimGraph/Source/UAFAnimGraphEditor/Private/RewindDebugger/SequenceInfoTrack.cpp
--- a/Engine/Plugins/Experimental/UAF/UAFAnimGraph/Source/UAFAnimGraphEditor/Private/RewindDebugger/SequenceInfoTrack.cpp
+++ b/Engine/Plugins/Experimental/UAF/UAFAnimGraph/Source/UAFAnimGraphEditor/Private/RewindDebugger/SequenceInfoTrack.cpp
@@ -1,4 +1,4 @@
-// Copyright Epic Games, Inc. All Rights Reserved.
+﻿// Copyright Epic Games, Inc. All Rights Reserved.
 
 #include "SequenceInfoTrack.h"
 
@@ -13,6 +13,7 @@
 #include "EvaluationVM/Tasks/PushAnimSequenceKeyframe.h"
 #include "ObjectAsTraceIdProxyArchiveReader.h"
 #include "Component/AnimNextComponent.h"
+#include "TraceServices/Model/Frames.h"
 #include "Widgets/SCanvas.h"
 #include "Widgets/Text/STextBlock.h"
 #include "Widgets/Colors/SColorBlock.h"
diff --git a/Engine/Plugins/Runtime/RigVM/Source/RigVMEditor/Private/Trace/RigVMTraceModule.cpp b/Engine/Plugins/Runtime/RigVM/Source/RigVMEditor/Private/Trace/RigVMTraceModule.cpp
--- a/Engine/Plugins/Runtime/RigVM/Source/RigVMEditor/Private/Trace/RigVMTraceModule.cpp
+++ b/Engine/Plugins/Runtime/RigVM/Source/RigVMEditor/Private/Trace/RigVMTraceModule.cpp
@@ -3,6 +3,7 @@
 #include "Trace/RigVMTraceModule.h"
 #include "RigVMTraceProvider.h"
 #include "RigVMTraceAnalyzer.h"
+#include "TraceServices/Model/Frames.h"
 
 #if RIGVM_TRACE_ENABLED
 
diff --git a/Engine/Plugins/Runtime/RigVM/Source/RigVMEditor/Private/Trace/RigVMTraceTrack.cpp b/Engine/Plugins/Runtime/RigVM/Source/RigVMEditor/Private/Trace/RigVMTraceTrack.cpp
--- a/Engine/Plugins/Runtime/RigVM/Source/RigVMEditor/Private/Trace/RigVMTraceTrack.cpp
+++ b/Engine/Plugins/Runtime/RigVM/Source/RigVMEditor/Private/Trace/RigVMTraceTrack.cpp
@@ -5,6 +5,7 @@
 #include "RigVMTraceAnalyzer.h"
 #include "DetailLayoutBuilder.h"
 #include "RigVMHost.h"
+#include "TraceServices/Model/Frames.h"
 #include "Widgets/Layout/SGridPanel.h"
 #include "Widgets/Input/SNumericEntryBox.h"
 
diff --git a/Engine/Source/Editor/RewindDebuggerInterface/Public/IRewindDebugger.h b/Engine/Source/Editor/RewindDebuggerInterface/Public/IRewindDebugger.h
--- a/Engine/Source/Editor/RewindDebuggerInterface/Public/IRewindDebugger.h
+++ b/Engine/Source/Editor/RewindDebuggerInterface/Public/IRewindDebugger.h
@@ -4,9 +4,14 @@
 #if UE_ENABLE_INCLUDE_ORDER_DEPRECATED_IN_5_7
 #include "CoreMinimal.h"
 #endif // UE_ENABLE_INCLUDE_ORDER_DEPRECATED_IN_5_7
+
 #include "RewindDebuggerTypes.h"
+
+#if UE_ENABLE_INCLUDE_ORDER_DEPRECATED_IN_5_8
 #include "TraceServices/Model/Frames.h"
 #include "UObject/NameTypes.h"
+#endif // UE_ENABLE_INCLUDE_ORDER_DEPRECATED_IN_5_8
+
 #include "IRewindDebugger.generated.h"
 
 REWINDDEBUGGERINTERFACE_API DECLARE_LOG_CATEGORY_EXTERN(LogRewindDebugger, Log, All);
