Fix for world time reset bug

#Synced-CL 49873372

diff --git a/Engine/Plugins/Animation/GameplayInsights/Source/RewindDebuggerRuntime/Private/RewindDebuggerRuntime.cpp b/Engine/Plugins/Animation/GameplayInsights/Source/RewindDebuggerRuntime/Private/RewindDebuggerRuntime.cpp
--- a/Engine/Plugins/Animation/GameplayInsights/Source/RewindDebuggerRuntime/Private/RewindDebuggerRuntime.cpp
+++ b/Engine/Plugins/Animation/GameplayInsights/Source/RewindDebuggerRuntime/Private/RewindDebuggerRuntime.cpp
@@ -64,9 +64,7 @@ namespace RewindDebugger
 
 	void FRewindDebuggerRuntime::OnEnableRequiredTraceChannelsInternal()
 	{
-		UE::Trace::ToggleChannel(TEXT("Object"), true);
-		UE::Trace::ToggleChannel(TEXT("ObjectProperties"), true);
-		UE::Trace::ToggleChannel(TEXT("Frame"), true);
+
 	}
 
 	void FRewindDebuggerRuntime::OnRecordingStartingInternal()
@@ -94,20 +92,23 @@ namespace RewindDebugger
 		// is why we defer the callbacks.
 		ExecuteOnGameThread(UE_SOURCE_LOCATION, [this]
 			{
-				// update extensions
-				IterateExtensions([](IRewindDebuggerRuntimeExtension* Extension)
-					{
-						Extension->RecordingStarted();
-					}
-				);
-
 #if OBJECT_TRACE_ENABLED
 				// trace each play-in-editor world, and all the actors in it.
 				for (TObjectIterator<UWorld> World; World; ++World)
 				{
 
 					FObjectTrace::ResetWorldElapsedTime(*World);
+				}
+
+#endif // OBJECT_TRACE_ENABLED
+
+				UE::Trace::ToggleChannel(TEXT("Object"), true);
+				UE::Trace::ToggleChannel(TEXT("ObjectProperties"), true);
+				UE::Trace::ToggleChannel(TEXT("Frame"), true);
 
+
+				for (TObjectIterator<UWorld> World; World; ++World)
+				{
 					TRACE_WORLD(*World);
 
 					for (TActorIterator<AController> Iterator(*World); Iterator; ++Iterator)
@@ -118,7 +119,13 @@ namespace RewindDebugger
 						}
 					}
 				}
-#endif // OBJECT_TRACE_ENABLED
+
+				// update extensions
+				IterateExtensions([](IRewindDebuggerRuntimeExtension* Extension)
+					{
+						Extension->RecordingStarted();
+					}
+				);
 			});
 
 		RecordingStarted.Broadcast();
