[RewindDebugger] added new virtual methods for runtime extensions to register new message types and handlers for remote connections
- added new log categories widget to enable/disable verbosity filtering during recording and control the verbosity of each category (note that the initial implementation populates the "known" categories from the analyzed data of an active recording)

#Synced-CL 49922632

diff --git a/Engine/Plugins/Animation/GameplayInsights/Source/RewindDebugger/Private/RewindDebugger.cpp b/Engine/Plugins/Animation/GameplayInsights/Source/RewindDebugger/Private/RewindDebugger.cpp
--- a/Engine/Plugins/Animation/GameplayInsights/Source/RewindDebugger/Private/RewindDebugger.cpp
+++ b/Engine/Plugins/Animation/GameplayInsights/Source/RewindDebugger/Private/RewindDebugger.cpp
@@ -1048,6 +1048,11 @@ void FRewindDebugger::GetScrubTimeInformation(double InDebugTime, FScrubTimeInfo
 	}
 }
 
+TSharedPtr<UE::TraceBasedDebuggers::FSessionInfo> FRewindDebugger::GetSelectedDebugSessionInfo() const
+{
+	return RecordingControls ? RecordingControls->GetCurrentSessionInfo() : nullptr;
+}
+
 const TraceServices::IAnalysisSession* FRewindDebugger::GetAnalysisSession() const
 {
 	return GetAnalysisSessionAsShared().Get();
diff --git a/Engine/Plugins/Animation/GameplayInsights/Source/RewindDebugger/Private/RewindDebugger.h b/Engine/Plugins/Animation/GameplayInsights/Source/RewindDebugger/Private/RewindDebugger.h
--- a/Engine/Plugins/Animation/GameplayInsights/Source/RewindDebugger/Private/RewindDebugger.h
+++ b/Engine/Plugins/Animation/GameplayInsights/Source/RewindDebugger/Private/RewindDebugger.h
@@ -1,4 +1,4 @@
-// Copyright Epic Games, Inc. All Rights Reserved.
+﻿// Copyright Epic Games, Inc. All Rights Reserved.
 #pragma once
 
 #if UE_ENABLE_INCLUDE_ORDER_DEPRECATED_IN_5_7
@@ -68,6 +68,7 @@ public:
 		return CurrentViewRange;
 	}
 
+	virtual TSharedPtr<UE::TraceBasedDebuggers::FSessionInfo> GetSelectedDebugSessionInfo() const override;
 	virtual const TraceServices::IAnalysisSession* GetAnalysisSession() const override;
 	virtual TSharedPtr<const TraceServices::IAnalysisSession> GetAnalysisSessionAsShared() const override;
 	virtual uint64 GetRootObjectId() const override;
diff --git a/Engine/Plugins/Animation/GameplayInsights/Source/RewindDebuggerRuntime/Private/RewindDebuggerRemoteSessionsHandler.cpp b/Engine/Plugins/Animation/GameplayInsights/Source/RewindDebuggerRuntime/Private/RewindDebuggerRemoteSessionsHandler.cpp
--- a/Engine/Plugins/Animation/GameplayInsights/Source/RewindDebuggerRuntime/Private/RewindDebuggerRemoteSessionsHandler.cpp
+++ b/Engine/Plugins/Animation/GameplayInsights/Source/RewindDebuggerRuntime/Private/RewindDebuggerRemoteSessionsHandler.cpp
@@ -3,8 +3,10 @@
 #if WITH_TRACE_BASED_DEBUGGERS
 
 #include "RewindDebuggerRemoteSessionsHandler.h"
+#include "Features/IModularFeatures.h"
 #include "MessageEndpointBuilder.h"
 #include "RewindDebuggerRuntime/RewindDebuggerRuntime.h"
+#include "RewindDebuggerRuntimeInterface/IRewindDebuggerRuntimeExtension.h"
 #include "RewindDebuggerRuntimeModule.h"
 #include "SessionInfo.h"
 
@@ -12,7 +14,7 @@ namespace UE::RewindDebugger
 {
 
 void FSessionsHandler::OnCreatingMessageEndpoint(
-	const TSharedPtr<TraceBasedDebuggers::FRemoteSessionsManager>& SessionsManager
+	const TSharedPtr<TraceBasedDebuggers::FRemoteSessionsManager>& InSessionsManager
 	, FMessageEndpointBuilder& InEndPointBuilder)
 {
 	InEndPointBuilder
@@ -36,6 +38,14 @@ void FSessionsHandler::OnCreatingMessageEndpoint(
 					}
 				};
 			});
+
+	// Allow extensions to register their message handlers
+	::RewindDebugger::IterateExtensions([&InSessionsManager, &InEndPointBuilder](IRewindDebuggerRuntimeExtension* Extension)
+		{
+			Extension->RegisterMessageHandlers(InSessionsManager, InEndPointBuilder);
+		}
+	);
+
 }
 
 void FSessionsHandler::OnSubscribingMessageTypes(
@@ -50,6 +60,14 @@ void FSessionsHandler::OnSubscribingMessageTypes(
 
 	InMessageEndpoint.Subscribe<FStartRecordingCommand>();
 	InMessageEndpoint.Subscribe<FStopRecordingCommand>();
+
+	// Allow extensions to register their message types
+	::RewindDebugger::IterateExtensions([&InSessionsManager, &InMessageEndpoint](IRewindDebuggerRuntimeExtension* Extension)
+		{
+			Extension->RegisterMessageTypes(InSessionsManager, InMessageEndpoint);
+		}
+	);
+
 }
 
 void FSessionsHandler::OnBuildingFullSessionInfoResponseMessage(
diff --git a/Engine/Plugins/Animation/GameplayInsights/Source/RewindDebuggerRuntime/Private/RewindDebuggerRuntime.cpp b/Engine/Plugins/Animation/GameplayInsights/Source/RewindDebuggerRuntime/Private/RewindDebuggerRuntime.cpp
--- a/Engine/Plugins/Animation/GameplayInsights/Source/RewindDebuggerRuntime/Private/RewindDebuggerRuntime.cpp
+++ b/Engine/Plugins/Animation/GameplayInsights/Source/RewindDebuggerRuntime/Private/RewindDebuggerRuntime.cpp
@@ -1,4 +1,4 @@
-// Copyright Epic Games, Inc. All Rights Reserved.
+﻿// Copyright Epic Games, Inc. All Rights Reserved.
 
 #if WITH_TRACE_BASED_DEBUGGERS
 
@@ -49,19 +49,6 @@ namespace RewindDebugger
 			});
 	}
 
-	static void IterateExtensions(TFunction<void(IRewindDebuggerRuntimeExtension* Extension)> IteratorFunction)
-	{
-		// update extensions
-		IModularFeatures& ModularFeatures = IModularFeatures::Get();
-
-		const int32 NumExtensions = ModularFeatures.GetModularFeatureImplementationCount(IRewindDebuggerRuntimeExtension::ModularFeatureName);
-		for (int32 ExtensionIndex = 0; ExtensionIndex < NumExtensions; ++ExtensionIndex)
-		{
-			IRewindDebuggerRuntimeExtension* Extension = static_cast<IRewindDebuggerRuntimeExtension*>(ModularFeatures.GetModularFeatureImplementation(IRewindDebuggerRuntimeExtension::ModularFeatureName, ExtensionIndex));
-			IteratorFunction(Extension);
-		}
-	}
-
 	void FRewindDebuggerRuntime::OnEnableRequiredTraceChannelsInternal()
 	{
 
diff --git a/Engine/Plugins/Animation/GameplayInsights/Source/RewindDebuggerVLog/Private/RewindDebuggerVLog.cpp b/Engine/Plugins/Animation/GameplayInsights/Source/RewindDebuggerVLog/Private/RewindDebuggerVLog.cpp
--- a/Engine/Plugins/Animation/GameplayInsights/Source/RewindDebuggerVLog/Private/RewindDebuggerVLog.cpp
+++ b/Engine/Plugins/Animation/GameplayInsights/Source/RewindDebuggerVLog/Private/RewindDebuggerVLog.cpp
@@ -9,8 +9,10 @@
 #include "LogVisualizerSettings.h"
 #include "ObjectTrace.h"
 #include "RewindDebuggerVLogSettings.h"
+#include "SceneView.h"
 #include "ToolMenus.h"
 #include "VisualLogEntryRenderer.h"
+#include "VLogRenderingActor.h"
 #include "Debug/DebugDrawService.h"
 #include "Editor/EditorEngine.h"
 #include "Engine/Canvas.h"
@@ -18,10 +20,16 @@
 #include "HAL/IConsoleManager.h"
 #include "TraceServices/Model/Frames.h"
 #include "VisualLogger/VisualLoggerTraceDevice.h"
+#include "Widgets/LogCategoryVerbositySelector.h"
 
 #define LOCTEXT_NAMESPACE "RewindDebuggerVLog"
 
+namespace UE::RewindDebugger::VLog
+{
+const FName MenuSectionName("Visual Logger");
+
 TAutoConsoleVariable<int32> CVarRewindDebuggerVLogUseActor(TEXT("a.RewindDebugger.VisualLogs.UseActor"), 0, TEXT("Use actor based debug renderer for visual logs"));
+}
 
 void FRewindDebuggerVLog::OnShowDebugInfo(UCanvas* Canvas, APlayerController* Player)
 {
@@ -52,7 +60,7 @@ void FRewindDebuggerVLog::OnShowDebugInfo(UCanvas* Canvas, APlayerController* Pl
 				const TraceServices::IFrameProvider& FrameProvider = TraceServices::ReadFrameProvider(*Session);
 				TraceServices::FFrame CurrentFrame;
 
-				if (FrameProvider.GetFrameFromTime(ETraceFrameType::TraceFrameType_Game, CurrentTraceTime, CurrentFrame))
+				if (FrameProvider.GetFrameFromTime(TraceFrameType_Game, CurrentTraceTime, CurrentFrame))
 				{
 					if (const IVisualLoggerProvider* VisualLoggerProvider = Session->ReadProvider<IVisualLoggerProvider>("VisualLoggerProvider"))
 					{
@@ -67,34 +75,59 @@ void FRewindDebuggerVLog::OnShowDebugInfo(UCanvas* Canvas, APlayerController* Pl
 FRewindDebuggerVLog::~FRewindDebuggerVLog()
 {
 	UDebugDrawService::Unregister(DelegateHandle);
+
+	if (const IRewindDebuggerModule* RewindDebuggerModule = FModuleManager::GetModulePtr<IRewindDebuggerModule>("RewindDebugger"))
+	{
+		if (UToolMenus* Menus = UToolMenus::Get())
+		{
+			if (UToolMenu* Menu = Menus->FindMenu(RewindDebuggerModule->GetMainToolbarName()))
+			{
+				Menu->RemoveSection(UE::RewindDebugger::VLog::MenuSectionName);
+			}
+		}
+	}
+
+	if (ToolMenuRegistrationHandle.IsValid())
+	{
+		UToolMenus::UnRegisterStartupCallback(ToolMenuRegistrationHandle);
+	}
 }
 
 void FRewindDebuggerVLog::Initialize()
 {
-	const IRewindDebuggerModule* RewindDebuggerModule = FModuleManager::GetModulePtr<IRewindDebuggerModule>("RewindDebugger");
-	if (!ensure(RewindDebuggerModule))
-	{
-		return;
-	}
+	ToolMenuRegistrationHandle = UToolMenus::RegisterStartupCallback(FSimpleMulticastDelegate::FDelegate::CreateLambda([this]()
+		{
+			const IRewindDebuggerModule* RewindDebuggerModule = FModuleManager::GetModulePtr<IRewindDebuggerModule>("RewindDebugger");
+			if (!ensure(RewindDebuggerModule))
+			{
+				return;
+			}
+
+			UToolMenu* Menu = UToolMenus::Get()->ExtendMenu(RewindDebuggerModule->GetMainToolbarName());
 
-	UToolMenu* Menu = UToolMenus::Get()->FindMenu(RewindDebuggerModule->GetMainToolbarName());
+			FToolMenuSection& AnalyzerSection = Menu->AddSection(UE::RewindDebugger::VLog::MenuSectionName, LOCTEXT("Visual Logger", "Visual Logger"));
 
-	FToolMenuSection& NewSection = Menu->AddSection("Visual Logger", LOCTEXT("Visual Logger", "Visual Logger"));
+			AnalyzerSection.AddSeparator("VisualLogger");
 
-	NewSection.AddSeparator("VisualLogger");
+			AnalyzerSection.AddEntry(FToolMenuEntry::InitComboButton(
+				"VLog Categories",
+				FUIAction(),
+				FNewToolMenuDelegate::CreateRaw(this, &FRewindDebuggerVLog::MakeCategoriesMenu),
+				LOCTEXT("VLog Categories", "VLog Categories")
+			));
 
-	NewSection.AddEntry(FToolMenuEntry::InitComboButton(
-		"VLog Categories",
-		FUIAction(),
-		FNewToolMenuDelegate::CreateRaw(this, &FRewindDebuggerVLog::MakeCategoriesMenu),
-		LOCTEXT("VLog Categories", "VLog Categories")
-	));
+			AnalyzerSection.AddEntry(FToolMenuEntry::InitComboButton(
+				"VLog Level",
+				FUIAction(),
+				FNewToolMenuDelegate::CreateRaw(this, &FRewindDebuggerVLog::MakeLogLevelMenu),
+				LOCTEXT("VLog Level", "VLog Level")));
 
-	NewSection.AddEntry(FToolMenuEntry::InitComboButton(
-		"VLog Level",
-		FUIAction(),
-		FNewToolMenuDelegate::CreateRaw(this, &FRewindDebuggerVLog::MakeLogLevelMenu),
-		LOCTEXT("VLog Level", "VLog Level")));
+			FToolMenuSection& RecordingSection = Menu->AddSection("Visual Logger Recording", LOCTEXT("Visual Logger", "Visual Logger"), FToolMenuInsert("RecordingControls", EToolMenuInsertType::After));
+			RecordingSection.AddEntry(FToolMenuEntry::InitWidget(
+				"Recording Log Verbosity",
+				SNew(UE::RewindDebugger::SLogCategoryFilter),
+				LOCTEXT("VLogRecordingVerbosityLevels", "VLog Recording Verbosity Levels")));
+		}));
 
 	FVisualLoggerTraceDevice& TraceDevice = FVisualLoggerTraceDevice::Get();
 	TraceDevice.ImmediateRenderDelegate.BindRaw(this, &FRewindDebuggerVLog::ImmediateRender);
@@ -116,7 +149,7 @@ bool MatchCategoryFilters(const FName& CategoryName, ELogVerbosity::Type Verbosi
 
 void FRewindDebuggerVLog::RenderLogEntry(const FVisualLogEntry& Entry, UCanvas* Canvas)
 {
-	if (CVarRewindDebuggerVLogUseActor.GetValueOnAnyThread())
+	if (UE::RewindDebugger::VLog::CVarRewindDebuggerVLogUseActor.GetValueOnAnyThread())
 	{
 		// old actor based code path
 		if (AVLogRenderingActor* RenderingActor = GetRenderingActor())
diff --git a/Engine/Plugins/Animation/GameplayInsights/Source/RewindDebuggerVLog/Private/RewindDebuggerVLog.h b/Engine/Plugins/Animation/GameplayInsights/Source/RewindDebuggerVLog/Private/RewindDebuggerVLog.h
--- a/Engine/Plugins/Animation/GameplayInsights/Source/RewindDebuggerVLog/Private/RewindDebuggerVLog.h
+++ b/Engine/Plugins/Animation/GameplayInsights/Source/RewindDebuggerVLog/Private/RewindDebuggerVLog.h
@@ -1,22 +1,36 @@
 // Copyright Epic Games, Inc. All Rights Reserved.
 #pragma once
 
-#include "IRewindDebugger.h"
 #include "IRewindDebuggerExtension.h"
-#include "VLogRenderingActor.h"
-#include "ToolMenu.h"
 // Needed for TStrongObjectPtr which requires complete type
 #include "Engine/Font.h"
 
-// Rewind debugger extension for Visual Logger support
+struct FDebugObjectInfo;
+struct FVisualLogEntry;
+class AVLogRenderingActor;
+class IRewindDebugger;
+class UCanvas;
+class UToolMenu;
 
+namespace UE::TraceBasedDebuggers
+{
+struct FRecordingControls;
+}
+
+namespace UE::RewindDebugger
+{
+struct FLogCategoryVerbosity;
+struct FDebuggerRecordingControls;
+}
+
+/** Rewind debugger extension for Visual Logger support */
 class FRewindDebuggerVLog : public IRewindDebuggerExtension
 {
 public:
 	void OnShowDebugInfo(UCanvas* Canvas, APlayerController* Player);
 	virtual ~FRewindDebuggerVLog();
 
-	virtual FString GetName()
+	virtual FString GetName() override
 	{
 		return TEXT("FRewindDebuggerVLog");
 	}
@@ -47,6 +61,7 @@ private:
 	int32 ScreenTextY = 0;
 
 	FDelegateHandle DelegateHandle;
+	FDelegateHandle ToolMenuRegistrationHandle;
 	TStrongObjectPtr<UFont> MonospaceFont;
 
 	TArray<uint64> DebuggedObjectIds;
diff --git a/Engine/Plugins/Animation/GameplayInsights/Source/RewindDebuggerVLog/Private/Widgets/LogCategoryVerbositySelector.cpp b/Engine/Plugins/Animation/GameplayInsights/Source/RewindDebuggerVLog/Private/Widgets/LogCategoryVerbositySelector.cpp
new file mode 100644
--- /dev/null
+++ b/Engine/Plugins/Animation/GameplayInsights/Source/RewindDebuggerVLog/Private/Widgets/LogCategoryVerbositySelector.cpp
@@ -0,0 +1,552 @@
+// Copyright Epic Games, Inc. All Rights Reserved.
+
+#include "LogCategoryVerbositySelector.h"
+
+#include "Components/VerticalBox.h"
+#include "IRewindDebugger.h"
+#include "IVisualLoggerProvider.h"
+#include "RemoteSessionsManager.h"
+#include "RewindDebuggerEngineEditorBridge.h"
+#include "RewindDebuggerVLogRuntimeTypes.h"
+#include "SessionInfo.h"
+#include "Styling/AppStyle.h"
+#include "TraceServices/Model/AnalysisSession.h"
+#include "Types/SlateEnums.h"
+#include "Widgets/DeclarativeSyntaxSupport.h"
+#include "Widgets/Images/SImage.h"
+#include "Widgets/Images/SThrobber.h"
+#include "Widgets/Input/SCheckBox.h"
+#include "Widgets/Input/SComboBox.h"
+#include "Widgets/Input/SComboButton.h"
+#include "Widgets/Input/SSearchBox.h"
+#include "Widgets/Layout/SWidgetSwitcher.h"
+#include "Widgets/SBoxPanel.h"
+#include "Widgets/Views/SListView.h"
+
+#define LOCTEXT_NAMESPACE "RewindDebuggerVLog"
+
+namespace UE::RewindDebugger
+{
+
+FVLogExtensionSessionData* GetVLogExtensionData()
+{
+	using namespace UE::TraceBasedDebuggers;
+	if (const IRewindDebugger* RewindDebugger = IRewindDebugger::Instance())
+	{
+		const TSharedPtr<FRemoteSessionsManager> RemoteSessionsManager = FRewindDebuggerEngineEditorBridge::Get().GetSessionsManager();
+		const TSharedPtr<FSessionInfo> SessionInfo = RewindDebugger->GetSelectedDebugSessionInfo();
+		if (SessionInfo && RemoteSessionsManager)
+		{
+			FVLogExtensionSessionData* DebuggerData = SessionInfo->GetDebuggerData<FVLogExtensionSessionData>();
+
+			// Create if missing
+			if (DebuggerData == nullptr)
+			{
+				SessionInfo->SetDebuggerData(FVLogExtensionSessionData{});
+				DebuggerData = SessionInfo->GetDebuggerData<FVLogExtensionSessionData>();
+			}
+			return DebuggerData;
+		}
+	}
+	return nullptr;
+}
+
+FText GetVerbosityAsText(const ELogVerbosity::Type InVerbosity)
+{
+	if (InVerbosity == UnknownVerbosity)
+	{
+		return LOCTEXT("UnknownVerbosity", "(unknown)");
+	}
+
+	return FText::FromString(ToString(InVerbosity));
+}
+
+FSlateFontInfo GetVerbosityFont(const ELogVerbosity::Type InVerbosity)
+{
+	if (InVerbosity == UnknownVerbosity)
+	{
+		return FAppStyle::GetFontStyle("NormalFontItalic");
+	}
+
+	return FAppStyle::GetFontStyle("BoldFont");
+}
+
+FSlateColor GetVerbosityColor(const ELogVerbosity::Type InVerbosity)
+{
+	// Using same colors as Unreal Insights (see FTimeMarkerTrackBuilder::GetColorByVerbosity)
+	static FLinearColor Colors[] =
+	{
+		FLinearColor(0.0f, 0.0f, 0.0f, 1.0f), // NoLogging
+		FLinearColor(1.0f, 0.0f, 0.0f, 1.0f), // Fatal
+		FLinearColor(1.0f, 0.1f, 0.1f, 1.0f), // Error
+		FLinearColor(0.7f, 0.5f, 0.0f, 1.0f), // Warning
+		FLinearColor(0.0f, 0.7f, 0.0f, 1.0f), // Display
+		FLinearColor(0.0f, 0.7f, 1.0f, 1.0f), // Log
+		FLinearColor(0.7f, 0.7f, 0.7f, 1.0f), // Verbose
+		FLinearColor(1.0f, 1.0f, 1.0f, 1.0f), // VeryVerbose
+	};
+
+	// Special case that represents verbosity levels waiting for the response message
+	// or for categories that couldn't be found by the remote process
+	static FLinearColor UnknownVerbosityColor(0.3f, 0.3f, 0.3f, 1.0f);
+	if (InVerbosity == UnknownVerbosity)
+	{
+		return UnknownVerbosityColor;
+	}
+
+	static_assert(sizeof(Colors) / sizeof(FLinearColor) == static_cast<int>(ELogVerbosity::Type::NumVerbosity), "ELogVerbosity::Type has changed!?");
+	return Colors[InVerbosity & ELogVerbosity::VerbosityMask];
+}
+
+void SLogCategoryFilter::Construct(const FArguments& InArgs)
+{
+	VerbosityOptions =
+	{
+		MakeShared<ELogVerbosity::Type>(ELogVerbosity::NoLogging),
+		MakeShared<ELogVerbosity::Type>(ELogVerbosity::Fatal),
+		MakeShared<ELogVerbosity::Type>(ELogVerbosity::Error),
+		MakeShared<ELogVerbosity::Type>(ELogVerbosity::Warning),
+		MakeShared<ELogVerbosity::Type>(ELogVerbosity::Display),
+		MakeShared<ELogVerbosity::Type>(ELogVerbosity::Log),
+		MakeShared<ELogVerbosity::Type>(ELogVerbosity::Verbose),
+		MakeShared<ELogVerbosity::Type>(ELogVerbosity::VeryVerbose),
+	};
+
+	ChildSlot
+	[
+		SNew(SComboButton)
+		.ComboButtonStyle(&FAppStyle::Get().GetWidgetStyle<FComboButtonStyle>("ComboButton"))
+		.OnGetMenuContent(this, &SLogCategoryFilter::OnGetMenuContent)
+		.ButtonContent()
+		[
+			SNew(STextBlock)
+			.Text(LOCTEXT("ManageRecordingVerbosityLevels", "Recording Verbosity Levels"))
+		]
+	];
+}
+
+void SLogCategoryFilter::OnSearchTextChanged(const FText& InText)
+{
+	SearchText = InText;
+	FilteredCategories.Empty(AllCategories.Num());
+
+	if (SearchText.IsEmpty())
+	{
+		FilteredCategories = AllCategories;
+	}
+	else
+	{
+		for (const TSharedPtr<FLogCategoryVerbosity>& Item : AllCategories)
+		{
+			if (Item->CategoryName.ToString().Contains(SearchText.ToString()))
+			{
+				FilteredCategories.Add(Item);
+			}
+		}
+	}
+
+	if (ListView)
+	{
+		ListView->RequestListRefresh();
+	}
+}
+
+TSharedRef<SWidget> SLogCategoryFilter::OnGetMenuContent()
+{
+	// Create switcher widget first since RefreshCategoryList will set the active slot
+	SAssignNew(Switcher, SWidgetSwitcher)
+		+ SWidgetSwitcher::Slot()
+		[
+			SNew(SHorizontalBox)
+			+ SHorizontalBox::Slot()
+			.FillWidth(1.0f)
+			.HAlign(HAlign_Center)
+			.VAlign(VAlign_Center)
+			[
+				SNew(SCircularThrobber)
+			]
+		]
+		+ SWidgetSwitcher::Slot()
+		[
+			SNew(SVerticalBox)
+			+ SVerticalBox::Slot()
+			.AutoHeight()
+			.Padding(FMargin(5, 2))
+			[
+				SNew(SHorizontalBox)
+				+ SHorizontalBox::Slot()
+				.AutoWidth()
+				.HAlign(HAlign_Right)
+				.Padding(FMargin(5, 2))
+				[
+					SNew(SCheckBox)
+					.ToolTipText(LOCTEXT("VerbosityFilteringTooltip", "Whether logging functions based on LogCategories should record an entry based on the category active verbosity level."))
+					.IsChecked_Lambda([this]()
+						{
+							return bUsingVerbosityFilterWhenRecording ? ECheckBoxState::Checked : ECheckBoxState::Unchecked;
+						})
+					.OnCheckStateChanged_Lambda([this](ECheckBoxState InCheckBoxState)
+						{
+							SetUseVerbosityLevelFiltering(InCheckBoxState == ECheckBoxState::Checked);
+						})
+				]
+				+ SHorizontalBox::Slot()
+				.AutoWidth()
+				.HAlign(HAlign_Right)
+				.VAlign(VAlign_Center)
+				[
+					SNew(STextBlock)
+					.Text(LOCTEXT("FilteringEnabled", "Use verbosity level filtering"))
+				]
+			]
+			+ SVerticalBox::Slot()
+			[
+				SAssignNew(ListView, SListView<TSharedPtr<FLogCategoryVerbosity>>)
+				.ListItemsSource(&FilteredCategories)
+				.OnGenerateRow(this, &SLogCategoryFilter::OnGenerateRow)
+				.IsEnabled_Lambda([this]
+					{
+						return bUsingVerbosityFilterWhenRecording;
+					})
+			]
+		];
+
+	RefreshCategoryList();
+
+	return SNew(SBox)
+	.WidthOverride(400.0f)
+	[
+		SNew(SVerticalBox)
+		+ SVerticalBox::Slot()
+		.AutoHeight()
+		[
+			SAssignNew(SearchBox, SSearchBox)
+			.OnTextChanged(this, &SLogCategoryFilter::OnSearchTextChanged)
+			.InitialText(SearchText)
+		]
+		+ SVerticalBox::Slot()
+		.AutoHeight()
+		[
+			SNew(SImage).Image(FAppStyle::GetBrush("Menu.Separator"))
+		]
+		+ SVerticalBox::Slot()
+		[
+			Switcher.ToSharedRef()
+		]
+	];
+}
+
+TSharedRef<ITableRow> SLogCategoryFilter::OnGenerateRow(TSharedPtr<FLogCategoryVerbosity> InItem, const TSharedRef<STableViewBase>& InTable)
+{
+	return SNew(STableRow<TSharedPtr<UE::RewindDebugger::FLogCategoryVerbosity>>, InTable)
+		.Padding(FMargin(5, 2))
+		.Style(&FAppStyle::Get().GetWidgetStyle<FTableRowStyle>("ComboBox.Row"))
+		[
+			SNew(SHorizontalBox)
+			// Category name
+			+ SHorizontalBox::Slot()
+			.FillWidth(1.0f)
+			.VAlign(VAlign_Center)
+			.MinWidth(40.f)
+			[
+				SNew(STextBlock)
+				.Text(FText::FromName(InItem->CategoryName))
+				.ToolTipText(FText::FromName(InItem->CategoryName))
+				.Font(FAppStyle::GetFontStyle("NormalFont"))
+			]
+			// Verbosity combo selector
+			+ SHorizontalBox::Slot()
+			.AutoWidth()
+			[
+				SNew(SBox)
+				.WidthOverride(150.0f)
+				[
+					SNew(SComboBox<TSharedPtr<ELogVerbosity::Type>>)
+					.OptionsSource(&VerbosityOptions)
+					.OnGenerateWidget_Lambda([](TSharedPtr<ELogVerbosity::Type> InOption)
+						{
+							return SNew(STextBlock)
+							.Text(GetVerbosityAsText(*InOption))
+							.ColorAndOpacity(GetVerbosityColor(*InOption))
+							.Font(GetVerbosityFont(*InOption));
+						})
+					.OnSelectionChanged_Lambda([this, InItem](TSharedPtr<ELogVerbosity::Type> InNewValue, ESelectInfo::Type InSelectInfo)
+						{
+							if (InNewValue)
+							{
+								SetCategoryVerbosity(InItem->CategoryName, *InNewValue);
+
+								// Set the verbosity for proper UI feedback but the value will be updated by the response message
+								InItem->Verbosity = *InNewValue;
+							}
+						})
+					[
+						SNew(STextBlock)
+						.Text_Lambda([InItem]()
+							{
+								return GetVerbosityAsText(InItem->GetVerbosity());
+							})
+						.ColorAndOpacity_Lambda([InItem]()
+							{
+								return GetVerbosityColor(InItem->GetVerbosity());
+							})
+						.Font(GetVerbosityFont(InItem->GetVerbosity()))
+					]
+				]
+			]
+		];
+}
+
+void SLogCategoryFilter::RefreshCategoryList()
+{
+	AllCategories.Empty();
+
+	auto InsertFunc = [&AllCategories = AllCategories](const FName& LogCategory)
+		{
+			int32 InsertIndex = 0;
+			for (InsertIndex = AllCategories.Num() - 1; InsertIndex >= 0; --InsertIndex)
+			{
+				TSharedPtr<FLogCategoryVerbosity> CheckCategory = AllCategories[InsertIndex];
+				// No duplicates
+				if (!CheckCategory.IsValid() || CheckCategory->CategoryName == LogCategory)
+				{
+					return;
+				}
+
+				if (CheckCategory->CategoryName.Compare(LogCategory) < 0)
+				{
+					break;
+				}
+			}
+			AllCategories.Insert(MakeShareable(new FLogCategoryVerbosity{ LogCategory, UnknownVerbosity }), InsertIndex + 1);
+		};
+
+	const IRewindDebugger* RewindDebugger = IRewindDebugger::Instance();
+	if (const TraceServices::IAnalysisSession* Session = RewindDebugger ? RewindDebugger->GetAnalysisSession() : nullptr)
+	{
+		TraceServices::FAnalysisSessionReadScope SessionReadScope(*Session);
+		if (const IVisualLoggerProvider* VisualLoggerProvider = Session->ReadProvider<IVisualLoggerProvider>("VisualLoggerProvider"))
+		{
+			VisualLoggerProvider->EnumerateCategories([InsertFunc](const FName& Category)
+				{
+					InsertFunc(Category);
+				});
+		}
+	}
+
+	// Fetch verbosity levels from the remote process
+	SyncVerbosityLevels();
+
+	// Reapply the current search text
+	OnSearchTextChanged(SearchText);
+}
+
+void SLogCategoryFilter::SyncVerbosityLevels()
+{
+	int32 SwitcherStateIndex = SwitcherSlotIndexForReadyState;
+	ON_SCOPE_EXIT
+	{
+		if (Switcher)
+		{
+			Switcher->SetActiveWidgetIndex(SwitcherStateIndex);
+		}
+	};
+
+	// Nothing to sync if the list of known categories is empty
+	if (AllCategories.IsEmpty())
+	{
+		return;
+	}
+
+	const IRewindDebugger* RewindDebugger = IRewindDebugger::Instance();
+	if (RewindDebugger == nullptr)
+	{
+		return;
+	}
+
+	using namespace UE::TraceBasedDebuggers;
+	const TSharedPtr<FRemoteSessionsManager> RemoteSessionsManager = FRewindDebuggerEngineEditorBridge::Get().GetSessionsManager();
+	const TSharedPtr<FSessionInfo> SessionInfo = RewindDebugger->GetSelectedDebugSessionInfo();
+	if (!SessionInfo.IsValid()
+		|| !RemoteSessionsManager.IsValid())
+	{
+		// Can't send request without a valid sessions managers and an active session
+		return;
+	}
+
+	FVLogExtensionSessionData* DebuggerData = GetVLogExtensionData();
+	if (DebuggerData == nullptr)
+	{
+		return;
+	}
+
+	// Build query message to get the verbosity levels of all discovered log categories
+	FLogCategoryStatusQueryMessage Query;
+	Query.Categories.Reserve(AllCategories.Num());
+	for (const TSharedPtr<FLogCategoryVerbosity>& Item : AllCategories)
+	{
+		Query.Categories.Push(*Item);
+
+		FLogCategoryVerbosity* CategoryState = DebuggerData->LogCategoriesStatesByName.Find(Item->CategoryName);
+		if (CategoryState == nullptr)
+		{
+			CategoryState = &DebuggerData->LogCategoriesStatesByName.Add(Item->CategoryName);
+			CategoryState->CategoryName = Item->CategoryName;
+			CategoryState->Verbosity = UnknownVerbosity;
+		}
+	}
+
+	// Mark session data as pending response
+	DebuggerData->bPendingRefresh = true;
+
+	// Hide the categories and display feedback for the waiting state
+	SwitcherStateIndex = SwitcherSlotIndexForWaitingState;
+
+	// Register callback when the response is received
+	DebuggerData->OnDataUpdated.BindSPLambda(this, [this]()
+		{
+			// Update our list from the updated session data
+			if (FVLogExtensionSessionData* DebuggerData = GetVLogExtensionData())
+			{
+				for (const TSharedPtr<FLogCategoryVerbosity>& Item : AllCategories)
+				{
+					if (const FLogCategoryVerbosity* State = DebuggerData->LogCategoriesStatesByName.Find(Item->CategoryName))
+					{
+						Item->Verbosity = State->Verbosity;
+						continue;
+					}
+
+					ensureMsgf(false, TEXT("The response message should contain the verbosity levels for all our current known categories."));
+				}
+				DebuggerData->bPendingRefresh = false;
+				bUsingVerbosityFilterWhenRecording = DebuggerData->bUsingVerbosityFilterWhenRecording;
+			}
+
+			// We can now display the categories
+			if (Switcher)
+			{
+				Switcher->SetActiveWidgetIndex(SwitcherSlotIndexForReadyState);
+			}
+		});
+
+	// Send the query
+	RemoteSessionsManager->SendCommand(SessionInfo->Address, Query);
+}
+
+void SLogCategoryFilter::SetUseVerbosityLevelFiltering(const bool bEnable) const
+{
+	const IRewindDebugger* RewindDebugger = IRewindDebugger::Instance();
+	if (RewindDebugger == nullptr)
+	{
+		return;
+	}
+
+	using namespace UE::TraceBasedDebuggers;
+	const TSharedPtr<FRemoteSessionsManager> RemoteSessionsManager = FRewindDebuggerEngineEditorBridge::Get().GetSessionsManager();
+	const TSharedPtr<FSessionInfo> SessionInfo = RewindDebugger->GetSelectedDebugSessionInfo();
+	if (!SessionInfo.IsValid()
+		|| !RemoteSessionsManager.IsValid())
+	{
+		// Can't send request without a valid sessions managers and an active session
+		return;
+	}
+
+	FVLogExtensionSessionData* DebuggerData = GetVLogExtensionData();
+	if (DebuggerData == nullptr)
+	{
+		return;
+	}
+
+	// Mark session data as pending response
+	DebuggerData->bPendingRefresh = true;
+
+	// Hide the categories and display feedback for the waiting state
+	if (Switcher)
+	{
+		Switcher->SetActiveWidgetIndex(SwitcherSlotIndexForWaitingState);
+	}
+
+	// Send command to enable/disable the filtering
+	RemoteSessionsManager->SendCommand(SessionInfo->Address, FVerbosityFilteringStateChangeCommandMessage{ .bEnableFiltering = bEnable });
+}
+
+void SLogCategoryFilter::SetCategoryVerbosity(const FName InCategoryName, const ELogVerbosity::Type InVerbosity) const
+{
+	const IRewindDebugger* RewindDebugger = IRewindDebugger::Instance();
+	if (RewindDebugger == nullptr)
+	{
+		return;
+	}
+
+	using namespace UE::TraceBasedDebuggers;
+	const TSharedPtr<FRemoteSessionsManager> RemoteSessionsManager = FRewindDebuggerEngineEditorBridge::Get().GetSessionsManager();
+	const TSharedPtr<FSessionInfo> SessionInfo = RewindDebugger->GetSelectedDebugSessionInfo();
+	if (!SessionInfo.IsValid()
+		|| !RemoteSessionsManager.IsValid())
+	{
+		// Can't send request without a valid sessions managers and an active session
+		return;
+	}
+
+	FVLogExtensionSessionData* DebuggerData = GetVLogExtensionData();
+	if (DebuggerData == nullptr)
+	{
+		return;
+	}
+
+	FLogCategoryVerbosity* CategoryState = DebuggerData->LogCategoriesStatesByName.Find(InCategoryName);
+	if (CategoryState == nullptr)
+	{
+		CategoryState = &DebuggerData->LogCategoriesStatesByName.Add(InCategoryName);
+		CategoryState->CategoryName = InCategoryName;
+	}
+
+	CategoryState->Verbosity = UnknownVerbosity;
+
+	// Mark session data as pending response
+	DebuggerData->bPendingRefresh = true;
+
+	// Hide the categories and display feedback for the waiting state
+	if (Switcher)
+	{
+		Switcher->SetActiveWidgetIndex(SwitcherSlotIndexForWaitingState);
+	}
+
+	// Send command to change the verbosity level
+	RemoteSessionsManager->SendCommand(SessionInfo->Address, FLogCategoryStateChangeCommandMessage
+		{
+			FLogCategoryVerbosity{.CategoryName = CategoryState->CategoryName, .Verbosity = InVerbosity}
+		});
+}
+
+ELogVerbosity::Type SLogCategoryFilter::GetCategoryVerbosity(const FName InCategoryName) const
+{
+	if (FVLogExtensionSessionData* DebuggerData = GetVLogExtensionData())
+	{
+		if (const FLogCategoryVerbosity* CategoryState = DebuggerData->LogCategoriesStatesByName.Find(InCategoryName))
+		{
+			return CategoryState->GetVerbosity();
+		}
+	}
+
+	return UnknownVerbosity;
+}
+
+bool SLogCategoryFilter::CanChangeCategoryVerbosity(const FName InCategoryName) const
+{
+	// Prevent users from changing again the verbosity while we wait for the confirmation from the remote process.
+	if (const FVLogExtensionSessionData* DebuggerData = GetVLogExtensionData())
+	{
+		if (DebuggerData->bPendingRefresh)
+		{
+			return false;
+		}
+
+		return true;
+	}
+
+	return false;
+}
+
+} // UE::RewindDebugger
+#undef LOCTEXT_NAMESPACE
diff --git a/Engine/Plugins/Animation/GameplayInsights/Source/RewindDebuggerVLog/Private/Widgets/LogCategoryVerbositySelector.h b/Engine/Plugins/Animation/GameplayInsights/Source/RewindDebuggerVLog/Private/Widgets/LogCategoryVerbositySelector.h
new file mode 100644
--- /dev/null
+++ b/Engine/Plugins/Animation/GameplayInsights/Source/RewindDebuggerVLog/Private/Widgets/LogCategoryVerbositySelector.h
@@ -0,0 +1,56 @@
+﻿// Copyright Epic Games, Inc. All Rights Reserved.
+#pragma once
+
+#include "Widgets/SCompoundWidget.h"
+#include "Templates/SharedPointerFwd.h"
+#include "Internationalization/Text.h"
+
+class ITableRow;
+class SWidgetSwitcher;
+class SSearchBox;
+class STableViewBase;
+template <typename ItemType> class SListView;
+namespace ELogVerbosity { enum Type : uint8; }
+
+namespace UE::RewindDebugger
+{
+struct FLogCategoryVerbosity;
+
+class SLogCategoryFilter : public SCompoundWidget
+{
+public:
+	SLATE_BEGIN_ARGS(SLogCategoryFilter) {}
+	SLATE_END_ARGS()
+
+	void Construct(const FArguments& InArgs);
+
+private:
+	TSharedRef<ITableRow> OnGenerateRow(TSharedPtr<FLogCategoryVerbosity> InItem, const TSharedRef<STableViewBase>& InTable);
+	void SyncVerbosityLevels();
+	TSharedRef<SWidget> OnGetMenuContent();
+	void OnSearchTextChanged(const FText& InText);
+
+	void RefreshCategoryList();
+
+	void SetUseVerbosityLevelFiltering(bool bEnable) const;
+
+	bool CanChangeCategoryVerbosity(FName InCategoryName) const;
+	void SetCategoryVerbosity(FName InCategoryName, ELogVerbosity::Type InVerbosity) const;
+	ELogVerbosity::Type GetCategoryVerbosity(FName InCategoryName) const;
+
+	TSharedPtr<SSearchBox> SearchBox;
+	TSharedPtr<SListView<TSharedPtr<FLogCategoryVerbosity>>> ListView;
+	TSharedPtr<SWidgetSwitcher> Switcher;
+
+	TArray<TSharedPtr<FLogCategoryVerbosity>> AllCategories;
+	TArray<TSharedPtr<FLogCategoryVerbosity>> FilteredCategories;
+	TArray<TSharedPtr<ELogVerbosity::Type>> VerbosityOptions;
+	FText SearchText;
+
+	bool bUsingVerbosityFilterWhenRecording = false;
+	static constexpr int32 SwitcherSlotIndexForWaitingState = 0;
+	static constexpr int32 SwitcherSlotIndexForReadyState = 1;
+};
+
+} // UE::RewindDebugger
+
diff --git a/Engine/Plugins/Animation/GameplayInsights/Source/RewindDebuggerVLog/RewindDebuggerVLog.Build.cs b/Engine/Plugins/Animation/GameplayInsights/Source/RewindDebuggerVLog/RewindDebuggerVLog.Build.cs
--- a/Engine/Plugins/Animation/GameplayInsights/Source/RewindDebuggerVLog/RewindDebuggerVLog.Build.cs
+++ b/Engine/Plugins/Animation/GameplayInsights/Source/RewindDebuggerVLog/RewindDebuggerVLog.Build.cs
@@ -11,20 +11,21 @@ namespace UnrealBuildTool.Rules
 			{
 				"Core",
 				"CoreUObject",
+				"DeveloperSettings",
 				"Engine",
+				"GameplayInsights",
 				"InputCore",
-				"UnrealEd",
 				"LogVisualizer",
-				"TraceAnalysis",
-				"TraceServices",
-				"TraceInsights",
 				"RewindDebuggerInterface",
-				"GameplayInsights",
+				"RewindDebuggerRuntime",
+				"RewindDebuggerVLogRuntime",
 				"Slate",
 				"SlateCore",
-				"ToolWidgets",
 				"ToolMenus",
-				"DeveloperSettings",
+				"TraceAnalysis",
+				"TraceBasedDebuggers",
+				"TraceServices",
+				"UnrealEd"
 			});
 		}
 	}
diff --git a/Engine/Plugins/Animation/GameplayInsights/Source/RewindDebuggerVLogRuntime/Private/RewindDebuggerVLogRuntime.cpp b/Engine/Plugins/Animation/GameplayInsights/Source/RewindDebuggerVLogRuntime/Private/RewindDebuggerVLogRuntime.cpp
new file mode 100644
--- /dev/null
+++ b/Engine/Plugins/Animation/GameplayInsights/Source/RewindDebuggerVLogRuntime/Private/RewindDebuggerVLogRuntime.cpp
@@ -0,0 +1,223 @@
+// Copyright Epic Games, Inc. All Rights Reserved.
+
+#include "RewindDebuggerVLogRuntime.h"
+
+#if UE_DEBUG_RECORDING_ENABLED
+#include "Engine/Engine.h"
+#include "MessageEndpoint.h"
+#include "MessageEndpointBuilder.h"
+#include "Misc/App.h"
+#include "Misc/StringOutputDevice.h"
+#include "RemoteSessionsManager.h"
+#include "RewindDebuggerVLogRuntimeTypes.h"
+#include "SessionInfo.h"
+#include "VisualLogger/VisualLogger.h"
+#include "VisualLogger/VisualLoggerSettings.h"
+
+namespace UE::RewindDebugger
+{
+
+using namespace TraceBasedDebuggers;
+
+FVLogExtensionSessionData* GetSessionData(const TSharedPtr<FRemoteSessionsManager>& InSessionManager, const FGuid InSessionId)
+{
+	if (InSessionManager)
+	{
+		if (const TSharedPtr<FSessionInfo> SessionInfoPtr = InSessionManager->GetSessionInfo(InSessionId).Pin())
+		{
+			return SessionInfoPtr->GetDebuggerData<FVLogExtensionSessionData>();
+		}
+	}
+	return nullptr;
+}
+
+void FRewindDebuggerVLogRuntime::RecordingStarted()
+{
+	// start recording visual logger data
+	bWasRecording = FVisualLogger::Get().IsRecording();
+	FVisualLogger::Get().SetIsRecordingToTrace(true);
+}
+
+void FRewindDebuggerVLogRuntime::RecordingStopped()
+{
+	// stop recording visual logger data
+	FVisualLogger::Get().SetIsRecordingToTrace(false);
+	if (!bWasRecording)
+	{
+		FVisualLogger::Get().SetIsRecording(false);
+	}
+}
+
+void FRewindDebuggerVLogRuntime::RegisterMessageHandlers(
+	const TSharedPtr<FRemoteSessionsManager>& InSessionsManager, FMessageEndpointBuilder& InEndPointBuilder)
+{
+	using namespace UE::TraceBasedDebuggers;
+
+	InEndPointBuilder
+		.Handling<FVerbosityFilteringStateChangeCommandMessage>(
+			[WeakSessionsManager = InSessionsManager.ToWeakPtr()](const FVerbosityFilteringStateChangeCommandMessage& InMessage, const TSharedRef<IMessageContext>&)
+			{
+				UE_AUTORTFM_ONCOMMIT(InMessage, WeakSessionsManager)
+				{
+					const TSharedPtr<FRemoteSessionsManager> SessionsManager = WeakSessionsManager.Pin();
+					if (!SessionsManager.IsValid())
+					{
+						return;
+					}
+
+					FVerbosityFilteringStateChangeResponseMessage ResponseMessage;
+					ResponseMessage.InstanceID = FApp::GetInstanceId();
+
+					UVisualLoggerSettings* Settings = GetMutableDefault<UVisualLoggerSettings>();
+					Settings->SetUseVerbosityFilterWhenRecording(InMessage.bEnableFiltering);
+					ResponseMessage.bFilteringEnabled = Settings->bUseVerbosityFilterWhenRecording;
+
+					SessionsManager->PublishMessage(ResponseMessage);
+				};
+			})
+		.Handling<FLogCategoryStateChangeCommandMessage>(
+			[WeakSessionsManager = InSessionsManager.ToWeakPtr()](const FLogCategoryStateChangeCommandMessage& InMessage, const TSharedRef<IMessageContext>&)
+			{
+				UE_AUTORTFM_ONCOMMIT(InMessage, WeakSessionsManager)
+				{
+					const TSharedPtr<FRemoteSessionsManager> SessionsManager = WeakSessionsManager.Pin();
+					if (!SessionsManager.IsValid())
+					{
+						return;
+					}
+
+					FLogCategoryStateChangeResponseMessage ResponseMessage;
+					ResponseMessage.InstanceID = FApp::GetInstanceId();
+					ResponseMessage.NewState.CategoryName = InMessage.NewState.CategoryName;
+					ResponseMessage.NewState.Verbosity = ELogVerbosity::NoLogging;
+
+					// Make sure the category exists
+					FStringOutputDevice OutputDevice;
+					GEngine->Exec(/*World*/nullptr, TEXT("Log list"), OutputDevice);
+
+					TArray<FString> AllCategories;
+					OutputDevice.ParseIntoArrayWS(AllCategories);
+
+					const int32 CategoryIndex = AllCategories.IndexOfByKey(InMessage.NewState.CategoryName);
+
+					// If it exists we can then execute the command and send confirmation of the new verbosity level
+					if (CategoryIndex != INDEX_NONE)
+					{
+						const FString Command = FString::Printf(TEXT("Log %s %s")
+							, *InMessage.NewState.CategoryName.ToString()
+							, ToString(InMessage.NewState.GetVerbosity()));
+						GEngine->Exec(/*World*/nullptr, *Command);
+
+						ResponseMessage.NewState.Verbosity = InMessage.NewState.Verbosity;
+					}
+
+					SessionsManager->PublishMessage(ResponseMessage);
+				};
+			})
+		.Handling<FLogCategoryStatusQueryMessage>(
+			[WeakSessionsManager = InSessionsManager.ToWeakPtr()](const FLogCategoryStatusQueryMessage& InMessage, const TSharedRef<IMessageContext>&)
+			{
+				UE_AUTORTFM_ONCOMMIT(InMessage, WeakSessionsManager)
+				{
+					if (const TSharedPtr<FRemoteSessionsManager> SessionsManager = WeakSessionsManager.Pin())
+					{
+						FStringOutputDevice OutputDevice;
+						GEngine->Exec(/*World*/nullptr, TEXT("Log list"), OutputDevice);
+
+						FLogCategoryStatusQueryResponseMessage ResponseMessage;
+						ResponseMessage.InstanceID = FApp::GetInstanceId();
+						ResponseMessage.Categories = InMessage.Categories;
+						ResponseMessage.bUsingVerbosityFilterWhenRecording = FVisualLogger::bUseVerbosityFilterWhenRecording;
+						TArray<FString> AllCategories;
+						OutputDevice.ParseIntoArrayWS(AllCategories);
+						for (FLogCategoryVerbosity& CategoryInfo : ResponseMessage.Categories)
+						{
+							CategoryInfo.Verbosity = UnknownVerbosity;
+							const int32 CategoryIndex = AllCategories.IndexOfByKey(CategoryInfo.CategoryName);
+							if (CategoryIndex != INDEX_NONE)
+							{
+								const int32 VerbosityIndex = CategoryIndex + 1;
+								if (AllCategories.IsValidIndex(VerbosityIndex))
+								{
+									const ELogVerbosity::Type Verbosity = ParseLogVerbosityFromString(AllCategories[VerbosityIndex]);
+									CategoryInfo.Verbosity = Verbosity;
+								}
+							}
+						}
+
+						SessionsManager->PublishMessage(ResponseMessage);
+					}
+				};
+			});
+
+	if (InSessionsManager && InSessionsManager->IsController())
+	{
+		InEndPointBuilder
+			.Handling<FVerbosityFilteringStateChangeResponseMessage>(
+				[WeakSessionsManager = InSessionsManager.ToWeakPtr()](const FVerbosityFilteringStateChangeResponseMessage& InMessage, const TSharedRef<IMessageContext>&)
+				{
+					if (FVLogExtensionSessionData* DebuggerData = GetSessionData(WeakSessionsManager.Pin(), InMessage.InstanceID))
+					{
+						DebuggerData->bUsingVerbosityFilterWhenRecording = InMessage.bFilteringEnabled;
+						DebuggerData->NotifyDataUpdated();
+					}
+				})
+			.Handling<FLogCategoryStateChangeResponseMessage>(
+				[WeakSessionsManager = InSessionsManager.ToWeakPtr()](const FLogCategoryStateChangeResponseMessage& InMessage, const TSharedRef<IMessageContext>&)
+				{
+					if (FVLogExtensionSessionData* DebuggerData = GetSessionData(WeakSessionsManager.Pin(), InMessage.InstanceID))
+					{
+						if (FLogCategoryVerbosity* FoundState = DebuggerData->LogCategoriesStatesByName.Find(InMessage.NewState.CategoryName))
+						{
+							*FoundState = InMessage.NewState;
+						}
+
+						DebuggerData->NotifyDataUpdated();
+					}
+				})
+			.Handling<FLogCategoryStatusQueryResponseMessage>(
+				[WeakSessionsManager = InSessionsManager.ToWeakPtr()](const FLogCategoryStatusQueryResponseMessage& InMessage, const TSharedRef<IMessageContext>&)
+				{
+					if (FVLogExtensionSessionData* DebuggerData = GetSessionData(WeakSessionsManager.Pin(), InMessage.InstanceID))
+					{
+						for (const FLogCategoryVerbosity& Category : InMessage.Categories)
+						{
+							if (FLogCategoryVerbosity* FoundState = DebuggerData->LogCategoriesStatesByName.Find(Category.CategoryName))
+							{
+								*FoundState = Category;
+							}
+						}
+						DebuggerData->NotifyDataUpdated();
+					}
+				});
+	}
+}
+
+void FRewindDebuggerVLogRuntime::RegisterMessageTypes(
+	const TSharedPtr<FRemoteSessionsManager>& InSessionsManager, FMessageEndpoint& InMessageEndpoint)
+{
+	if (InSessionsManager)
+	{
+		InSessionsManager->RegisterExternalSupportedMessageType<FVerbosityFilteringStateChangeCommandMessage>();
+		InSessionsManager->RegisterExternalSupportedMessageType<FVerbosityFilteringStateChangeResponseMessage>();
+		InSessionsManager->RegisterExternalSupportedMessageType<FLogCategoryStateChangeCommandMessage>();
+		InSessionsManager->RegisterExternalSupportedMessageType<FLogCategoryStateChangeResponseMessage>();
+		InSessionsManager->RegisterExternalSupportedMessageType<FLogCategoryStatusQueryMessage>();
+		InSessionsManager->RegisterExternalSupportedMessageType<FLogCategoryStatusQueryResponseMessage>();
+
+		if (InSessionsManager->IsController())
+		{
+			InMessageEndpoint.Subscribe<FVerbosityFilteringStateChangeResponseMessage>();
+			InMessageEndpoint.Subscribe<FLogCategoryStateChangeResponseMessage>();
+			InMessageEndpoint.Subscribe<FLogCategoryStatusQueryResponseMessage>();
+		}
+	}
+
+	InMessageEndpoint.Subscribe<FVerbosityFilteringStateChangeCommandMessage>();
+	InMessageEndpoint.Subscribe<FLogCategoryStateChangeCommandMessage>();
+	InMessageEndpoint.Subscribe<FLogCategoryStatusQueryMessage>();
+}
+
+} // UE::RewindDebugger
+
+#endif // UE_DEBUG_RECORDING_ENABLED
\ No newline at end of file
diff --git a/Engine/Plugins/Animation/GameplayInsights/Source/RewindDebuggerVLogRuntime/Private/RewindDebuggerVLogRuntime.h b/Engine/Plugins/Animation/GameplayInsights/Source/RewindDebuggerVLogRuntime/Private/RewindDebuggerVLogRuntime.h
--- a/Engine/Plugins/Animation/GameplayInsights/Source/RewindDebuggerVLogRuntime/Private/RewindDebuggerVLogRuntime.h
+++ b/Engine/Plugins/Animation/GameplayInsights/Source/RewindDebuggerVLogRuntime/Private/RewindDebuggerVLogRuntime.h
@@ -1,4 +1,4 @@
-// Copyright Epic Games, Inc. All Rights Reserved.
+﻿// Copyright Epic Games, Inc. All Rights Reserved.
 
 #pragma once
 
@@ -6,32 +6,23 @@
 
 #if UE_DEBUG_RECORDING_ENABLED
 #include "RewindDebuggerRuntimeInterface/IRewindDebuggerRuntimeExtension.h"
-#include "VisualLogger/VisualLogger.h"
 
-namespace RewindDebugger
+struct FMessageEndpointBuilder;
+class FMessageEndpoint;
+
+namespace UE::RewindDebugger
+{
+class FRewindDebuggerVLogRuntime : public IRewindDebuggerRuntimeExtension
 {
-	class FRewindDebuggerVLogRuntime : public IRewindDebuggerRuntimeExtension
-	{
-	public:
-		virtual void RecordingStarted() override
-		{
-			// start recording visual logger data
-			bWasRecording = FVisualLogger::Get().IsRecording();
-			FVisualLogger::Get().SetIsRecordingToTrace(true);
-		}
+public:
+	virtual void RecordingStarted() override;
+	virtual void RecordingStopped() override;
+	virtual void RegisterMessageHandlers(const TSharedPtr<TraceBasedDebuggers::FRemoteSessionsManager>&, FMessageEndpointBuilder&) override;
+	virtual void RegisterMessageTypes(const TSharedPtr<TraceBasedDebuggers::FRemoteSessionsManager>&, FMessageEndpoint&) override;
 
-		virtual void RecordingStopped() override
-		{
-			// stop recording visual logger data
-			FVisualLogger::Get().SetIsRecordingToTrace(false);
-			if (!bWasRecording)
-			{
-				FVisualLogger::Get().SetIsRecording(false);
-			}
-		}
+private:
+	bool bWasRecording = false;
+};
 
-	private:
-		bool bWasRecording = false;
-	};
 }
 #endif // UE_DEBUG_RECORDING_ENABLED
diff --git a/Engine/Plugins/Animation/GameplayInsights/Source/RewindDebuggerVLogRuntime/Private/RewindDebuggerVLogRuntimeModule.h b/Engine/Plugins/Animation/GameplayInsights/Source/RewindDebuggerVLogRuntime/Private/RewindDebuggerVLogRuntimeModule.h
--- a/Engine/Plugins/Animation/GameplayInsights/Source/RewindDebuggerVLogRuntime/Private/RewindDebuggerVLogRuntimeModule.h
+++ b/Engine/Plugins/Animation/GameplayInsights/Source/RewindDebuggerVLogRuntime/Private/RewindDebuggerVLogRuntimeModule.h
@@ -18,6 +18,6 @@ public:
 	virtual void ShutdownModule() override;
 
 private:
-	RewindDebugger::FRewindDebuggerVLogRuntime VLogExtension;
+	UE::RewindDebugger::FRewindDebuggerVLogRuntime VLogExtension;
 #endif // UE_DEBUG_RECORDING_ENABLED
 };
diff --git a/Engine/Plugins/Animation/GameplayInsights/Source/RewindDebuggerVLogRuntime/Public/RewindDebuggerVLogRuntimeTypes.h b/Engine/Plugins/Animation/GameplayInsights/Source/RewindDebuggerVLogRuntime/Public/RewindDebuggerVLogRuntimeTypes.h
new file mode 100644
--- /dev/null
+++ b/Engine/Plugins/Animation/GameplayInsights/Source/RewindDebuggerVLogRuntime/Public/RewindDebuggerVLogRuntimeTypes.h
@@ -0,0 +1,124 @@
+// Copyright Epic Games, Inc. All Rights Reserved.
+
+#pragma once
+
+#include "Logging/LogVerbosity.h"
+#include "Misc/Guid.h"
+#include "RewindDebuggerVLogRuntimeTypes.generated.h"
+
+namespace UE::RewindDebugger
+{
+
+/**
+ * Specific value for category verbosity levels waiting for their status provided by the remote process
+ * or for categories not found by the remote process (edge case for now, but will be useful when allowing users to type manually a category)
+ */
+constexpr ELogVerbosity::Type UnknownVerbosity = ELogVerbosity::NumVerbosity;
+
+USTRUCT()
+struct FLogCategoryVerbosity
+{
+	GENERATED_BODY()
+
+	UPROPERTY()
+	FName CategoryName;
+
+	UPROPERTY()
+	uint8 Verbosity = 0;
+
+	ELogVerbosity::Type GetVerbosity() const
+	{
+		return static_cast<ELogVerbosity::Type>(Verbosity);
+	}
+};
+
+USTRUCT()
+struct FVerbosityFilteringStateChangeCommandMessage
+{
+	GENERATED_BODY()
+
+	UPROPERTY()
+	bool bEnableFiltering = false;
+};
+
+USTRUCT()
+struct FVerbosityFilteringStateChangeResponseMessage
+{
+	GENERATED_BODY()
+
+	UPROPERTY()
+	FGuid InstanceID = FGuid();
+
+	UPROPERTY()
+	bool bFilteringEnabled = false;
+};
+
+USTRUCT()
+struct FLogCategoryStateChangeCommandMessage
+{
+	GENERATED_BODY()
+
+	UPROPERTY()
+	FLogCategoryVerbosity NewState;
+};
+
+USTRUCT()
+struct FLogCategoryStateChangeResponseMessage
+{
+	GENERATED_BODY()
+
+	UPROPERTY()
+	FGuid InstanceID = FGuid();
+
+	UPROPERTY()
+	FLogCategoryVerbosity NewState;
+};
+
+USTRUCT()
+struct FLogCategoryStatusQueryMessage
+{
+	GENERATED_BODY()
+
+	UPROPERTY()
+	TArray<FLogCategoryVerbosity> Categories;
+};
+
+USTRUCT()
+struct FLogCategoryStatusQueryResponseMessage
+{
+	GENERATED_BODY()
+
+	UPROPERTY()
+	FGuid InstanceID = FGuid();
+
+	UPROPERTY()
+	bool bUsingVerbosityFilterWhenRecording = false;
+
+	UPROPERTY()
+	TArray<FLogCategoryVerbosity> Categories;
+};
+
+/**
+ * Structure specific to RewindDebugger VLog extension to store the states of the log categories in the session info struct.
+ * The struct is stored as a DebuggerSpecificSessionData inside UE::TraceBasedDebuggers::FSessionInfo
+ * and get be retrieved using GetDebuggerData<UE::RewindDebugger::FVLogExtensionSessionData>().
+ */
+USTRUCT()
+struct FVLogExtensionSessionData
+{
+	GENERATED_BODY()
+
+	TMap<FName, FLogCategoryVerbosity> LogCategoriesStatesByName;
+
+	FSimpleDelegate OnDataUpdated;
+	bool bUsingVerbosityFilterWhenRecording = false;
+	bool bPendingRefresh = false;
+
+	void NotifyDataUpdated()
+	{
+		bPendingRefresh = false;
+		OnDataUpdated.ExecuteIfBound();
+	}
+};
+
+}
diff --git a/Engine/Plugins/Animation/GameplayInsights/Source/RewindDebuggerVLogRuntime/RewindDebuggerVLogRuntime.Build.cs b/Engine/Plugins/Animation/GameplayInsights/Source/RewindDebuggerVLogRuntime/RewindDebuggerVLogRuntime.Build.cs
--- a/Engine/Plugins/Animation/GameplayInsights/Source/RewindDebuggerVLogRuntime/RewindDebuggerVLogRuntime.Build.cs
+++ b/Engine/Plugins/Animation/GameplayInsights/Source/RewindDebuggerVLogRuntime/RewindDebuggerVLogRuntime.Build.cs
@@ -6,14 +6,17 @@ namespace UnrealBuildTool.Rules
 	{
 		public RewindDebuggerVLogRuntime(ReadOnlyTargetRules Target) : base(Target)
 		{
+			bAllowUETypesInNamespaces = true;
+
 			PrivateDependencyModuleNames.AddRange(
 			new string[]
 			{
 				"Core",
 				"CoreUObject",
 				"Engine",
-				"TraceLog",
 				"RewindDebuggerRuntimeInterface",
+				"TraceBasedDebuggers",
+				"TraceLog"
 			});
 		}
 	}
diff --git a/Engine/Source/Editor/RewindDebuggerInterface/Public/IRewindDebugger.h b/Engine/Source/Editor/RewindDebuggerInterface/Public/IRewindDebugger.h
--- a/Engine/Source/Editor/RewindDebuggerInterface/Public/IRewindDebugger.h
+++ b/Engine/Source/Editor/RewindDebuggerInterface/Public/IRewindDebugger.h
@@ -1,22 +1,30 @@
-﻿// Copyright Epic Games, Inc. All Rights Reserved.
+// Copyright Epic Games, Inc. All Rights Reserved.
 #pragma once
 
 #if UE_ENABLE_INCLUDE_ORDER_DEPRECATED_IN_5_7
 #include "CoreMinimal.h"
 #endif // UE_ENABLE_INCLUDE_ORDER_DEPRECATED_IN_5_7
 
-#include "RewindDebuggerTypes.h"
-
 #if UE_ENABLE_INCLUDE_ORDER_DEPRECATED_IN_5_8
 #include "TraceServices/Model/Frames.h"
 #include "UObject/NameTypes.h"
 #endif // UE_ENABLE_INCLUDE_ORDER_DEPRECATED_IN_5_8
 
+#include "RewindDebuggerTypes.h"
 #include "IRewindDebugger.generated.h"
 
 REWINDDEBUGGERINTERFACE_API DECLARE_LOG_CATEGORY_EXTERN(LogRewindDebugger, Log, All);
 
-namespace TraceServices { class IAnalysisSession; }
+namespace TraceServices
+{
+	class IAnalysisSession;
+}
+
+namespace UE::TraceBasedDebuggers
+{
+	struct FSessionInfo;
+}
+
 namespace RewindDebugger
 {
 	class FRewindDebuggerTrack;
@@ -95,6 +103,12 @@ public:
 	/** @return the current visible range in Rewind Debugger recording time units */
 	virtual const TRange<double>& GetCurrentViewRange() const = 0;
 
+	/**
+	 * Returns shared pointer to the selected debug session (PIE or remote).
+	 * @return the currently selected debug session
+	 */
+	virtual TSharedPtr<UE::TraceBasedDebuggers::FSessionInfo> GetSelectedDebugSessionInfo() const = 0;
+
 	/**
 	 * Returns pointer to the current analysis sessions for immediate access.
 	 * If the goal is to store in a variable for later access then GetAnalysisSessionAsShared should be used instead.
diff --git a/Engine/Source/Runtime/Engine/Private/VisualLogger/VisualLogger.cpp b/Engine/Source/Runtime/Engine/Private/VisualLogger/VisualLogger.cpp
--- a/Engine/Source/Runtime/Engine/Private/VisualLogger/VisualLogger.cpp
+++ b/Engine/Source/Runtime/Engine/Private/VisualLogger/VisualLogger.cpp
@@ -39,10 +39,8 @@ FAutoConsoleCommandWithOutputDevice CmdActivateVerbosityFilterWhenRecording(
 	TEXT("Active log category verbosity will be used to filter out entries to record.\n"),
 	FConsoleCommandWithOutputDeviceDelegate::CreateLambda([](FOutputDevice& OutputDevice)
 		{
-			FVisualLogger::bUseVerbosityFilterWhenRecording = true;
-
-			// Keep setting in sync
-			GetMutableDefault<UVisualLoggerSettings>()->bUseVerbosityFilterWhenRecording = FVisualLogger::bUseVerbosityFilterWhenRecording;
+			GetMutableDefault<UVisualLoggerSettings>()->SetUseVerbosityFilterWhenRecording(true);
+			check(FVisualLogger::bUseVerbosityFilterWhenRecording);
 
 			OutputDevice.Log(ELogVerbosity::Log, TEXT("LogCategory verbosity based filtering is now 'ACTIVE'"));
 		}),
@@ -53,10 +51,8 @@ FAutoConsoleCommandWithOutputDevice CmdDeactivateVerbosityFilterWhenRecording(
 	TEXT("Log category verbosity is deactivated and everything is recorded.\n"),
 	FConsoleCommandWithOutputDeviceDelegate::CreateLambda([](FOutputDevice& OutputDevice)
 		{
-			FVisualLogger::bUseVerbosityFilterWhenRecording = false;
-
-			// Keep setting in sync
-			GetMutableDefault<UVisualLoggerSettings>()->bUseVerbosityFilterWhenRecording = FVisualLogger::bUseVerbosityFilterWhenRecording;
+			GetMutableDefault<UVisualLoggerSettings>()->SetUseVerbosityFilterWhenRecording(false);
+			check(!FVisualLogger::bUseVerbosityFilterWhenRecording);
 
 			OutputDevice.Log(ELogVerbosity::Log, TEXT("LogCategory verbosity based filtering is now 'INACTIVE'"));
 		}),
diff --git a/Engine/Source/Runtime/Engine/Private/VisualLogger/VisualLoggerSettings.cpp b/Engine/Source/Runtime/Engine/Private/VisualLogger/VisualLoggerSettings.cpp
--- a/Engine/Source/Runtime/Engine/Private/VisualLogger/VisualLoggerSettings.cpp
+++ b/Engine/Source/Runtime/Engine/Private/VisualLogger/VisualLoggerSettings.cpp
@@ -2,9 +2,15 @@
 
 #include "VisualLogger/VisualLoggerSettings.h"
 
-#if ENABLE_VISUAL_LOG
+#if UE_DEBUG_RECORDING_ENABLED
 #include "VisualLogger/VisualLogger.h"
 
+void UVisualLoggerSettings::SetUseVerbosityFilterWhenRecording(const bool bEnable)
+{
+	bUseVerbosityFilterWhenRecording = bEnable;
+	FVisualLogger::bUseVerbosityFilterWhenRecording = bUseVerbosityFilterWhenRecording;
+}
+
 void UVisualLoggerSettings::PostInitProperties()
 {
 	Super::PostInitProperties();
@@ -22,8 +28,8 @@ void UVisualLoggerSettings::PostEditChangeProperty(FPropertyChangedEvent& Proper
 	if (PropertyChangedEvent.Property
 		&& PropertyChangedEvent.Property->GetFName() == NAME_UseVerbosityFilterWhenRecording)
 	{
-		FVisualLogger::bUseVerbosityFilterWhenRecording = bUseVerbosityFilterWhenRecording;
+		SetUseVerbosityFilterWhenRecording(bUseVerbosityFilterWhenRecording);
 	}
 }
 #endif // WITH_EDITOR
-#endif // ENABLE_VISUAL_LOG
+#endif // UE_DEBUG_RECORDING_ENABLED
diff --git a/Engine/Source/Runtime/Engine/Public/VisualLogger/VisualLoggerSettings.h b/Engine/Source/Runtime/Engine/Public/VisualLogger/VisualLoggerSettings.h
--- a/Engine/Source/Runtime/Engine/Public/VisualLogger/VisualLoggerSettings.h
+++ b/Engine/Source/Runtime/Engine/Public/VisualLogger/VisualLoggerSettings.h
@@ -3,7 +3,7 @@
 #pragma once
 
 #include "Engine/DeveloperSettings.h"
-#include "EngineDefines.h"
+#include "VisualLoggerDefines.h"
 #include "VisualLoggerSettings.generated.h"
 
 #define UE_API ENGINE_API
@@ -26,14 +26,16 @@ public:
 	UPROPERTY(config, EditAnywhere, Category = "Recording")
 	bool bUseVerbosityFilterWhenRecording = false;
 
-#if ENABLE_VISUAL_LOG
+#if UE_DEBUG_RECORDING_ENABLED
+	UE_API void SetUseVerbosityFilterWhenRecording(bool bEnable);
+
 protected:
 	UE_API virtual void PostInitProperties() override;
 
 #if WITH_EDITOR
 	UE_API virtual void PostEditChangeProperty(FPropertyChangedEvent& PropertyChangedEvent) override;
 #endif // WITH_EDITOR
-#endif // ENABLE_VISUAL_LOG
+#endif // UE_DEBUG_RECORDING_ENABLED
 };
 
 #undef UE_API
diff --git a/Engine/Source/Runtime/Experimental/TraceBasedDebuggersAnalysis/Public/Widgets/RecordingControls.h b/Engine/Source/Runtime/Experimental/TraceBasedDebuggersAnalysis/Public/Widgets/RecordingControls.h
--- a/Engine/Source/Runtime/Experimental/TraceBasedDebuggersAnalysis/Public/Widgets/RecordingControls.h
+++ b/Engine/Source/Runtime/Experimental/TraceBasedDebuggersAnalysis/Public/Widgets/RecordingControls.h
@@ -1,4 +1,4 @@
-﻿// Copyright Epic Games, Inc. All Rights Reserved.
+// Copyright Epic Games, Inc. All Rights Reserved.
 
 #pragma once
 
@@ -45,6 +45,9 @@ struct FRecordingControls : TSharedFromThis<FRecordingControls>
 	UE_API void Deinitialize();
 	UE_API void AddToMenu(FName ExistingMenu);
 
+	/** Return the details about the currently selected session */
+	UE_API TSharedPtr<FSessionInfo> GetCurrentSessionInfo() const;
+
 	/** Delegate executed when recording got started in a given session */
 	FOnSessionEvent OnRecordingStarted;
 
@@ -105,7 +108,6 @@ protected:
 	virtual const FTraceSessionDescriptor& GetCurrentSessionDescriptorInternal() = 0;
 
 	UE_API void SelectTargetSession(FGuid SessionId);
-	UE_API TSharedPtr<FSessionInfo> GetCurrentSessionInfo() const;
 
 	UE_API void ToggleRecordingState(const TSharedPtr<FSessionInfo>& InSessionInfo);
 	FReply ToggleRecordingState()
diff --git a/Engine/Source/Runtime/RewindDebuggerRuntimeInterface/Private/IRewindDebuggerRuntimeExtension.cpp b/Engine/Source/Runtime/RewindDebuggerRuntimeInterface/Private/IRewindDebuggerRuntimeExtension.cpp
--- a/Engine/Source/Runtime/RewindDebuggerRuntimeInterface/Private/IRewindDebuggerRuntimeExtension.cpp
+++ b/Engine/Source/Runtime/RewindDebuggerRuntimeInterface/Private/IRewindDebuggerRuntimeExtension.cpp
@@ -1,6 +1,25 @@
 // Copyright Epic Games, Inc. All Rights Reserved.
 
 #include "RewindDebuggerRuntimeInterface/IRewindDebuggerRuntimeExtension.h"
+#include "Features/IModularFeatures.h"
 
 const FName IRewindDebuggerRuntimeExtension::ModularFeatureName = "RewindDebuggerRuntimeExtension";
-	
+
+namespace RewindDebugger
+{
+
+void IterateExtensions(TFunctionRef<void(IRewindDebuggerRuntimeExtension* Extension)> IteratorFunction)
+{
+	// update extensions
+	IModularFeatures& ModularFeatures = IModularFeatures::Get();
+
+	const int32 NumExtensions = ModularFeatures.GetModularFeatureImplementationCount(IRewindDebuggerRuntimeExtension::ModularFeatureName);
+	for (int32 ExtensionIndex = 0; ExtensionIndex < NumExtensions; ++ExtensionIndex)
+	{
+		IRewindDebuggerRuntimeExtension* Extension = static_cast<IRewindDebuggerRuntimeExtension*>(ModularFeatures.GetModularFeatureImplementation(IRewindDebuggerRuntimeExtension::ModularFeatureName, ExtensionIndex));
+		IteratorFunction(Extension);
+	}
+}
+
+} // RewindDebugger
+
diff --git a/Engine/Source/Runtime/RewindDebuggerRuntimeInterface/Public/RewindDebuggerRuntimeInterface/IRewindDebuggerRuntimeExtension.h b/Engine/Source/Runtime/RewindDebuggerRuntimeInterface/Public/RewindDebuggerRuntimeInterface/IRewindDebuggerRuntimeExtension.h
--- a/Engine/Source/Runtime/RewindDebuggerRuntimeInterface/Public/RewindDebuggerRuntimeInterface/IRewindDebuggerRuntimeExtension.h
+++ b/Engine/Source/Runtime/RewindDebuggerRuntimeInterface/Public/RewindDebuggerRuntimeInterface/IRewindDebuggerRuntimeExtension.h
@@ -2,33 +2,79 @@
 
 #pragma once
 
+#if UE_ENABLE_INCLUDE_ORDER_DEPRECATED_IN_5_8
 #include "CoreTypes.h"
+#endif // UE_ENABLE_INCLUDE_ORDER_DEPRECATED_IN_5_8
+
 #include "UObject/NameTypes.h"
 #include "Features/IModularFeature.h"
+#include "Templates/FunctionFwd.h"
+#include "Templates/SharedPointerFwd.h"
 
 #define UE_API REWINDDEBUGGERRUNTIMEINTERFACE_API
 
-// IRewindDebuggerRuntimeExtension 
-//
-// interface class for extensions which add functionality to the rewind debugger
-// these get a callback on recording start/stop, to enable/disable trace channels and on clear to clean up any cached data
-//
+class FMessageEndpoint;
+class IRewindDebuggerRuntimeExtension;
+
+namespace UE::TraceBasedDebuggers
+{
+struct FRemoteSessionsManager;
+}
+
+namespace RewindDebugger
+{
+UE_API extern void IterateExtensions(TFunctionRef<void(IRewindDebuggerRuntimeExtension* Extension)> IteratorFunction);
+}
+
+struct FMessageEndpointBuilder;
 
+/**
+ * IRewindDebuggerRuntimeExtension
+ *
+ * Interface class for extensions which add functionality to the rewind debugger runtime.
+ * These get callbacks on recording start/stop, to enable/disable trace channels and on clear to clean up any cached data.
+ * They can also register and handle messages for remote sessions debugging.
+ */
 class IRewindDebuggerRuntimeExtension : public IModularFeature
 {
 public:
 	virtual ~IRewindDebuggerRuntimeExtension() = default;
-	
+
 	static UE_API const FName ModularFeatureName;
 
-	// called when recording has started
-	virtual void RecordingStarted() {};
+	/** called when recording has started */
+	virtual void RecordingStarted()
+	{
+	};
+
+	/** called when recording has ended */
+	virtual void RecordingStopped()
+	{
+	};
+
+	/** called when recording is unloaded or a new one starts */
+	virtual void Clear()
+	{
+	};
+
+	/**
+	 * This method allows extensions to register their own message type handlers to the message system and send response using
+	 * the remote sessions manager.
+	 * @param InRemoteSessionsManager The remote sessions manager that can be used to send response messages
+	 * @param InMessageEndpointBuilder The message endpoint builder where message type handlers needs to be registered.
+	 */
+	virtual void RegisterMessageHandlers(const TSharedPtr<UE::TraceBasedDebuggers::FRemoteSessionsManager>& InRemoteSessionsManager, FMessageEndpointBuilder& InMessageEndpointBuilder)
+	{
+	};
 
-	// called when recording has ended
-	virtual void RecordingStopped() {};
-	
-	// called when recording is unloaded or a new one starts
-	virtual void Clear() {};
+	/**
+	 * This method allows extensions to register their own message types in the remote sessions manager and subscribe in the message endpoint.
+	 * @param InRemoteSessionsManager The remote sessions manager that can be used to register supported messages
+	 * @param InMessageEndpoint The message endpoint to subscribe to specific message types
+	 */
+	virtual void RegisterMessageTypes(const TSharedPtr<UE::TraceBasedDebuggers::FRemoteSessionsManager>& InRemoteSessionsManager, FMessageEndpoint& InMessageEndpoint)
+	{
+	};
 };
 
 #undef UE_API
