[RewindDebugger] fixed targets when WITH_TRACE_BASED_DEBUGGERS is 0

#Synced-CL 49932076

diff --git a/Engine/Plugins/Animation/GameplayInsights/Source/RewindDebuggerVLog/Private/Widgets/LogCategoryVerbositySelector.cpp b/Engine/Plugins/Animation/GameplayInsights/Source/RewindDebuggerVLog/Private/Widgets/LogCategoryVerbositySelector.cpp
--- a/Engine/Plugins/Animation/GameplayInsights/Source/RewindDebuggerVLog/Private/Widgets/LogCategoryVerbositySelector.cpp
+++ b/Engine/Plugins/Animation/GameplayInsights/Source/RewindDebuggerVLog/Private/Widgets/LogCategoryVerbositySelector.cpp
@@ -30,6 +30,7 @@ namespace UE::RewindDebugger
 
 FVLogExtensionSessionData* GetVLogExtensionData()
 {
+#if WITH_TRACE_BASED_DEBUGGERS
 	using namespace UE::TraceBasedDebuggers;
 	if (const IRewindDebugger* RewindDebugger = IRewindDebugger::Instance())
 	{
@@ -48,6 +49,8 @@ FVLogExtensionSessionData* GetVLogExtensionData()
 			return DebuggerData;
 		}
 	}
+#endif // WITH_TRACE_BASED_DEBUGGERS
+
 	return nullptr;
 }
 
@@ -343,6 +346,7 @@ void SLogCategoryFilter::RefreshCategoryList()
 
 void SLogCategoryFilter::SyncVerbosityLevels()
 {
+#if WITH_TRACE_BASED_DEBUGGERS
 	int32 SwitcherStateIndex = SwitcherSlotIndexForReadyState;
 	ON_SCOPE_EXIT
 	{
@@ -431,10 +435,12 @@ void SLogCategoryFilter::SyncVerbosityLevels()
 
 	// Send the query
 	RemoteSessionsManager->SendCommand(SessionInfo->Address, Query);
+#endif // WITH_TRACE_BASED_DEBUGGERS
 }
 
 void SLogCategoryFilter::SetUseVerbosityLevelFiltering(const bool bEnable) const
 {
+#if WITH_TRACE_BASED_DEBUGGERS
 	const IRewindDebugger* RewindDebugger = IRewindDebugger::Instance();
 	if (RewindDebugger == nullptr)
 	{
@@ -468,10 +474,12 @@ void SLogCategoryFilter::SetUseVerbosityLevelFiltering(const bool bEnable) const
 
 	// Send command to enable/disable the filtering
 	RemoteSessionsManager->SendCommand(SessionInfo->Address, FVerbosityFilteringStateChangeCommandMessage{ .bEnableFiltering = bEnable });
+#endif // WITH_TRACE_BASED_DEBUGGERS
 }
 
 void SLogCategoryFilter::SetCategoryVerbosity(const FName InCategoryName, const ELogVerbosity::Type InVerbosity) const
 {
+#if WITH_TRACE_BASED_DEBUGGERS
 	const IRewindDebugger* RewindDebugger = IRewindDebugger::Instance();
 	if (RewindDebugger == nullptr)
 	{
@@ -517,6 +525,7 @@ void SLogCategoryFilter::SetCategoryVerbosity(const FName InCategoryName, const
 		{
 			FLogCategoryVerbosity{.CategoryName = CategoryState->CategoryName, .Verbosity = InVerbosity}
 		});
+#endif // WITH_TRACE_BASED_DEBUGGERS
 }
 
 ELogVerbosity::Type SLogCategoryFilter::GetCategoryVerbosity(const FName InCategoryName) const
