[RewindDebugger] fixed non-unity build

#Synced-CL 49932076

diff --git a/Engine/Plugins/Animation/GameplayInsights/Source/RewindDebuggerVLog/Private/RewindDebuggerVLogModule.cpp b/Engine/Plugins/Animation/GameplayInsights/Source/RewindDebuggerVLog/Private/RewindDebuggerVLogModule.cpp
--- a/Engine/Plugins/Animation/GameplayInsights/Source/RewindDebuggerVLog/Private/RewindDebuggerVLogModule.cpp
+++ b/Engine/Plugins/Animation/GameplayInsights/Source/RewindDebuggerVLog/Private/RewindDebuggerVLogModule.cpp
@@ -2,9 +2,10 @@
 
 #include "RewindDebuggerVLogModule.h"
 
-#include "VisualLogTrack.h"
-#include "Modules/ModuleManager.h"
 #include "Features/IModularFeatures.h"
+#include "Modules/ModuleManager.h"
+#include "VisualLogger/VisualLoggerTypes.h"
+#include "VisualLogTrack.h"
 
 #define LOCTEXT_NAMESPACE "RewindDebuggerVLogModule"
 
