[RewindDebugger] GameplayProvider will now accept recording info events from server world when the analysis is from a server target.

#Synced-CL 49932076

diff --git a/Engine/Plugins/Animation/GameplayInsights/Source/GameplayInsights/Private/GameplayProvider.cpp b/Engine/Plugins/Animation/GameplayInsights/Source/GameplayInsights/Private/GameplayProvider.cpp
--- a/Engine/Plugins/Animation/GameplayInsights/Source/GameplayInsights/Private/GameplayProvider.cpp
+++ b/Engine/Plugins/Animation/GameplayInsights/Source/GameplayInsights/Private/GameplayProvider.cpp
@@ -1,4 +1,4 @@
-﻿// Copyright Epic Games, Inc. All Rights Reserved.
+// Copyright Epic Games, Inc. All Rights Reserved.
 
 #include "GameplayProvider.h"
 
@@ -389,8 +389,10 @@ bool FGameplayProvider::IsWorld(const uint64 InObjectId) const
 {
 	if (const FClassInfo* WorldClass = FindClassInfo(TEXT("/Script/Engine.World")))
 	{
-		const FObjectInfo* ObjectInfo = FindObjectInfo(InObjectId);
-		return ObjectInfo->ClassId == WorldClass->Id;
+		if (const FObjectInfo* ObjectInfo = FindObjectInfo(InObjectId))
+		{
+			return ObjectInfo->ClassId == WorldClass->Id;
+		}
 	}
 
 	return false;
@@ -766,15 +768,28 @@ void FGameplayProvider::AppendRecordingInfo(uint64 InWorldId, double InProfileTi
 {
 	Session.WriteAccessCheck();
 
-	if (int32* WorldIndex = WorldIdToIndexMap.Find(InWorldId))
+	if (!TargetType.IsSet())
+	{
+		TargetType = EBuildTargetType::Unknown;
+		Session.EnumerateMetadata(
+			[&TargetType = TargetType](const TraceServices::FTraceSessionMetadata& Metadata)
+			{
+				if (Metadata.Name == FName("TargetType"))
+				{
+					TargetType = static_cast<EBuildTargetType>(Metadata.Int64Value);
+				}
+			});
+	}
+
+	if (const int32* WorldIndex = WorldIdToIndexMap.Find(InWorldId))
 	{
 		if (WorldInfos.IsValidIndex(*WorldIndex))
 		{
 			// We should put each world into its own timeline in the future, to support frame stepping in sync with server frames
 			if (WorldInfos[*WorldIndex].NetMode == FWorldInfo::ENetMode::Standalone
-				|| WorldInfos[*WorldIndex].NetMode == FWorldInfo::ENetMode::Client)
-				// Disabling while investigating regression and implementing proper fix
-				//|| WorldInfos[*WorldIndex].NetMode == FWorldInfo::ENetMode::DedicatedServer)
+				|| WorldInfos[*WorldIndex].NetMode == FWorldInfo::ENetMode::Client
+				|| (TargetType == EBuildTargetType::Server
+					&& WorldInfos[*WorldIndex].NetMode == FWorldInfo::ENetMode::DedicatedServer))
 			{
 				FRecordingInfoMessage NewRecordingInfo;
 				NewRecordingInfo.WorldId = InWorldId;
diff --git a/Engine/Plugins/Animation/GameplayInsights/Source/GameplayInsights/Private/GameplayProvider.h b/Engine/Plugins/Animation/GameplayInsights/Source/GameplayInsights/Private/GameplayProvider.h
--- a/Engine/Plugins/Animation/GameplayInsights/Source/GameplayInsights/Private/GameplayProvider.h
+++ b/Engine/Plugins/Animation/GameplayInsights/Source/GameplayInsights/Private/GameplayProvider.h
@@ -1,4 +1,4 @@
-// Copyright Epic Games, Inc. All Rights Reserved.
+﻿// Copyright Epic Games, Inc. All Rights Reserved.
 
 #pragma once
 
@@ -213,15 +213,18 @@ private:
 	// map from object id to index in ObjectRecordingLifetimes, for lookup
 	TMap<RewindDebugger::FObjectId, uint64> ActiveObjectRecordingLifetimes;
 
+	/** MultiMap from parent object to sub objects, to avoid slow reverse lookups */
+	TMultiMap<RewindDebugger::FObjectId, RewindDebugger::FObjectId> ObjectHierarchy;
+
+	/** The target type specified in the current analysis session */
+	TOptional<EBuildTargetType> TargetType;
+
 	/** Whether we have any data */
 	bool bHasAnyData = false;
 
-	/** Whether we have any data */
+	/** The elapsed world time */
 	double RecordingDuration = 0.0;
 
 	/** Whether we have any object properties */
 	bool bHasObjectProperties = false;
-
-	/** MultiMap from parent object to sub objects, to avoid slow reverse lookups */
-	TMultiMap<RewindDebugger::FObjectId, RewindDebugger::FObjectId> ObjectHierarchy;
 };
