[RewindDebugger] added new button to open the folder containing analysis files and one to delete all utrace files in that folder.

#Synced-CL 49932076

diff --git a/Engine/Plugins/Animation/GameplayInsights/Source/RewindDebugger/Private/RewindDebugger.cpp b/Engine/Plugins/Animation/GameplayInsights/Source/RewindDebugger/Private/RewindDebugger.cpp
--- a/Engine/Plugins/Animation/GameplayInsights/Source/RewindDebugger/Private/RewindDebugger.cpp
+++ b/Engine/Plugins/Animation/GameplayInsights/Source/RewindDebugger/Private/RewindDebugger.cpp
@@ -1,4 +1,4 @@
-﻿// Copyright Epic Games, Inc. All Rights Reserved.
+// Copyright Epic Games, Inc. All Rights Reserved.
 
 #include "RewindDebugger.h"
 
@@ -1821,7 +1821,7 @@ void FRewindDebugger::RegisterToolBar()
 		Commands.ClearAnalysis,
 		FText(),
 		TAttribute<FText>(),
-		FSlateIcon(FAppStyle::GetAppStyleSetName(), "Icons.Delete")));
+		FSlateIcon("RewindDebuggerStyle", "RewindDebugger.ClearAnalysis")));
 
 	Section.AddSeparator(NAME_None);
 
diff --git a/Engine/Plugins/Animation/GameplayInsights/Source/RewindDebugger/Private/RewindDebuggerRecordingControls.cpp b/Engine/Plugins/Animation/GameplayInsights/Source/RewindDebugger/Private/RewindDebuggerRecordingControls.cpp
--- a/Engine/Plugins/Animation/GameplayInsights/Source/RewindDebugger/Private/RewindDebuggerRecordingControls.cpp
+++ b/Engine/Plugins/Animation/GameplayInsights/Source/RewindDebugger/Private/RewindDebuggerRecordingControls.cpp
@@ -4,10 +4,13 @@
 
 #include "RewindDebugger.h"
 #include "RewindDebuggerEngineEditorBridge.h"
+#include "RewindDebuggerModule.h"
 #include "RewindDebuggerRemoteSessionsHandler.h"
 #include "RewindDebuggerStyle.h"
 #include "SessionInfo.h"
 
+#define LOCTEXT_NAMESPACE "RewindDebuggerRecordingControls"
+
 namespace UE::RewindDebugger
 {
 
@@ -58,6 +61,58 @@ TSharedPtr<TraceBasedDebuggers::FSessionInfo> FDebuggerRecordingControls::GetLoc
 	return LocalSessionInfo;
 }
 
+void FDebuggerRecordingControls::AddToMenuInternal(FToolMenuSection& InSection)
+{
+	InSection.AddEntry(FToolMenuEntry::InitToolBarButton("Open Folder",
+		FUIAction(
+			FExecuteAction::CreateSPLambda(this,
+				[this]()
+				{
+					FPlatformProcess::ExploreFolder(*GetSaveDirPath());
+				}
+			),
+			FCanExecuteAction()
+		),
+		FText::GetEmpty(),
+		LOCTEXT("OpenRecordingsFolderTooltip", "Open trace analysis recordings folder"),
+		FSlateIcon(FRewindDebuggerStyle::Get().GetStyleSetName(), "RewindDebugger.OpenRecordingDirectory")
+	));
+
+	InSection.AddEntry(FToolMenuEntry::InitToolBarButton("Delete all old recordings",
+		FUIAction(
+			FExecuteAction::CreateSPLambda(this,
+				[this]()
+				{
+					const FString Dir = GetSaveDirPath();
+					if (Dir.Len() > 0)
+					{
+						TArray<FString> Files;
+						IFileManager::Get().FindFiles(Files, *Dir, TEXT(".utrace"));
+						for (const FString& File : Files)
+						{
+							const FString FullPath = FString::Printf(TEXT("%s/%s"), *Dir, *File);
+							IFileManager::Get().Delete(*FullPath);
+						}
+					}
+				}
+			),
+			FCanExecuteAction()
+		),
+		FText::GetEmpty(),
+		LOCTEXT("DeleteAllRecordingsTooltip", "Delete all old trace analysis files"),
+		FSlateIcon(FRewindDebuggerStyle::Get().GetStyleSetName(), "RewindDebugger.DeleteOldTraces")
+	));
+}
+
+FString FDebuggerRecordingControls::GetSaveDirPath() const
+{
+	if (const TSharedPtr<TraceBasedDebuggers::FTraceSessionsManager> TraceSessionsManager = FRewindDebuggerModule::Get().GetTraceSessionsManager())
+	{
+		return TraceSessionsManager->GetSaveDirPath();
+	}
+	return TEXT("");
+}
+
 bool FDebuggerRecordingControls::CanStartRecordingLocalSession() const
 {
 	using namespace UE::TraceBasedDebuggers;
@@ -158,4 +213,6 @@ const TraceBasedDebuggers::FTraceSessionDescriptor& FDebuggerRecordingControls::
 #endif
 }
 
-} // UE::RewindDebugger
\ No newline at end of file
+} // UE::RewindDebugger
+
+#undef LOCTEXT_NAMESPACE
\ No newline at end of file
diff --git a/Engine/Plugins/Animation/GameplayInsights/Source/RewindDebugger/Private/RewindDebuggerRecordingControls.h b/Engine/Plugins/Animation/GameplayInsights/Source/RewindDebugger/Private/RewindDebuggerRecordingControls.h
--- a/Engine/Plugins/Animation/GameplayInsights/Source/RewindDebugger/Private/RewindDebuggerRecordingControls.h
+++ b/Engine/Plugins/Animation/GameplayInsights/Source/RewindDebugger/Private/RewindDebuggerRecordingControls.h
@@ -27,6 +27,8 @@ struct FDebuggerRecordingControls : TraceBasedDebuggers::FRecordingControls
 protected:
 	TSharedPtr<TraceBasedDebuggers::FSessionInfo> GetLocalSessionInfo() const;
 
+	virtual void AddToMenuInternal(FToolMenuSection& InSection) override;
+	virtual FString GetSaveDirPath() const override;
 	virtual bool ConnectToLiveSession_RelayInternal(const FGuid& InSessionId) override;
 	virtual bool ConnectToLiveSession_DirectInternal(const FGuid& InSessionId) override;
 	virtual bool IsRecordingAvailableForSessionInternal(const TSharedRef<TraceBasedDebuggers::FSessionInfo>& InSessionInfo) const override;
diff --git a/Engine/Plugins/Animation/GameplayInsights/Source/RewindDebugger/Private/RewindDebuggerStyle.cpp b/Engine/Plugins/Animation/GameplayInsights/Source/RewindDebugger/Private/RewindDebuggerStyle.cpp
--- a/Engine/Plugins/Animation/GameplayInsights/Source/RewindDebugger/Private/RewindDebuggerStyle.cpp
+++ b/Engine/Plugins/Animation/GameplayInsights/Source/RewindDebugger/Private/RewindDebuggerStyle.cpp
@@ -38,6 +38,10 @@ FRewindDebuggerStyle::FRewindDebuggerStyle() :
 	Set("RewindDebugger.AutoEject", new IMAGE_BRUSH("autoeject", {30.f, 24.f}));
 	Set("RewindDebugger.AutoRecord", new IMAGE_BRUSH("autorecord", {30.f, 24.f}));
 
+	Set("RewindDebugger.ClearAnalysis", new CORE_IMAGE_BRUSH("Editor/Slate/Icons/icon_delete_16px", Icon16x16));
+	Set("RewindDebugger.DeleteOldTraces", new CORE_IMAGE_BRUSH_SVG("Slate/Starship/Common/delete-outline", Icon16x16));
+	Set("RewindDebugger.OpenRecordingDirectory", new CORE_IMAGE_BRUSH_SVG("Editor/Slate/Starship/Common/BrowseContent", Icon16x16));
+
 	Set("RewindDebugger.ConnectToSession", new IMAGE_BRUSH_SVG("Session_20", {20.f, 20.f}));
 	// Actor Picker
 	Set("RewindDebugger.SelectActor", new CORE_IMAGE_BRUSH("Editor/Slate/Icons/eyedropper_16px", Icon16x16));
diff --git a/Engine/Plugins/ChaosVD/Source/ChaosVD/Private/Widgets/ChaosVDRecordingControls.cpp b/Engine/Plugins/ChaosVD/Source/ChaosVD/Private/Widgets/ChaosVDRecordingControls.cpp
--- a/Engine/Plugins/ChaosVD/Source/ChaosVD/Private/Widgets/ChaosVDRecordingControls.cpp
+++ b/Engine/Plugins/ChaosVD/Source/ChaosVD/Private/Widgets/ChaosVDRecordingControls.cpp
@@ -116,6 +116,15 @@ void FChaosVDRecordingControls::OnGenerateTargetSessionSelectorInternal(FMenuBui
 	MenuBuilder.EndSection();
 }
 
+FString FChaosVDRecordingControls::GetSaveDirPath() const
+{
+	if (const TSharedPtr<UE::TraceBasedDebuggers::FTraceSessionsManager>& TraceSessionsManager = FChaosVDModule::Get().GetTraceSessionsManager())
+	{
+		return TraceSessionsManager->GetSaveDirPath();
+	}
+	return TEXT("");
+}
+
 int32 FChaosVDRecordingControls::GetMaxConnectionRetriesInternal() const
 {
 	if (const UChaosVDGeneralSettings* Settings = FChaosVDSettingsManager::Get().GetSettingsObject<UChaosVDGeneralSettings>())
diff --git a/Engine/Plugins/ChaosVD/Source/ChaosVD/Private/Widgets/ChaosVDRecordingControls.h b/Engine/Plugins/ChaosVD/Source/ChaosVD/Private/Widgets/ChaosVDRecordingControls.h
--- a/Engine/Plugins/ChaosVD/Source/ChaosVD/Private/Widgets/ChaosVDRecordingControls.h
+++ b/Engine/Plugins/ChaosVD/Source/ChaosVD/Private/Widgets/ChaosVDRecordingControls.h
@@ -31,6 +31,7 @@ protected:
 	//~ Begin FRecordingControls interface
 	virtual void AddToMenuInternal(FToolMenuSection& InSection) override;
 	virtual void OnGenerateTargetSessionSelectorInternal(FMenuBuilder& InMenuBuilder) override;
+	virtual FString GetSaveDirPath() const override;
 
 	virtual int32 GetMaxConnectionRetriesInternal() const override;
 	virtual UE::TraceBasedDebuggers::ETraceTransportMode GetTransportModeOverrideForTargetTypeInternal(EBuildTargetType InBuildTarget) const override;
diff --git a/Engine/Source/Runtime/Experimental/TraceBasedDebuggersAnalysis/Private/Widgets/RecordingControls.cpp b/Engine/Source/Runtime/Experimental/TraceBasedDebuggersAnalysis/Private/Widgets/RecordingControls.cpp
--- a/Engine/Source/Runtime/Experimental/TraceBasedDebuggersAnalysis/Private/Widgets/RecordingControls.cpp
+++ b/Engine/Source/Runtime/Experimental/TraceBasedDebuggersAnalysis/Private/Widgets/RecordingControls.cpp
@@ -771,7 +771,9 @@ void FRecordingControls::AddToMenu(const FName ExistingMenu)
 		{
 			const TSharedRef<SWidget> StartRecordingButton = SNew(SBox).Padding(4.0f, 0.0f)
 				[
-					RecordingControls->GenerateToggleRecordingStateButton(LOCTEXT("RecordButtonDesc", "Starts a recording for the current session, saving it directly to file"))
+					RecordingControls->GenerateToggleRecordingStateButton(
+						FText::Format(LOCTEXT("RecordButtonDesc", "Starts a recording and analysis for the current session, saving it directly to '{0}'")
+							, FText::FromString(RecordingControls->GetSaveDirPath())))
 				];
 			const TSharedRef<SWidget> SessionsDropdown = SNew(SBox).Padding(4.0f, 0.0f)
 				[
diff --git a/Engine/Source/Runtime/Experimental/TraceBasedDebuggersAnalysis/Public/Widgets/RecordingControls.h b/Engine/Source/Runtime/Experimental/TraceBasedDebuggersAnalysis/Public/Widgets/RecordingControls.h
--- a/Engine/Source/Runtime/Experimental/TraceBasedDebuggersAnalysis/Public/Widgets/RecordingControls.h
+++ b/Engine/Source/Runtime/Experimental/TraceBasedDebuggersAnalysis/Public/Widgets/RecordingControls.h
@@ -62,10 +62,13 @@ protected:
 	virtual void AddToMenuInternal(FToolMenuSection& InSection)
 	{
 	}
+
 	virtual void OnGenerateTargetSessionSelectorInternal(FMenuBuilder& InMenuBuilder)
 	{
 	}
 
+	virtual FString GetSaveDirPath() const = 0;
+
 	UE_API virtual int32 GetMaxConnectionRetriesInternal() const;
 	UE_API ETraceTransportMode GetTransportMode(EBuildTargetType InBuildTarget) const;
 	UE_API virtual ETraceTransportMode GetTransportModeOverrideForTargetTypeInternal(EBuildTargetType InBuildTarget) const;
