[RewindDebugger] made sure that world times got reset before enabling trace channels from OnEnableRequiredTraceChannelsInternal virtual method.

#Synced-CL 49932076

diff --git a/Engine/Plugins/Animation/GameplayInsights/Source/RewindDebuggerRuntime/Private/RewindDebuggerRuntime.cpp b/Engine/Plugins/Animation/GameplayInsights/Source/RewindDebuggerRuntime/Private/RewindDebuggerRuntime.cpp
--- a/Engine/Plugins/Animation/GameplayInsights/Source/RewindDebuggerRuntime/Private/RewindDebuggerRuntime.cpp
+++ b/Engine/Plugins/Animation/GameplayInsights/Source/RewindDebuggerRuntime/Private/RewindDebuggerRuntime.cpp
@@ -51,7 +51,9 @@ namespace RewindDebugger
 
 	void FRewindDebuggerRuntime::OnEnableRequiredTraceChannelsInternal()
 	{
-
+		UE::Trace::ToggleChannel(TEXT("Object"), true);
+		UE::Trace::ToggleChannel(TEXT("ObjectProperties"), true);
+		UE::Trace::ToggleChannel(TEXT("Frame"), true);
 	}
 
 	void FRewindDebuggerRuntime::OnRecordingStartingInternal()
@@ -61,13 +63,22 @@ namespace RewindDebugger
 #endif
 
 		ClearRecording.Broadcast();
-	
+
 		// update extensions
 		IterateExtensions([](IRewindDebuggerRuntimeExtension* Extension)
 			{
 				Extension->Clear();
 			}
 		);
+
+		// Reset world elapsed time at the beginning of the "start recording" action before channels 
+		// gets enabled and events traced (i.e., HandleTraceConnectionEstablished)
+#if OBJECT_TRACE_ENABLED
+		for (TObjectIterator<UWorld> World; World; ++World)
+		{
+			FObjectTrace::ResetWorldElapsedTime(*World);
+		}
+#endif // OBJECT_TRACE_ENABLED
 	}
 
 	void FRewindDebuggerRuntime::OnRecordingStartedInternal()
@@ -79,21 +90,7 @@ namespace RewindDebugger
 		// is why we defer the callbacks.
 		ExecuteOnGameThread(UE_SOURCE_LOCATION, [this]
 			{
-#if OBJECT_TRACE_ENABLED
 				// trace each play-in-editor world, and all the actors in it.
-				for (TObjectIterator<UWorld> World; World; ++World)
-				{
-
-					FObjectTrace::ResetWorldElapsedTime(*World);
-				}
-
-#endif // OBJECT_TRACE_ENABLED
-
-				UE::Trace::ToggleChannel(TEXT("Object"), true);
-				UE::Trace::ToggleChannel(TEXT("ObjectProperties"), true);
-				UE::Trace::ToggleChannel(TEXT("Frame"), true);
-
-
 				for (TObjectIterator<UWorld> World; World; ++World)
 				{
 					TRACE_WORLD(*World);
diff --git a/Engine/Source/Runtime/Experimental/TraceBasedDebuggers/Private/TraceBasedDebuggerRuntime.cpp b/Engine/Source/Runtime/Experimental/TraceBasedDebuggers/Private/TraceBasedDebuggerRuntime.cpp
--- a/Engine/Source/Runtime/Experimental/TraceBasedDebuggers/Private/TraceBasedDebuggerRuntime.cpp
+++ b/Engine/Source/Runtime/Experimental/TraceBasedDebuggers/Private/TraceBasedDebuggerRuntime.cpp
@@ -1,4 +1,4 @@
-// Copyright Epic Games, Inc. All Rights Reserved.
+﻿// Copyright Epic Games, Inc. All Rights Reserved.
 
 #include "TraceBasedDebuggerRuntime.h"
 
@@ -180,10 +180,10 @@ void FRuntimeModule::StartRecording(const FStartRecordingCommandMessage& InRecor
 		}
 	}
 
+	// Disable current channels when beginning the "start recording" action, then wait
+	// for the connection to be established to enable the required channels (i.e., HandleTraceConnectionEstablished)
 	SaveAndDisabledCurrentEnabledTraceChannels();
 
-	EnableRequiredTraceChannels();
-
 	// Allow derived classes to perform additional actions before we start the traces
 	OnRecordingStartingInternal();
 
@@ -513,6 +513,8 @@ void FRuntimeModule::HandleTraceConnectionEstablished()
 		return;
 	}
 
+	EnableRequiredTraceChannels();
+
 	// Allow derived classes to perform additional actions
 	OnRecordingStartedInternal();
 
