[RewindDebugger] added settings to control the desired transport mode for remote connections.

#Synced-CL 50479750

diff --git a/Engine/Plugins/Animation/GameplayInsights/Source/RewindDebugger/Private/RewindDebuggerRecordingControls.cpp b/Engine/Plugins/Animation/GameplayInsights/Source/RewindDebugger/Private/RewindDebuggerRecordingControls.cpp
--- a/Engine/Plugins/Animation/GameplayInsights/Source/RewindDebugger/Private/RewindDebuggerRecordingControls.cpp
+++ b/Engine/Plugins/Animation/GameplayInsights/Source/RewindDebugger/Private/RewindDebuggerRecordingControls.cpp
@@ -147,6 +147,11 @@ void FDebuggerRecordingControls::StopRecordingLocalSession()
 	}
 }
 
+UE::TraceBasedDebuggers::ETraceTransportMode FDebuggerRecordingControls::GetTransportModeOverrideForTargetTypeInternal(const EBuildTargetType InBuildTarget) const
+{
+	return URewindDebuggerSettings::Get().GetTransportModeForTargetType(InBuildTarget);
+}
+
 bool FDebuggerRecordingControls::ConnectToLiveSession_RelayInternal(const FGuid& InSessionId)
 {
 #if UE_WITH_TRACE_BASED_DEBUGGERS_ANALYSIS
diff --git a/Engine/Plugins/Animation/GameplayInsights/Source/RewindDebugger/Private/RewindDebuggerRecordingControls.h b/Engine/Plugins/Animation/GameplayInsights/Source/RewindDebugger/Private/RewindDebuggerRecordingControls.h
--- a/Engine/Plugins/Animation/GameplayInsights/Source/RewindDebugger/Private/RewindDebuggerRecordingControls.h
+++ b/Engine/Plugins/Animation/GameplayInsights/Source/RewindDebugger/Private/RewindDebuggerRecordingControls.h
@@ -29,6 +29,8 @@ protected:
 
 	virtual void AddToMenuInternal(FToolMenuSection& InSection) override;
 	virtual FString GetSaveDirPath() const override;
+
+	virtual UE::TraceBasedDebuggers::ETraceTransportMode GetTransportModeOverrideForTargetTypeInternal(EBuildTargetType InBuildTarget) const override;
 	virtual bool ConnectToLiveSession_RelayInternal(const FGuid& InSessionId) override;
 	virtual bool ConnectToLiveSession_DirectInternal(const FGuid& InSessionId) override;
 	virtual bool IsRecordingAvailableForSessionInternal(const TSharedRef<TraceBasedDebuggers::FSessionInfo>& InSessionInfo) const override;
diff --git a/Engine/Plugins/Animation/GameplayInsights/Source/RewindDebugger/Private/RewindDebuggerSettings.cpp b/Engine/Plugins/Animation/GameplayInsights/Source/RewindDebugger/Private/RewindDebuggerSettings.cpp
--- a/Engine/Plugins/Animation/GameplayInsights/Source/RewindDebugger/Private/RewindDebuggerSettings.cpp
+++ b/Engine/Plugins/Animation/GameplayInsights/Source/RewindDebugger/Private/RewindDebuggerSettings.cpp
@@ -1,11 +1,19 @@
 ﻿// Copyright Epic Games, Inc. All Rights Reserved.
 
 #include "RewindDebuggerSettings.h"
+
+#include "IRewindDebugger.h"
+#include "TraceDataRelayTransport.h"
 #include "Misc/CoreDelegates.h"
 
 #include UE_INLINE_GENERATED_CPP_BY_NAME(RewindDebuggerSettings)
 
-URewindDebuggerSettings::URewindDebuggerSettings() : CameraMode(ERewindDebuggerCameraMode::Replay), bShouldAutoEject(false), bShouldAutoRecordOnPIE(false)
+URewindDebuggerSettings::URewindDebuggerSettings()
+	: CameraMode(ERewindDebuggerCameraMode::Replay)
+	, bShouldAutoEject(false)
+	, bShouldAutoRecordOnPIE(false)
+	, GameServerDataTransportMode(UE::TraceBasedDebuggers::ETraceTransportMode::Direct)
+	, GameClientDataTransportMode(UE::TraceBasedDebuggers::ETraceTransportMode::Direct)
 {
 	FCoreDelegates::OnPreExit.AddLambda([]()
 	{
@@ -32,11 +40,46 @@ FName URewindDebuggerSettings::GetCategoryName() const
 	return FName(TEXT("Plugins"));
 }
 
+UE::TraceBasedDebuggers::ETraceTransportMode URewindDebuggerSettings::GetTransportModeForTargetType(const EBuildTargetType InTargetType) const
+{
+	switch (InTargetType)
+	{
+	case EBuildTargetType::Server:
+		return GameServerDataTransportMode;
+	case EBuildTargetType::Game:
+	case EBuildTargetType::Client:
+		return GameClientDataTransportMode;
+	case EBuildTargetType::Editor:
+	case EBuildTargetType::Unknown:
+	case EBuildTargetType::Program:
+	default:
+		return UE::TraceBasedDebuggers::ETraceTransportMode::Direct;
+	}
+}
+
+void URewindDebuggerSettings::SetTransportModeForTargetType(const EBuildTargetType InTargetType, const UE::TraceBasedDebuggers::ETraceTransportMode InTransportMode)
+{
+	switch (InTargetType)
+	{
+	case EBuildTargetType::Server:
+		GameServerDataTransportMode = InTransportMode;
+		break;
+	case EBuildTargetType::Game:
+	case EBuildTargetType::Client:
+		GameClientDataTransportMode = InTransportMode;
+		break;
+	case EBuildTargetType::Editor:
+	case EBuildTargetType::Unknown:
+	case EBuildTargetType::Program:
+	default:
+		UE_LOG(LogRewindDebugger, Error, TEXT("%hs Changing transport mode in the selected target type is not supported. Target type [%s]"), __func__, LexToString(InTargetType));
+	}
+}
+
 URewindDebuggerSettings& URewindDebuggerSettings::Get()
 {
 	URewindDebuggerSettings* MutableCDO = GetMutableDefault<URewindDebuggerSettings>();
 	check(MutableCDO != nullptr)
-	
 	return *MutableCDO;
 }
 
diff --git a/Engine/Plugins/Animation/GameplayInsights/Source/RewindDebugger/Private/RewindDebuggerSettings.h b/Engine/Plugins/Animation/GameplayInsights/Source/RewindDebugger/Private/RewindDebuggerSettings.h
deleted file mode 100644
diff --git a/Engine/Plugins/Animation/GameplayInsights/Source/RewindDebugger/Public/RewindDebuggerSettings.h b/Engine/Plugins/Animation/GameplayInsights/Source/RewindDebugger/Public/RewindDebuggerSettings.h
new file mode 100644
--- /dev/null
+++ b/Engine/Plugins/Animation/GameplayInsights/Source/RewindDebugger/Public/RewindDebuggerSettings.h
@@ -0,0 +1,97 @@
+﻿// Copyright Epic Games, Inc. All Rights Reserved.
+
+#pragma once
+
+#include "CoreMinimal.h"
+#include "Engine/DeveloperSettings.h"
+#include "RewindDebuggerSettings.generated.h"
+
+class URewindDebuggerExtensionSettings;
+
+#define UE_API REWINDDEBUGGER_API
+
+namespace UE::TraceBasedDebuggers
+{
+enum class ETraceTransportMode : uint8;
+}
+
+UENUM()
+enum class ERewindDebuggerCameraMode : uint8
+{
+	Replay UMETA(Tooltip="Replay Recorded Camera"),
+	FollowTargetActor UMETA(Tooltip="Follow Target Actor"),
+	Disabled UMETA(Tooltip="Disable Camera On Playback"),
+};
+
+/**
+ * Implements the settings for the Rewind Debugger.
+ */
+UCLASS(config = EditorPerProjectUserSettings, meta=(DisplayName="Rewind Debugger"))
+class URewindDebuggerSettings : public UDeveloperSettings
+{
+	GENERATED_BODY()
+	
+public:
+	URewindDebuggerSettings();
+
+#if WITH_EDITOR
+	virtual FText GetSectionText() const override;
+	virtual FText GetSectionDescription() const override;
+#endif
+
+	virtual FName GetCategoryName() const override;
+
+	/** Returns the configured transport mode for the provided target type
+	 * @param InTargetType Build target type to evaluate
+	 */
+	UE_API UE::TraceBasedDebuggers::ETraceTransportMode GetTransportModeForTargetType(EBuildTargetType InTargetType) const;
+
+	/** Sets a specific transport mode for the provided target type
+	 * @param InTargetType Build target type to evaluate
+	 * @param InTransportMode New transport mode to use for the provided target
+	 */
+	UE_API void SetTransportModeForTargetType(EBuildTargetType InTargetType, UE::TraceBasedDebuggers::ETraceTransportMode InTransportMode);
+
+	/** Rewind Debugger Playback Camera Mode */
+	UPROPERTY(EditAnywhere, Config, Category = Camera)
+	ERewindDebuggerCameraMode CameraMode;
+	
+	/** If enabled, automatically detach player control when PIE is paused */
+	UPROPERTY(EditAnywhere, Config, Category = Other)
+	bool bShouldAutoEject;
+	
+	/** If enabled, start recording information at the start of PIE */
+	UPROPERTY(EditAnywhere, Config, Category = Other)
+	bool bShouldAutoRecordOnPIE;
+	
+	/** Playback speed multiplier */
+	UPROPERTY(EditAnywhere, Config, Category = Other)
+	float PlaybackRate = 1.0;
+
+	/** If enabled, show empty tracks on Rewind Debugger Timeline*/
+	UPROPERTY(EditAnywhere, Config, Category = Filters)
+	bool bShowEmptyObjectTracks;
+	
+	/** The track types listed here will be hidden from the track tree view */
+	UPROPERTY(EditAnywhere, Config, Category = Filters)
+	TArray<FName> HiddenTrackTypes;
+
+	
+	/** Currently selected target actor's name */
+	UPROPERTY(Config)
+	FString DebugTargetActor;
+	
+	/** Get Mutable CDO of URewindDebuggerSettings */
+	static URewindDebuggerSettings & Get();
+
+private:
+	/** Default Data Transport for Game Servers - Used to change the transport mode for the trace data. This is not saved in the config by design */
+	UPROPERTY(EditAnywhere, Category = Connection)
+	UE::TraceBasedDebuggers::ETraceTransportMode GameServerDataTransportMode;
+
+	/** Default Data Transport for Game Clients - Used to change the transport mode for the trace data. This is not saved in the config by design */
+	UPROPERTY(EditAnywhere, Category = Connection)
+	UE::TraceBasedDebuggers::ETraceTransportMode GameClientDataTransportMode;
+};
+
+#undef UE_API
\ No newline at end of file
