Fixed linker errors linking FNRewindDebuggerEditor in modular builds.

#rnx
#Synced-CL 50586249

diff --git a/Engine/Plugins/Animation/GameplayInsights/Source/RewindDebugger/Public/RewindDebuggerSettings.h b/Engine/Plugins/Animation/GameplayInsights/Source/RewindDebugger/Public/RewindDebuggerSettings.h
--- a/Engine/Plugins/Animation/GameplayInsights/Source/RewindDebugger/Public/RewindDebuggerSettings.h
+++ b/Engine/Plugins/Animation/GameplayInsights/Source/RewindDebugger/Public/RewindDebuggerSettings.h
@@ -82,7 +82,7 @@ public:
 	FString DebugTargetActor;
 	
 	/** Get Mutable CDO of URewindDebuggerSettings */
-	static URewindDebuggerSettings & Get();
+	UE_API static URewindDebuggerSettings & Get();
 
 private:
 	/** Default Data Transport for Game Servers - Used to change the transport mode for the trace data. This is not saved in the config by design */
