Fix issues with rewind debugger details tab stealing focus from the game

#Synced-CL 50686245

diff --git a/Engine/Plugins/Animation/GameplayInsights/Source/RewindDebugger/Private/RewindDebugger.cpp b/Engine/Plugins/Animation/GameplayInsights/Source/RewindDebugger/Private/RewindDebugger.cpp
--- a/Engine/Plugins/Animation/GameplayInsights/Source/RewindDebugger/Private/RewindDebugger.cpp
+++ b/Engine/Plugins/Animation/GameplayInsights/Source/RewindDebugger/Private/RewindDebugger.cpp
@@ -1404,15 +1404,18 @@ void FRewindDebugger::OpenDetailsPanel()
 void FRewindDebugger::TrackSelectionChanged(const TSharedPtr<RewindDebugger::FRewindDebuggerTrack> InSelectedTrack)
 {
 	SelectedTrack = InSelectedTrack;
+	if (!SelectedTrack.IsValid())
+	{
+		// when a selected track goes out of view, we get track selection changed notifications but we don't want to invoke any tab
+		return;
+	}
 
 	if (bIsDetailsPanelOpen)
 	{
 		const FLevelEditorModule& LevelEditorModule = FModuleManager::GetModuleChecked<FLevelEditorModule>("LevelEditor");
 		const TSharedPtr<FTabManager> LevelEditorTabManager = LevelEditorModule.GetLevelEditorTabManager();
 
-		// if we now have no selection, don't force the tab into focus - this happens when tracks disappear and can cause PIE to lose focus while playing
-		const bool bInvokeAsInactive = !SelectedTrack.IsValid();
-		const TSharedPtr<SDockTab> DetailsTab = LevelEditorTabManager->TryInvokeTab(FRewindDebuggerModule::DetailsTabName, bInvokeAsInactive);
+		const TSharedPtr<SDockTab> DetailsTab = LevelEditorTabManager->TryInvokeTab(FRewindDebuggerModule::DetailsTabName );
 
 		if (DetailsTab.IsValid())
 		{
