[RewindDebugger] updated FRewindDebuggerTrack::GetStepFrameTimeInternal to use a scrub info struct providing more information to derived classes in order to select the best scrub time to use when using the stepping forward/backward buttons.
VisualLogger track now allow users to jump directly to frames with events (i.e., skipping empty frames)

#Synced-CL 50865525

diff --git a/Engine/Plugins/Animation/GameplayInsights/Source/RewindDebugger/Private/RewindDebugger.cpp b/Engine/Plugins/Animation/GameplayInsights/Source/RewindDebugger/Private/RewindDebugger.cpp
--- a/Engine/Plugins/Animation/GameplayInsights/Source/RewindDebugger/Private/RewindDebugger.cpp
+++ b/Engine/Plugins/Animation/GameplayInsights/Source/RewindDebugger/Private/RewindDebugger.cpp
@@ -728,7 +728,7 @@ void FRewindDebugger::Play()
 	{
 		// Loop from the beginning if already at the end,
 		// otherwise keep the current time but make sure to disable the auto-scroll
-		const double AdjustedTime = (CurrentScrubTime >= RecordingDuration.Get()) ? 0 : CurrentScrubTime;
+		const double AdjustedTime = (ScrubTimeInformation.ElapsedTime >= RecordingDuration.Get()) ? 0 : ScrubTimeInformation.ElapsedTime;
 		SetCurrentScrubTime(AdjustedTime, EAutoScrollChange::Disable);
 		TrackCursorDelegate.ExecuteIfBound(/*bReverse*/false);
 
@@ -747,7 +747,7 @@ void FRewindDebugger::PlayReverse()
 	{
 		// Loop from the end if already at the beginning,
 		// otherwise keep the current time but make sure to disable the auto-scroll
-		const double AdjustedTime = (CurrentScrubTime <= 0) ? RecordingDuration.Get() : CurrentScrubTime;
+		const double AdjustedTime = (ScrubTimeInformation.ElapsedTime <= 0) ? RecordingDuration.Get() : ScrubTimeInformation.ElapsedTime;
 		SetCurrentScrubTime(AdjustedTime, EAutoScrollChange::Disable);
 		TrackCursorDelegate.ExecuteIfBound(/*bReverse*/true);
 
@@ -793,7 +793,7 @@ void FRewindDebugger::Step(const int32 InNumberOfFrames)
 				InNumberOfFrames == 1
 					? RewindDebugger::EStepMode::Forward
 					: RewindDebugger::EStepMode::Backward
-				, CurrentScrubTime);
+				, ScrubTimeInformation);
 
 			if (NewScrubTime.IsSet())
 			{
@@ -905,15 +905,16 @@ void FRewindDebugger::SetCurrentViewRange(const TRange<double>& Range)
 
 void FRewindDebugger::SetCurrentScrubTime(const double Time, const EAutoScrollChange AutoScrollChange)
 {
-	CurrentScrubTime = Time;
-
 	if (const TraceServices::IAnalysisSession* Session = GetAnalysisSession())
 	{
-		GetScrubTimeInformation(CurrentScrubTime, ScrubTimeInformation, RecordingIndex, Session);
+		GetScrubTimeInformation(Time, ScrubTimeInformation, RecordingIndex, Session);
 
 		TraceTime.Set(ScrubTimeInformation.ProfileTime);
 	}
 
+	// Enforce the specific Time that was provided
+	ScrubTimeInformation.ElapsedTime = Time;
+
 	if (bAutoScrollToLastItem
 		&& AutoScrollChange == EAutoScrollChange::Disable)
 	{
@@ -1043,6 +1044,7 @@ void FRewindDebugger::GetScrubTimeInformation(double InDebugTime, FScrubTimeInfo
 				const FRecordingInfoMessage& Event = Recording->GetEvent(ScrubFrameIndex);
 				InOutTimeInformation.FrameIndex = ScrubFrameIndex;
 				InOutTimeInformation.ProfileTime = Event.ProfileTime;
+				InOutTimeInformation.ElapsedTime = Event.ElapsedTime;
 			}
 		}
 	}
@@ -1354,6 +1356,7 @@ void FRewindDebugger::Tick(float DeltaTime)
 		}
 		else
 		{
+			const double CurrentScrubTime = ScrubTimeInformation.ElapsedTime;
 			if (RecordingDuration.Get() > 0 && CurrentScrubTime <= RecordingDuration.Get())
 			{
 				if (ControlState == EControlState::Play || ControlState == EControlState::PlayReverse)
diff --git a/Engine/Plugins/Animation/GameplayInsights/Source/RewindDebugger/Private/RewindDebugger.h b/Engine/Plugins/Animation/GameplayInsights/Source/RewindDebugger/Private/RewindDebugger.h
--- a/Engine/Plugins/Animation/GameplayInsights/Source/RewindDebugger/Private/RewindDebugger.h
+++ b/Engine/Plugins/Animation/GameplayInsights/Source/RewindDebugger/Private/RewindDebugger.h
@@ -1,4 +1,4 @@
-﻿// Copyright Epic Games, Inc. All Rights Reserved.
+// Copyright Epic Games, Inc. All Rights Reserved.
 #pragma once
 
 #if UE_ENABLE_INCLUDE_ORDER_DEPRECATED_IN_5_7
@@ -55,7 +55,7 @@ public:
 
 	virtual double GetScrubTime() const override
 	{
-		return CurrentScrubTime;
+		return ScrubTimeInformation.ElapsedTime;
 	}
 
 	virtual const TRange<double>& GetCurrentTraceRange() const override
@@ -365,19 +365,10 @@ private:
 	bool bAutoScrollToLastItem = false;
 
 	double PreviousTraceTime = -1;
-	double CurrentScrubTime = 0;
 	TRange<double> CurrentViewRange{ 0, 10 };
 	TRange<double> CurrentTraceRange{ 0, 0 };
 	uint16 RecordingIndex = 0;
 
-	struct FScrubTimeInformation
-	{
-		/** Profiling/Tracing time */
-		double ProfileTime = 0;
-		/** Scrub Frame Index */
-		int64 FrameIndex = 0;
-	};
-
 	FScrubTimeInformation ScrubTimeInformation;
 	FScrubTimeInformation LowerBoundViewTimeInformation;
 	FScrubTimeInformation UpperBoundViewTimeInformation;
diff --git a/Engine/Plugins/Animation/GameplayInsights/Source/RewindDebugger/Private/RewindDebuggerObjectTrack.cpp b/Engine/Plugins/Animation/GameplayInsights/Source/RewindDebugger/Private/RewindDebuggerObjectTrack.cpp
--- a/Engine/Plugins/Animation/GameplayInsights/Source/RewindDebugger/Private/RewindDebuggerObjectTrack.cpp
+++ b/Engine/Plugins/Animation/GameplayInsights/Source/RewindDebugger/Private/RewindDebuggerObjectTrack.cpp
@@ -167,7 +167,7 @@ FText FRewindDebuggerObjectTrack::GetStepCommandTooltipInternal(const EStepMode
 	return {};
 }
 
-TOptional<double> FRewindDebuggerObjectTrack::GetStepFrameTimeInternal(const EStepMode StepMode, const double CurrentScrubTime) const
+TOptional<double> FRewindDebuggerObjectTrack::GetStepFrameTimeInternal(const EStepMode StepMode, const FScrubTimeInformation& CurrentScrubTime) const
 {
 	// Allow primary child track to override step frame time
 	if (PrimaryChildTrack.IsValid())
diff --git a/Engine/Plugins/Animation/GameplayInsights/Source/RewindDebugger/Private/RewindDebuggerObjectTrack.h b/Engine/Plugins/Animation/GameplayInsights/Source/RewindDebugger/Private/RewindDebuggerObjectTrack.h
--- a/Engine/Plugins/Animation/GameplayInsights/Source/RewindDebugger/Private/RewindDebuggerObjectTrack.h
+++ b/Engine/Plugins/Animation/GameplayInsights/Source/RewindDebugger/Private/RewindDebuggerObjectTrack.h
@@ -39,7 +39,7 @@ private:
 	virtual bool HasDebugDataInternal() const override;
 	virtual bool HandleDoubleClickInternal() override;
 	virtual FText GetStepCommandTooltipInternal(EStepMode) const override;
-	virtual TOptional<double> GetStepFrameTimeInternal(EStepMode, double) const override;
+	virtual TOptional<double> GetStepFrameTimeInternal(EStepMode, const FScrubTimeInformation&) const override;
 
 	mutable FText DisplayName;
 	FString ObjectName;
diff --git a/Engine/Plugins/Animation/GameplayInsights/Source/RewindDebuggerVLog/Private/VisualLogTrack.cpp b/Engine/Plugins/Animation/GameplayInsights/Source/RewindDebuggerVLog/Private/VisualLogTrack.cpp
--- a/Engine/Plugins/Animation/GameplayInsights/Source/RewindDebuggerVLog/Private/VisualLogTrack.cpp
+++ b/Engine/Plugins/Animation/GameplayInsights/Source/RewindDebuggerVLog/Private/VisualLogTrack.cpp
@@ -1,6 +1,8 @@
 // Copyright Epic Games, Inc. All Rights Reserved.
 
 #include "VisualLogTrack.h"
+
+#include "IGameplayProvider.h"
 #include "IRewindDebugger.h"
 #include "Modules/ModuleManager.h"
 #include "PropertyEditorModule.h"
@@ -195,6 +197,125 @@ TSharedPtr<SWidget> FVisualLogCategoryTrack::GetTimelineViewInternal()
 		.EventData_Raw(this, &FVisualLogCategoryTrack::GetEventData);
 }
 
+FText FVisualLogCategoryTrack::GetStepCommandTooltipInternal(const RewindDebugger::EStepMode StepMode) const
+{
+	switch (StepMode)
+	{
+	case RewindDebugger::EStepMode::Forward:
+		return LOCTEXT("StepForwardTooltip", "Step to next event in Logging track");
+	case RewindDebugger::EStepMode::Backward:
+		return LOCTEXT("StepBackwardTooltip", "Step to previous event in Logging track");
+	}
+
+	return {};
+}
+
+TOptional<double> FVisualLogCategoryTrack::GetStepFrameTimeInternal(const EStepMode StepMode, const FScrubTimeInformation& CurrentScrubTime) const
+{
+	const IRewindDebugger* RewindDebugger = IRewindDebugger::Instance();
+	if (RewindDebugger == nullptr)
+	{
+		return {};
+	}
+
+	const TraceServices::IAnalysisSession* AnalysisSession = RewindDebugger->GetAnalysisSession();
+	if (AnalysisSession == nullptr)
+	{
+		return {};
+	}
+
+	const IGameplayProvider* GameplayProvider = AnalysisSession->ReadProvider<IGameplayProvider>("GameplayProvider");
+	if (GameplayProvider == nullptr)
+	{
+		return {};
+	}
+
+	const FVisualLoggerProvider* VisLogProvider = AnalysisSession->ReadProvider<FVisualLoggerProvider>(FVisualLoggerProvider::ProviderName);
+	if (VisLogProvider == nullptr)
+	{
+		return {};
+	}
+
+	const TraceServices::IFrameProvider& FramesProvider = TraceServices::ReadFrameProvider(*AnalysisSession);
+	TraceServices::FAnalysisSessionReadScope SessionReadScope(*AnalysisSession);
+
+	// Legacy Index that is always 0 now
+	constexpr int32 RecordingIndex = 0;
+	const IGameplayProvider::RecordingInfoTimeline* Recording = GameplayProvider->GetRecordingInfo(RecordingIndex);
+	if (Recording == nullptr)
+	{
+		return {};
+	}
+
+	TraceServices::FFrame Frame;
+	if (FramesProvider.GetFrameFromTime(TraceFrameType_Game, CurrentScrubTime.ProfileTime, Frame))
+	{
+		bool bEventFound = false;
+
+		// Set the frame index increment based on the step mode
+		const int32 FrameInc = (StepMode == EStepMode::Backward) ? -1 : 1;
+
+		// Set the index of our first candidate
+		// Note that we don't care about decrementing pass value 0 since GetFrame will simply fail
+		uint64 CandidateFrameIndex = Frame.Index + FrameInc;
+
+		while (const TraceServices::FFrame* CandidateFrame = FramesProvider.GetFrame(TraceFrameType_Game, CandidateFrameIndex))
+		{
+			// Inspect the VisLog timeline to see if there are events for that frame
+			VisLogProvider->ReadVisualLogEntryTimeline(ObjectId,
+				[&CandidateFrame, &bEventFound, Category = Category](const FVisualLoggerProvider::VisualLogEntryTimeline& InTimeline)
+				{
+					InTimeline.EnumerateEvents(CandidateFrame->StartTime, CandidateFrame->EndTime,
+						[&bEventFound, Category](double InStartTime, double InEndTime, uint32 InDepth, const FVisualLogEntry& InMessage)
+						{
+							for (const FVisualLogShapeElement& Element : InMessage.ElementsToDraw)
+							{
+								if (Element.Category == Category)
+								{
+									bEventFound = true;
+									break;
+								}
+							}
+							for (const FVisualLogLine& Line : InMessage.LogLines)
+							{
+								if (Line.Category == Category)
+								{
+									bEventFound = true;
+								}
+							}
+							return bEventFound ? TraceServices::EEventEnumerate::Stop : TraceServices::EEventEnumerate::Continue;
+						});
+				});
+
+			// Event was found
+			if (bEventFound)
+			{
+				const FRecordingInfoMessage* MatchingRecordingInfo = nullptr;
+				Recording->EnumerateEvents(CandidateFrame->StartTime, CandidateFrame->EndTime, 
+					[&MatchingRecordingInfo](double InStartTime, double InEndTime, uint32 InDepth, const FRecordingInfoMessage& InMessage)
+					{
+						MatchingRecordingInfo = &InMessage;
+						return TraceServices::EEventEnumerate::Stop;
+					});
+
+				if (!ensureMsgf(MatchingRecordingInfo, TEXT("We expect to always have recording info for any frame with events")))
+				{
+					// Fallback on default behavior
+					return {};
+				}
+
+				return MatchingRecordingInfo->ElapsedTime;
+			}
+
+			// Try prev/next frame
+			CandidateFrameIndex += FrameInc;
+		}
+	}
+
+	// Everything was properly analyzed and there is simply no more events before/after so preserve the current time
+	return CurrentScrubTime.ElapsedTime;
+}
+
 //----------------------------------------------------------------------//
 // FVisualLogTrack
 //----------------------------------------------------------------------//
diff --git a/Engine/Plugins/Animation/GameplayInsights/Source/RewindDebuggerVLog/Private/VisualLogTrack.h b/Engine/Plugins/Animation/GameplayInsights/Source/RewindDebuggerVLog/Private/VisualLogTrack.h
--- a/Engine/Plugins/Animation/GameplayInsights/Source/RewindDebuggerVLog/Private/VisualLogTrack.h
+++ b/Engine/Plugins/Animation/GameplayInsights/Source/RewindDebuggerVLog/Private/VisualLogTrack.h
@@ -28,6 +28,8 @@ private:
 	virtual bool UpdateInternal() override;
 	virtual TSharedPtr<SWidget> GetDetailsViewInternal() override;
 	virtual TSharedPtr<SWidget> GetTimelineViewInternal() override;
+	virtual FText GetStepCommandTooltipInternal(EStepMode) const override;
+	virtual TOptional<double> GetStepFrameTimeInternal(EStepMode, const FScrubTimeInformation&) const override;
 
 	virtual FSlateIcon GetIconInternal() override
 	{
diff --git a/Engine/Plugins/Runtime/StateTree/Source/StateTreeEditorModule/Private/Debugger/StateTreeRewindDebuggerTrack.cpp b/Engine/Plugins/Runtime/StateTree/Source/StateTreeEditorModule/Private/Debugger/StateTreeRewindDebuggerTrack.cpp
--- a/Engine/Plugins/Runtime/StateTree/Source/StateTreeEditorModule/Private/Debugger/StateTreeRewindDebuggerTrack.cpp
+++ b/Engine/Plugins/Runtime/StateTree/Source/StateTreeEditorModule/Private/Debugger/StateTreeRewindDebuggerTrack.cpp
@@ -250,9 +250,10 @@ FText FRewindDebuggerTrack::GetStepCommandTooltipInternal(const RewindDebugger::
 	return {};
 }
 
-TOptional<double> FRewindDebuggerTrack::GetStepFrameTimeInternal(const RewindDebugger::EStepMode StepMode, const double CurrentScrubTime) const
+TOptional<double> FRewindDebuggerTrack::GetStepFrameTimeInternal(const RewindDebugger::EStepMode StepMode, const FScrubTimeInformation& ScrubTimeInformation) const
 {
 	FScrubState ScrubState(&EventCollection);
+	const double CurrentScrubTime = ScrubTimeInformation.ElapsedTime;
 	ScrubState.SetScrubTime(CurrentScrubTime, /*ForceRefresh*/true);
 	if (ScrubState.IsInBounds())
 	{
diff --git a/Engine/Plugins/Runtime/StateTree/Source/StateTreeEditorModule/Private/Debugger/StateTreeRewindDebuggerTrack.h b/Engine/Plugins/Runtime/StateTree/Source/StateTreeEditorModule/Private/Debugger/StateTreeRewindDebuggerTrack.h
--- a/Engine/Plugins/Runtime/StateTree/Source/StateTreeEditorModule/Private/Debugger/StateTreeRewindDebuggerTrack.h
+++ b/Engine/Plugins/Runtime/StateTree/Source/StateTreeEditorModule/Private/Debugger/StateTreeRewindDebuggerTrack.h
@@ -57,7 +57,7 @@ private:
 	virtual uint64 GetObjectIdInternal() const override;
 	virtual bool HandleDoubleClickInternal() override;
 	virtual FText GetStepCommandTooltipInternal(RewindDebugger::EStepMode) const override;
-	virtual TOptional<double> GetStepFrameTimeInternal(RewindDebugger::EStepMode, double) const override;
+	virtual TOptional<double> GetStepFrameTimeInternal(RewindDebugger::EStepMode, const FScrubTimeInformation&) const override;
 	//~ End RewindDebugger::FRewindDebuggerTrack interface
 
 	FInstanceEventCollection EventCollection;
diff --git a/Engine/Source/Editor/RewindDebuggerInterface/Public/IRewindDebugger.h b/Engine/Source/Editor/RewindDebuggerInterface/Public/IRewindDebugger.h
--- a/Engine/Source/Editor/RewindDebuggerInterface/Public/IRewindDebugger.h
+++ b/Engine/Source/Editor/RewindDebuggerInterface/Public/IRewindDebugger.h
@@ -60,6 +60,16 @@ struct FDebugObjectInfo
 #endif
 };
 
+struct FScrubTimeInformation
+{
+	/** Elapsed simulation/world time */
+	double ElapsedTime = 0;
+	/** Profiling/Tracing time */
+	double ProfileTime = 0;
+	/** Scrub Frame Index */
+	int64 FrameIndex = 0;
+};
+
 UCLASS(MinimalAPI)
 class URewindDebuggerTrackContextMenuContext : public UObject
 {
diff --git a/Engine/Source/Editor/RewindDebuggerInterface/Public/RewindDebuggerTrack.h b/Engine/Source/Editor/RewindDebuggerInterface/Public/RewindDebuggerTrack.h
--- a/Engine/Source/Editor/RewindDebuggerInterface/Public/RewindDebuggerTrack.h
+++ b/Engine/Source/Editor/RewindDebuggerInterface/Public/RewindDebuggerTrack.h
@@ -156,8 +156,17 @@ public:
 		return GetStepCommandTooltipInternal(StepMode);
 	}
 
+	UE_DEPRECATED(5.8, "Use overload taking a scrub time information struct instead.")
 	/** tracks can return a different time for the step action; default the main recording frame time is used */
 	TOptional<double> GetStepFrameTime(const EStepMode StepMode, const double CurrentScrubTime) const
+	{
+PRAGMA_DISABLE_DEPRECATION_WARNINGS
+		return GetStepFrameTimeInternal(StepMode, CurrentScrubTime);
+PRAGMA_ENABLE_DEPRECATION_WARNINGS
+	}
+
+	/** tracks can return a different time for the step action; default the main recording frame time is used */
+	TOptional<double> GetStepFrameTime(const EStepMode StepMode, const FScrubTimeInformation& CurrentScrubTime) const
 	{
 		return GetStepFrameTimeInternal(StepMode, CurrentScrubTime);
 	}
@@ -317,7 +326,13 @@ private:
 		return {};
 	}
 
-	virtual TOptional<double> GetStepFrameTimeInternal(EStepMode, double) const
+	UE_DEPRECATED(5.8, "Use overload taking a scrub time information struct instead.")
+	virtual TOptional<double> GetStepFrameTimeInternal(EStepMode, const double) const final
+	{
+		return {};
+	}
+
+	virtual TOptional<double> GetStepFrameTimeInternal(EStepMode, const FScrubTimeInformation&) const
 	{
 		return {};
 	}
