[RewindDebugger] improved file analysis support
- prevent remote session update from closing an active file analysis
- clear active file analysis when starting a new analysis (from file or remote)

#rnx

#Synced-CL 50865525

diff --git a/Engine/Plugins/Animation/GameplayInsights/Source/RewindDebugger/Private/RewindDebugger.cpp b/Engine/Plugins/Animation/GameplayInsights/Source/RewindDebugger/Private/RewindDebugger.cpp
--- a/Engine/Plugins/Animation/GameplayInsights/Source/RewindDebugger/Private/RewindDebugger.cpp
+++ b/Engine/Plugins/Animation/GameplayInsights/Source/RewindDebugger/Private/RewindDebugger.cpp
@@ -476,6 +476,13 @@ void FRewindDebugger::HandleSessionSelected(const TSharedPtr<UE::TraceBasedDebug
 {
 	using namespace UE::TraceBasedDebuggers;
 
+	// If a file analysis is active we don't want to stop it
+	if (SelectedDescriptor.IsValid()
+		&& !SelectedDescriptor.RemoteSessionID.IsValid())
+	{
+		return;
+	}
+
 	const FTraceSessionDescriptor* ExistingDescriptor = ActiveAnalysisDescriptors.FindByPredicate(
 		[SessionInfo = InSessionInfo](const FTraceSessionDescriptor& Descriptor)
 		{
@@ -1181,7 +1188,13 @@ void FRewindDebugger::OpenAnalysisSession(const UE::TraceBasedDebuggers::FTraceS
 	const TSharedPtr<FTraceSessionsManager> TraceManager = FRewindDebuggerModule::Get().GetTraceSessionsManager();
 	ActiveAnalysisDescriptors.RemoveAll([&InSessionDescriptor, TraceManager](const FTraceSessionDescriptor& Element)
 		{
-			if (Element.RemoteSessionID == InSessionDescriptor.RemoteSessionID)
+			// We currently don't have a way to reselect an analysis session from a trace file
+			// so we can clear it when a new analysis starts (from file or remote).
+			// In the case or remote sessions, we close only if from the same remote session Id.
+			const bool bClosePreviousSession = !Element.RemoteSessionID.IsValid()
+				|| Element.RemoteSessionID == InSessionDescriptor.RemoteSessionID;
+
+			if (bClosePreviousSession)
 			{
 				if (TraceManager)
 				{
