Fix crash in rewind debugger when a new recording is started with data in the details panel
#Synced-CL 51055096

diff --git a/Engine/Plugins/Animation/GameplayInsights/Source/RewindDebugger/Private/RewindDebugger.cpp b/Engine/Plugins/Animation/GameplayInsights/Source/RewindDebugger/Private/RewindDebugger.cpp
--- a/Engine/Plugins/Animation/GameplayInsights/Source/RewindDebugger/Private/RewindDebugger.cpp
+++ b/Engine/Plugins/Animation/GameplayInsights/Source/RewindDebugger/Private/RewindDebugger.cpp
@@ -630,7 +630,7 @@ void FRewindDebugger::ClearTrackData()
 	RecordingDuration.Set(0.0);
 	SetCurrentScrubTime(0.0, EAutoScrollChange::PreserveCurrentState);
 
-	TrackSelectionChanged(nullptr);
+	ClearTrackSelection();
 
 	// update extensions
 	UE::RewindDebugger::IterateExtensions([this](IRewindDebuggerExtension* Extension)
@@ -1417,6 +1417,24 @@ void FRewindDebugger::OpenDetailsPanel()
 	TrackSelectionChanged(SelectedTrack);
 }
 
+void FRewindDebugger::ClearTrackSelection()
+{
+	SelectedTrack.Reset();
+
+	if (bIsDetailsPanelOpen)
+	{
+		const FLevelEditorModule& LevelEditorModule = FModuleManager::GetModuleChecked<FLevelEditorModule>("LevelEditor");
+		const TSharedPtr<FTabManager> LevelEditorTabManager = LevelEditorModule.GetLevelEditorTabManager();
+
+		const TSharedPtr<SDockTab> DetailsTab = LevelEditorTabManager->TryInvokeTab(FRewindDebuggerModule::DetailsTabName);
+
+		if (DetailsTab.IsValid())
+		{
+			UpdateDetailsPanel(DetailsTab.ToSharedRef());
+		}
+	}
+}
+
 void FRewindDebugger::TrackSelectionChanged(const TSharedPtr<RewindDebugger::FRewindDebuggerTrack> InSelectedTrack)
 {
 	SelectedTrack = InSelectedTrack;
diff --git a/Engine/Plugins/Animation/GameplayInsights/Source/RewindDebugger/Private/RewindDebugger.h b/Engine/Plugins/Animation/GameplayInsights/Source/RewindDebugger/Private/RewindDebugger.h
--- a/Engine/Plugins/Animation/GameplayInsights/Source/RewindDebugger/Private/RewindDebugger.h
+++ b/Engine/Plugins/Animation/GameplayInsights/Source/RewindDebugger/Private/RewindDebugger.h
@@ -205,6 +205,8 @@ public:
 
 	void TrackDoubleClicked(TSharedPtr<RewindDebugger::FRewindDebuggerTrack> InSelectedTrack);
 	void TrackSelectionChanged(TSharedPtr<RewindDebugger::FRewindDebuggerTrack> InSelectedTrack);
+	void ClearTrackSelection();
+
 	TSharedPtr<SWidget> BuildTrackContextMenu() const;
 
 	void UpdateDetailsPanel(TSharedRef<SDockTab> DetailsTab);
