prevent details panel updates in rewind debugger from force focussing the tab

#Synced-CL 51513499

diff --git a/Engine/Plugins/Animation/GameplayInsights/Source/RewindDebugger/Private/RewindDebugger.cpp b/Engine/Plugins/Animation/GameplayInsights/Source/RewindDebugger/Private/RewindDebugger.cpp
--- a/Engine/Plugins/Animation/GameplayInsights/Source/RewindDebugger/Private/RewindDebugger.cpp
+++ b/Engine/Plugins/Animation/GameplayInsights/Source/RewindDebugger/Private/RewindDebugger.cpp
@@ -1426,11 +1426,15 @@ void FRewindDebugger::ClearTrackSelection()
 		const FLevelEditorModule& LevelEditorModule = FModuleManager::GetModuleChecked<FLevelEditorModule>("LevelEditor");
 		const TSharedPtr<FTabManager> LevelEditorTabManager = LevelEditorModule.GetLevelEditorTabManager();
 
-		const TSharedPtr<SDockTab> DetailsTab = LevelEditorTabManager->TryInvokeTab(FRewindDebuggerModule::DetailsTabName);
-
-		if (DetailsTab.IsValid())
+		if (DetailsPanelWidget.IsValid())
 		{
-			UpdateDetailsPanel(DetailsTab.ToSharedRef());
+			// avoid actually summoning the tab if we are only clearing it
+				
+			if (EmptyDetails == nullptr)
+			{
+				EmptyDetails = SNew(SSpacer);
+			}
+			DetailsPanelWidget->SetContent(EmptyDetails.ToSharedRef());
 		}
 	}
 }
@@ -1462,6 +1466,12 @@ void FRewindDebugger::UpdateDetailsPanel(TSharedRef<SDockTab> DetailsTab)
 {
 	if (bIsDetailsPanelOpen)
 	{
+		if (!DetailsPanelWidget.IsValid())
+		{
+			DetailsPanelWidget = SNew(SBox);
+			DetailsTab->SetContent(DetailsPanelWidget.ToSharedRef());
+		}
+		
 		TSharedPtr<SWidget> DetailsView;
 
 		if (SelectedTrack)
@@ -1469,19 +1479,16 @@ void FRewindDebugger::UpdateDetailsPanel(TSharedRef<SDockTab> DetailsTab)
 			DetailsView = SelectedTrack->GetDetailsView();
 		}
 
-		if (DetailsView)
-		{
-			DetailsTab->SetContent(DetailsView.ToSharedRef());
-		}
-		else
+		if (!DetailsView.IsValid())
 		{
-			static TSharedPtr<SWidget> EmptyDetails;
 			if (EmptyDetails == nullptr)
 			{
 				EmptyDetails = SNew(SSpacer);
 			}
-			DetailsTab->SetContent(EmptyDetails.ToSharedRef());
+			DetailsView = EmptyDetails;
 		}
+		
+		DetailsPanelWidget->SetContent(DetailsView.ToSharedRef());
 	}
 }
 
diff --git a/Engine/Plugins/Animation/GameplayInsights/Source/RewindDebugger/Private/RewindDebugger.h b/Engine/Plugins/Animation/GameplayInsights/Source/RewindDebugger/Private/RewindDebugger.h
--- a/Engine/Plugins/Animation/GameplayInsights/Source/RewindDebugger/Private/RewindDebugger.h
+++ b/Engine/Plugins/Animation/GameplayInsights/Source/RewindDebugger/Private/RewindDebugger.h
@@ -12,6 +12,7 @@
 #include "IRewindDebuggerTrackCreator.h"
 #include "RewindDebuggerTrack.h"
 #include "TraceSessionsManager.h"
+#include "Widgets/Layout/SBox.h"
 
 namespace TraceServices
 {
@@ -219,6 +220,9 @@ public:
 	DECLARE_DELEGATE_OneParam(FOnTrackCursor, bool)
 	void SetTrackCursorDelegate(const FOnTrackCursor& InTrackCursorDelegate);
 
+	TSharedPtr<SBox> DetailsPanelWidget;
+	TSharedPtr<SWidget> EmptyDetails;
+
 	TBindableProperty<double>* GetTraceTimeProperty()
 	{
 		return &TraceTime;
