[RewindDebugger]
Add configurable allowed class types filter to the Rewind Debugger target selector

Replaces the hardcoded ULevel/UWorld child filter in the object selector dropdown
with a configurable allowlist of class types.

- Adds a new URewindDebuggerProjectSettings class (project-scoped, saved to
RewindDebugger.ini) with a SelectorAllowedTypes array, defaulting to AActor
- Extends URewindDebuggerSettings (user-scoped) with a matching SelectorAllowedTypes
array that can add types on top of the project settings but not override them
- The selector now filters by checking IsSubClassOf against the combined allowed
types; an empty list falls back to allowing everything
- Renames URewindDebuggerSettings display name to 'Rewind Debugger (user preferences)'
and fixes VLog settings to use its own localization namespace to avoid collision

#Synced-CL 51685953

diff --git a/Engine/Plugins/Animation/GameplayInsights/Source/RewindDebugger/Private/RewindDebuggerSettings.cpp b/Engine/Plugins/Animation/GameplayInsights/Source/RewindDebugger/Private/RewindDebuggerSettings.cpp
--- a/Engine/Plugins/Animation/GameplayInsights/Source/RewindDebugger/Private/RewindDebuggerSettings.cpp
+++ b/Engine/Plugins/Animation/GameplayInsights/Source/RewindDebugger/Private/RewindDebuggerSettings.cpp
@@ -2,12 +2,26 @@
 
 #include "RewindDebuggerSettings.h"
 
+#include "GameFramework/Actor.h"
 #include "IRewindDebugger.h"
-#include "TraceDataRelayTransport.h"
 #include "Misc/CoreDelegates.h"
+#include "TraceDataRelayTransport.h"
+#include "UObject/SoftObjectPath.h"
 
 #include UE_INLINE_GENERATED_CPP_BY_NAME(RewindDebuggerSettings)
 
+#define LOCTEXT_NAMESPACE "RewindDebuggerSettings"
+
+URewindDebuggerProjectSettings::URewindDebuggerProjectSettings()
+{
+	SelectorAllowedTypes.Add(FSoftClassPath(AActor::StaticClass()));
+}
+
+FName URewindDebuggerProjectSettings::GetCategoryName() const
+{
+	return FName(TEXT("Plugins"));
+}
+
 URewindDebuggerSettings::URewindDebuggerSettings()
 	: CameraMode(ERewindDebuggerCameraMode::Replay)
 	, bShouldAutoEject(false)
@@ -15,22 +29,32 @@ URewindDebuggerSettings::URewindDebuggerSettings()
 	, GameServerDataTransportMode(UE::TraceBasedDebuggers::ETraceTransportMode::Direct)
 	, GameClientDataTransportMode(UE::TraceBasedDebuggers::ETraceTransportMode::Direct)
 {
-	FCoreDelegates::OnPreExit.AddLambda([]()
+	OnPreExitHandle = FCoreDelegates::OnPreExit.AddLambda([]()
 	{
-		Get().SaveConfig();
+		GetMutableDefault<URewindDebuggerSettings>()->SaveConfig();
 	});
 }
 
+URewindDebuggerSettings::~URewindDebuggerSettings()
+{
+	FCoreDelegates::OnPreExit.Remove(OnPreExitHandle);
+}
+
 #if WITH_EDITOR
 
+FText URewindDebuggerProjectSettings::GetSectionDescription() const
+{
+	return LOCTEXT("RewindDebuggerProjectSettingsDescription", "Configure project options for the Rewind Debugger.");
+}
+
 FText URewindDebuggerSettings::GetSectionText() const
 {
-	return NSLOCTEXT("RewindDebugger", "RewindDebuggerSettingsName", "Rewind Debugger");
+	return LOCTEXT("RewindDebuggerSettingsName", "Rewind Debugger");
 }
 
 FText URewindDebuggerSettings::GetSectionDescription() const
 {
-	return NSLOCTEXT("RewindDebugger", "RewindDebuggerSettingsDescription", "Configure options for the Rewind Debugger.");
+	return LOCTEXT("RewindDebuggerSettingsDescription", "Configure user options for the Rewind Debugger.");
 }
 
 #endif
@@ -83,3 +107,4 @@ URewindDebuggerSettings& URewindDebuggerSettings::Get()
 	return *MutableCDO;
 }
 
+#undef LOCTEXT_NAMESPACE
\ No newline at end of file
diff --git a/Engine/Plugins/Animation/GameplayInsights/Source/RewindDebugger/Private/SRewindDebugger.cpp b/Engine/Plugins/Animation/GameplayInsights/Source/RewindDebugger/Private/SRewindDebugger.cpp
--- a/Engine/Plugins/Animation/GameplayInsights/Source/RewindDebugger/Private/SRewindDebugger.cpp
+++ b/Engine/Plugins/Animation/GameplayInsights/Source/RewindDebugger/Private/SRewindDebugger.cpp
@@ -1,4 +1,4 @@
-﻿// Copyright Epic Games, Inc. All Rights Reserved.
+// Copyright Epic Games, Inc. All Rights Reserved.
 
 #include "SRewindDebugger.h"
 #include "ActorPickerMode.h"
@@ -22,6 +22,7 @@
 #include "Selection.h"
 #include "SSimpleTimeSlider.h"
 #include "StatusBarSubsystem.h"
+#include "StructUtils/InstancedStruct.h"
 #include "Styling/SlateIconFinder.h"
 #include "ToolMenu.h"
 #include "ToolMenus.h"
@@ -212,34 +213,62 @@ TSharedRef<SWidget> SRewindDebugger::MakeObjectSelectionMenu()
 			TraceServices::FAnalysisSessionReadScope SessionReadScope(*Session);
 			const IGameplayProvider* GameplayProvider = Session->ReadProvider<IGameplayProvider>("GameplayProvider");
 
-			const FClassInfo* LevelClassInfo = GameplayProvider->FindClassInfo(*ULevel::StaticClass()->GetPathName());
-			const FClassInfo* WorldClassInfo = GameplayProvider->FindClassInfo(*UWorld::StaticClass()->GetPathName());
+			TArray<FString> AllowedTypeNames;
 
-			// Accept any child of the World or Level
-			GameplayProvider->EnumerateObjects([this, &MenuBuilder, GameplayProvider, LevelClassInfo, WorldClassInfo](const FObjectInfo& ObjectInfo)
+			// Gather types allowed by the project
+			const URewindDebuggerProjectSettings* ProjectSettings = GetDefault<URewindDebuggerProjectSettings>();
+			Algo::Transform(ProjectSettings->SelectorAllowedTypes, AllowedTypeNames, [](const FSoftClassPath& Path) { return Path.ToString(); });
+			Algo::Transform(ProjectSettings->SelectorAllowedStructTypes, AllowedTypeNames, [](const FInstancedStruct& Type) { return GetPathNameSafe(Type.GetScriptStruct()); });
+
+			// Gather types allowed by the user
+			const URewindDebuggerSettings* UserSettings = GetDefault<URewindDebuggerSettings>();
+			Algo::Transform(UserSettings->SelectorAllowedTypes, AllowedTypeNames, [](const FSoftClassPath& Path) { return Path.ToString(); });
+			Algo::Transform(UserSettings->SelectorAllowedStructTypes, AllowedTypeNames, [](const FInstancedStruct& Type) { return GetPathNameSafe(Type.GetScriptStruct()); });
+
+			TArray<const uint64> AllowedTypeIds;
+			AllowedTypeIds.Reserve(AllowedTypeNames.Num());
+			for (const FString& AllowedTypeName : AllowedTypeNames)
+			{
+				if (const FClassInfo* ClassInfo = GameplayProvider->FindClassInfo(*AllowedTypeName))
 				{
-					const RewindDebugger::FObjectId& OuterIdentifier = ObjectInfo.GetOuterId();
-					if (!OuterIdentifier.IsSet())
+					AllowedTypeIds.AddUnique(ClassInfo->Id);
+				}
+			}
+
+			auto IsClassOrChildOfAllowed = [&AllowedTypeIds, GameplayProvider](const uint64 ClassId)
+				{
+					if (AllowedTypeIds.Num() > 0)
 					{
-						return;
+						for (const uint64 AllowedTypeId : AllowedTypeIds)
+						{
+							if (GameplayProvider->IsSubClassOf(ClassId, AllowedTypeId))
+							{
+								return true;
+							}
+						}
+						return false;
 					}
-	
-					const FObjectInfo& OuterInfo = GameplayProvider->GetObjectInfo(OuterIdentifier);
-					if (LevelClassInfo == nullptr
-						|| WorldClassInfo == nullptr
-						|| OuterInfo.ClassId == LevelClassInfo->Id
-						|| OuterInfo.ClassId == WorldClassInfo->Id)
-					{
-						const FString Name = ObjectInfo.Name;
-						const FText SelectedLabel = FText::FromString(Name);
-						const FSlateIcon Icon = GameplayProvider->FindIconForClass(ObjectInfo.ClassId);
 
-						MenuBuilder.AddMenuEntry(SelectedLabel, FText(), Icon, FExecuteAction::CreateLambda([this, Name]()
-							{
-								FSlateApplication::Get().DismissAllMenus();
-								DebuggedObjectName.Set(Name);
-							}));
+					// Allow everything if allowed class list is not used
+					return true;
+				};
+
+			GameplayProvider->EnumerateObjects([this, &MenuBuilder, GameplayProvider, &IsClassOrChildOfAllowed](const FObjectInfo& ObjectInfo)
+				{
+					if (!IsClassOrChildOfAllowed(ObjectInfo.ClassId))
+					{
+						return;
 					}
+
+					const FString Name = ObjectInfo.Name;
+					const FText SelectedLabel = FText::FromString(Name);
+					const FSlateIcon Icon = GameplayProvider->FindIconForClass(ObjectInfo.ClassId);
+
+					MenuBuilder.AddMenuEntry(SelectedLabel, FText(), Icon, FExecuteAction::CreateLambda([this, Name]()
+						{
+							FSlateApplication::Get().DismissAllMenus();
+							DebuggedObjectName.Set(Name);
+						}));
 				});
 		}
 		else
diff --git a/Engine/Plugins/Animation/GameplayInsights/Source/RewindDebugger/Public/RewindDebuggerSettings.h b/Engine/Plugins/Animation/GameplayInsights/Source/RewindDebugger/Public/RewindDebuggerSettings.h
--- a/Engine/Plugins/Animation/GameplayInsights/Source/RewindDebugger/Public/RewindDebuggerSettings.h
+++ b/Engine/Plugins/Animation/GameplayInsights/Source/RewindDebugger/Public/RewindDebuggerSettings.h
@@ -2,10 +2,15 @@
 
 #pragma once
 
+#if UE_ENABLE_INCLUDE_ORDER_DEPRECATED_IN_5_8
 #include "CoreMinimal.h"
+#endif // UE_ENABLE_INCLUDE_ORDER_DEPRECATED_IN_5_8
+
 #include "Engine/DeveloperSettings.h"
 #include "RewindDebuggerSettings.generated.h"
 
+struct FInstancedStruct;
+struct FSoftClassPath;
 class URewindDebuggerExtensionSettings;
 
 #define UE_API REWINDDEBUGGER_API
@@ -24,15 +29,48 @@ enum class ERewindDebuggerCameraMode : uint8
 };
 
 /**
- * Implements the settings for the Rewind Debugger.
+ * Implements project settings for the Rewind Debugger.
+ */
+UCLASS(config = RewindDebugger, defaultconfig, EditInlineNew, DisplayName = "Rewind Debugger", collapseCategories, MinimalAPI)
+class URewindDebuggerProjectSettings : public UDeveloperSettings
+{
+	GENERATED_BODY()
+
+public:
+	URewindDebuggerProjectSettings();
+
+#if WITH_EDITOR
+	virtual FText GetSectionDescription() const override;
+#endif
+
+	virtual FName GetCategoryName() const override;
+
+	/**
+	 * Class types allowed in the target selection dropdown for the project.
+	 * Empty list means that all types are allowed (if SelectorAllowedStructTypes is also empty).
+	 */
+	UPROPERTY(Config, EditAnywhere, Category = "Target selection")
+	TArray<FSoftClassPath> SelectorAllowedTypes;
+
+	/**
+	 * Struct types allowed in the target selection dropdown for the project.
+	 * Empty list means that all types are allowed (if SelectorAllowedTypes is also empty).
+	 */
+	UPROPERTY(Config, EditAnywhere, Category = "Target selection")
+	TArray<FInstancedStruct> SelectorAllowedStructTypes;
+};
+
+/**
+ * Implements user settings for the Rewind Debugger.
  */
-UCLASS(config = EditorPerProjectUserSettings, meta=(DisplayName="Rewind Debugger"))
+UCLASS(config = EditorPerProjectUserSettings, meta=(DisplayName="Rewind Debugger (user preferences)"))
 class URewindDebuggerSettings : public UDeveloperSettings
 {
 	GENERATED_BODY()
-	
+
 public:
 	URewindDebuggerSettings();
+	~URewindDebuggerSettings();
 
 #if WITH_EDITOR
 	virtual FText GetSectionText() const override;
@@ -55,15 +93,15 @@ public:
 	/** Rewind Debugger Playback Camera Mode */
 	UPROPERTY(EditAnywhere, Config, Category = Camera)
 	ERewindDebuggerCameraMode CameraMode;
-	
+
 	/** If enabled, automatically detach player control when PIE is paused */
 	UPROPERTY(EditAnywhere, Config, Category = Other)
 	bool bShouldAutoEject;
-	
+
 	/** If enabled, start recording information at the start of PIE */
 	UPROPERTY(EditAnywhere, Config, Category = Other)
 	bool bShouldAutoRecordOnPIE;
-	
+
 	/** Playback speed multiplier */
 	UPROPERTY(EditAnywhere, Config, Category = Other)
 	float PlaybackRate = 1.0;
@@ -71,18 +109,31 @@ public:
 	/** If enabled, show empty tracks on Rewind Debugger Timeline*/
 	UPROPERTY(EditAnywhere, Config, Category = Filters)
 	bool bShowEmptyObjectTracks;
-	
+
 	/** The track types listed here will be hidden from the track tree view */
 	UPROPERTY(EditAnywhere, Config, Category = Filters)
 	TArray<FName> HiddenTrackTypes;
 
-	
 	/** Currently selected target actor's name */
 	UPROPERTY(Config)
 	FString DebugTargetActor;
-	
+
+	/**
+	 * Additional class types allowed in the target selection dropdown
+	 * @note These can only add more allowed types to the types allowed by the project settings, not replace/override them
+	 */
+	UPROPERTY(Config, EditAnywhere, Category = "Target selection")
+	TArray<FSoftClassPath> SelectorAllowedTypes;
+
+	/**
+	 * Additional struct types allowed in the target selection dropdown
+	 * @note These can only add more allowed types to the types allowed by the project settings, not replace/override them
+	 */
+	UPROPERTY(Config, EditAnywhere, Category = "Target selection")
+	TArray<FInstancedStruct> SelectorAllowedStructTypes;
+
 	/** Get Mutable CDO of URewindDebuggerSettings */
-	UE_API static URewindDebuggerSettings & Get();
+	UE_API static URewindDebuggerSettings& Get();
 
 private:
 	/** Default Data Transport for Game Servers - Used to change the transport mode for the trace data. This is not saved in the config by design */
@@ -92,6 +143,8 @@ private:
 	/** Default Data Transport for Game Clients - Used to change the transport mode for the trace data. This is not saved in the config by design */
 	UPROPERTY(EditAnywhere, Category = Connection)
 	UE::TraceBasedDebuggers::ETraceTransportMode GameClientDataTransportMode;
+
+	FDelegateHandle OnPreExitHandle;
 };
 
 #undef UE_API
\ No newline at end of file
diff --git a/Engine/Plugins/Animation/GameplayInsights/Source/RewindDebuggerVLog/Private/RewindDebuggerVLogSettings.cpp b/Engine/Plugins/Animation/GameplayInsights/Source/RewindDebuggerVLog/Private/RewindDebuggerVLogSettings.cpp
--- a/Engine/Plugins/Animation/GameplayInsights/Source/RewindDebuggerVLog/Private/RewindDebuggerVLogSettings.cpp
+++ b/Engine/Plugins/Animation/GameplayInsights/Source/RewindDebuggerVLog/Private/RewindDebuggerVLogSettings.cpp
@@ -17,12 +17,12 @@ URewindDebuggerVLogSettings::URewindDebuggerVLogSettings()
 
 FText URewindDebuggerVLogSettings::GetSectionText() const
 {
-	return NSLOCTEXT("RewindDebugger", "RewindDebuggerVLogSettingsName", "Rewind Debugger");
+	return NSLOCTEXT("RewindDebuggerVLogSettings", "RewindDebuggerVLogSettingsName", "Rewind Debugger (Visual Log)");
 }
 
 FText URewindDebuggerVLogSettings::GetSectionDescription() const
 {
-	return NSLOCTEXT("RewindDebugger", "RewindDebuggerVLogSettingsDescription", "Configure options for the Rewind Debugger.");
+	return NSLOCTEXT("RewindDebuggerVLogSettings", "RewindDebuggerVLogSettingsDescription", "Configure options for the visual log tracks in the Rewind Debugger.");
 }
 
 #endif
