[RewindDebuggerVLog] Add a no-connection state to the log category verbosity selector widget

Introduce a third switcher slot that displays a placeholder message when no remote session is established, replacing the previous early-return behavior

Register an OnDataUpdated callback in SetUseVerbosityLevelFiltering to update UI state reactively when the remote response arrives

#Synced-CL 51720594

diff --git a/Engine/Plugins/Animation/GameplayInsights/Source/RewindDebuggerVLog/Private/Widgets/LogCategoryVerbositySelector.cpp b/Engine/Plugins/Animation/GameplayInsights/Source/RewindDebuggerVLog/Private/Widgets/LogCategoryVerbositySelector.cpp
--- a/Engine/Plugins/Animation/GameplayInsights/Source/RewindDebuggerVLog/Private/Widgets/LogCategoryVerbositySelector.cpp
+++ b/Engine/Plugins/Animation/GameplayInsights/Source/RewindDebuggerVLog/Private/Widgets/LogCategoryVerbositySelector.cpp
@@ -158,6 +158,7 @@ TSharedRef<SWidget> SLogCategoryFilter::OnGetMenuContent()
 {
 	// Create switcher widget first since RefreshCategoryList will set the active slot
 	SAssignNew(Switcher, SWidgetSwitcher)
+		// Slot 0: waiting state (throbber while fetching verbosity levels from the remote process)
 		+ SWidgetSwitcher::Slot()
 		[
 			SNew(SHorizontalBox)
@@ -169,6 +170,7 @@ TSharedRef<SWidget> SLogCategoryFilter::OnGetMenuContent()
 				SNew(SCircularThrobber)
 			]
 		]
+		// Slot 1: ready state (verbosity filter checkbox + category list)
 		+ SWidgetSwitcher::Slot()
 		[
 			SNew(SVerticalBox)
@@ -212,6 +214,20 @@ TSharedRef<SWidget> SLogCategoryFilter::OnGetMenuContent()
 						return bUsingVerbosityFilterWhenRecording;
 					})
 			]
+		]
+		// Slot 2: no-connection state (shown when no remote session is established)
+		+ SWidgetSwitcher::Slot()
+		[
+			SNew(SBox)
+			.HAlign(HAlign_Center)
+			.VAlign(VAlign_Center)
+			.Padding(FMargin(10.0f))
+			[
+				SNew(STextBlock)
+				.Text(LOCTEXT("NoRemoteSessionConnected", "No remote session connected.\nConnect to a remote session to manage verbosity filters."))
+				.AutoWrapText(true)
+				.ColorAndOpacity(FSlateColor::UseSubduedForeground())
+			]
 		];
 
 	RefreshCategoryList();
@@ -356,25 +372,22 @@ void SLogCategoryFilter::SyncVerbosityLevels()
 		}
 	};
 
-	// Nothing to sync if the list of known categories is empty
-	if (AllCategories.IsEmpty())
-	{
-		return;
-	}
-
 	const IRewindDebugger* RewindDebugger = IRewindDebugger::Instance();
-	if (RewindDebugger == nullptr)
-	{
-		return;
-	}
 
 	using namespace UE::TraceBasedDebuggers;
 	const TSharedPtr<FRemoteSessionsManager> RemoteSessionsManager = FRewindDebuggerEngineEditorBridge::Get().GetSessionsManager();
-	const TSharedPtr<FSessionInfo> SessionInfo = RewindDebugger->GetSelectedDebugSessionInfo();
+	const TSharedPtr<FSessionInfo> SessionInfo = RewindDebugger ? RewindDebugger->GetSelectedDebugSessionInfo() : nullptr;
 	if (!SessionInfo.IsValid()
 		|| !RemoteSessionsManager.IsValid())
 	{
-		// Can't send request without a valid sessions managers and an active session
+		// No active remote session — show the no-connection placeholder instead of the controls
+		SwitcherStateIndex = SwitcherSlotIndexForNoConnectionState;
+		return;
+	}
+
+	// Nothing to sync if the list of known categories is empty
+	if (AllCategories.IsEmpty())
+	{
 		return;
 	}
 
@@ -422,11 +435,10 @@ void SLogCategoryFilter::SyncVerbosityLevels()
 
 					ensureMsgf(false, TEXT("The response message should contain the verbosity levels for all our current known categories."));
 				}
-				DebuggerData->bPendingRefresh = false;
 				bUsingVerbosityFilterWhenRecording = DebuggerData->bUsingVerbosityFilterWhenRecording;
 			}
 
-			// We can now display the categories
+			// We can now display the refreshed categories
 			if (Switcher)
 			{
 				Switcher->SetActiveWidgetIndex(SwitcherSlotIndexForReadyState);
@@ -438,7 +450,7 @@ void SLogCategoryFilter::SyncVerbosityLevels()
 #endif // WITH_TRACE_BASED_DEBUGGERS
 }
 
-void SLogCategoryFilter::SetUseVerbosityLevelFiltering(const bool bEnable) const
+void SLogCategoryFilter::SetUseVerbosityLevelFiltering(const bool bEnable)
 {
 #if WITH_TRACE_BASED_DEBUGGERS
 	const IRewindDebugger* RewindDebugger = IRewindDebugger::Instance();
@@ -472,6 +484,22 @@ void SLogCategoryFilter::SetUseVerbosityLevelFiltering(const bool bEnable) const
 		Switcher->SetActiveWidgetIndex(SwitcherSlotIndexForWaitingState);
 	}
 
+	// Register callback when the response is received
+	DebuggerData->OnDataUpdated.BindSPLambda(this, [this]()
+		{
+			// Update our local flag from the updated session data
+			if (FVLogExtensionSessionData* DebuggerData = GetVLogExtensionData())
+			{
+				bUsingVerbosityFilterWhenRecording = DebuggerData->bUsingVerbosityFilterWhenRecording;
+			}
+
+			// We can now display the refreshed state
+			if (Switcher)
+			{
+				Switcher->SetActiveWidgetIndex(SwitcherSlotIndexForReadyState);
+			}
+		});
+
 	// Send command to enable/disable the filtering
 	RemoteSessionsManager->SendCommand(SessionInfo->Address, FVerbosityFilteringStateChangeCommandMessage{ .bEnableFiltering = bEnable });
 #endif // WITH_TRACE_BASED_DEBUGGERS
diff --git a/Engine/Plugins/Animation/GameplayInsights/Source/RewindDebuggerVLog/Private/Widgets/LogCategoryVerbositySelector.h b/Engine/Plugins/Animation/GameplayInsights/Source/RewindDebuggerVLog/Private/Widgets/LogCategoryVerbositySelector.h
--- a/Engine/Plugins/Animation/GameplayInsights/Source/RewindDebuggerVLog/Private/Widgets/LogCategoryVerbositySelector.h
+++ b/Engine/Plugins/Animation/GameplayInsights/Source/RewindDebuggerVLog/Private/Widgets/LogCategoryVerbositySelector.h
@@ -32,7 +32,7 @@ private:
 
 	void RefreshCategoryList();
 
-	void SetUseVerbosityLevelFiltering(bool bEnable) const;
+	void SetUseVerbosityLevelFiltering(bool bEnable);
 
 	bool CanChangeCategoryVerbosity(FName InCategoryName) const;
 	void SetCategoryVerbosity(FName InCategoryName, ELogVerbosity::Type InVerbosity) const;
@@ -50,6 +50,7 @@ private:
 	bool bUsingVerbosityFilterWhenRecording = false;
 	static constexpr int32 SwitcherSlotIndexForWaitingState = 0;
 	static constexpr int32 SwitcherSlotIndexForReadyState = 1;
+	static constexpr int32 SwitcherSlotIndexForNoConnectionState = 2;
 };
 
 } // UE::RewindDebugger
