better sort order for rewind debugger blend weights tracks

#Synced-CL 51932823

diff --git a/Engine/Plugins/Animation/GameplayInsights/Source/GameplayInsights/Private/BlendWeightsTrack.cpp b/Engine/Plugins/Animation/GameplayInsights/Source/GameplayInsights/Private/BlendWeightsTrack.cpp
--- a/Engine/Plugins/Animation/GameplayInsights/Source/GameplayInsights/Private/BlendWeightsTrack.cpp
+++ b/Engine/Plugins/Animation/GameplayInsights/Source/GameplayInsights/Private/BlendWeightsTrack.cpp
@@ -87,7 +87,14 @@ bool FBlendWeightsTrack::UpdateInternal()
 			});
 		});
 
-		UniqueTrackIds.StableSort([](const TrackId &A, const TrackId &B)  { return A.NodeId > B.NodeId; });
+		UniqueTrackIds.StableSort([](const TrackId &A, const TrackId &B)
+			{
+				if (A.NodeId == B.NodeId)
+				{
+					return A.AssetId > B.AssetId;
+				}
+				return A.NodeId > B.NodeId;
+			});
 		
 		AnimationCount = UniqueTrackIds.Num();
 
