Add Object Track sort order setting for rewind debugger - new default sorts based on spawn time, but old default (Alphabetical) can be restored with a setting

#Synced-CL 51932823

#Synced-CL 51932823

diff --git a/Engine/Plugins/Animation/GameplayInsights/Source/RewindDebugger/Private/RewindDebuggerModule.cpp b/Engine/Plugins/Animation/GameplayInsights/Source/RewindDebugger/Private/RewindDebuggerModule.cpp
--- a/Engine/Plugins/Animation/GameplayInsights/Source/RewindDebugger/Private/RewindDebuggerModule.cpp
+++ b/Engine/Plugins/Animation/GameplayInsights/Source/RewindDebugger/Private/RewindDebuggerModule.cpp
@@ -227,7 +227,12 @@ TSharedRef<SDockTab> FRewindDebuggerModule::SpawnRewindDebuggerTab(const FSpawnT
 								.OnTrackSelectionChanged_Raw(DebuggerInstance, &FRewindDebugger::TrackSelectionChanged)
 								.BuildTrackContextMenu_Raw(DebuggerInstance, &FRewindDebugger::BuildTrackContextMenu)
 								.TrackTypes_Lambda([]() { return FRewindDebugger::Instance()->GetTrackTypes(); })
-								.ScrubTime_Lambda([]() { return FRewindDebugger::Instance()->GetScrubTime(); });
+								.ScrubTime_Lambda([]() { return FRewindDebugger::Instance()->GetScrubTime(); })
+								.ForceRefreshTracks_Lambda([]
+									{
+										FRewindDebugger::Instance()->ClearTrackData();
+										FRewindDebugger::Instance()->RefreshDebugTracks();
+									});
 
 	DebuggerInstance->SetTrackCursorDelegate(FRewindDebugger::FOnTrackCursor::CreateSP(RewindDebuggerWidget.Get(), &SRewindDebugger::TrackCursor));
 	DebuggerInstance->SetTrackListChangedDelegate(FRewindDebugger::FOnTrackListChanged::CreateSP(RewindDebuggerWidget.Get(), &SRewindDebugger::RefreshTracks));
diff --git a/Engine/Plugins/Animation/GameplayInsights/Source/RewindDebugger/Private/RewindDebuggerObjectTrack.cpp b/Engine/Plugins/Animation/GameplayInsights/Source/RewindDebugger/Private/RewindDebuggerObjectTrack.cpp
--- a/Engine/Plugins/Animation/GameplayInsights/Source/RewindDebugger/Private/RewindDebuggerObjectTrack.cpp
+++ b/Engine/Plugins/Animation/GameplayInsights/Source/RewindDebugger/Private/RewindDebuggerObjectTrack.cpp
@@ -13,6 +13,7 @@
 #include "Styling/SlateIconFinder.h"
 #include "LevelSequence.h"
 #include "LevelSequenceActor.h"
+#include "RewindDebuggerSettings.h"
 
 #define LOCTEXT_NAMESPACE "RewindDebuggerObjectTrack"
 
@@ -496,12 +497,24 @@ bool FRewindDebuggerObjectTrack::UpdateInternal()
 
 	if (bChanged)
 	{
-		// sort child object tracks by name
-		// note that stable track ordering requires non-empty display names
-		Children.Sort([](const TSharedPtr<FRewindDebuggerTrack>& A, const TSharedPtr<FRewindDebuggerTrack>& B)
-			{
-				return A->GetDisplayName().ToString() < B->GetDisplayName().ToString();
-			});
+		const URewindDebuggerSettings* UserSettings = GetDefault<URewindDebuggerSettings>();
+		if (UserSettings->ObjectTrackSortMode == ERewindDebuggerObjectSortMode::Name)
+		{
+			// sort child object tracks by name
+			// note that stable track ordering requires non-empty display names
+			Children.Sort([](const TSharedPtr<FRewindDebuggerTrack>& A, const TSharedPtr<FRewindDebuggerTrack>& B)
+				{
+					return A->GetDisplayName().ToString() < B->GetDisplayName().ToString();
+				});
+		}
+		else
+		{
+			// sort child object tracks by Id which will generally make older objects first because ids are allocated as increasing integers
+			Children.Sort([](const TSharedPtr<FRewindDebuggerTrack>& A, const TSharedPtr<FRewindDebuggerTrack>& B)
+				{
+					return A->GetUObjectId() < B->GetUObjectId();
+				});
+		}
 	}
 
 	{
diff --git a/Engine/Plugins/Animation/GameplayInsights/Source/RewindDebugger/Private/SRewindDebugger.cpp b/Engine/Plugins/Animation/GameplayInsights/Source/RewindDebugger/Private/SRewindDebugger.cpp
--- a/Engine/Plugins/Animation/GameplayInsights/Source/RewindDebugger/Private/SRewindDebugger.cpp
+++ b/Engine/Plugins/Animation/GameplayInsights/Source/RewindDebugger/Private/SRewindDebugger.cpp
@@ -133,6 +133,28 @@ bool SRewindDebugger::ShouldDisplayEmptyTracks() const
 	return Settings.bShowEmptyObjectTracks;
 }
 
+void SRewindDebugger::SortObjectsByName()
+{
+	Settings.ObjectTrackSortMode = ERewindDebuggerObjectSortMode::Name;
+	ForceRefreshTracks.ExecuteIfBound();
+}
+
+bool SRewindDebugger::ShouldSortObjectsByName() const
+{
+	return Settings.ObjectTrackSortMode == ERewindDebuggerObjectSortMode::Name;
+}
+
+void SRewindDebugger::SortObjectsByTime()
+{
+	Settings.ObjectTrackSortMode = ERewindDebuggerObjectSortMode::Time;
+	ForceRefreshTracks.ExecuteIfBound();
+}
+
+bool SRewindDebugger::ShouldSortObjectsByTime() const
+{
+	return Settings.ObjectTrackSortMode == ERewindDebuggerObjectSortMode::Time;
+}
+
 FReply SRewindDebugger::OnPreviewKeyDown(const FGeometry& MyGeometry, const FKeyEvent& InKeyEvent)
 {
 	const FInputChord KeyEventAsInputChord = FInputChord(InKeyEvent.GetKey(), EModifierKey::FromBools(InKeyEvent.IsControlDown(), InKeyEvent.IsAltDown(), InKeyEvent.IsShiftDown(), InKeyEvent.IsCommandDown()));
@@ -287,6 +309,7 @@ TSharedRef<SWidget> SRewindDebugger::MakeObjectSelectionMenu()
 
 void SRewindDebugger::Construct(const FArguments& InArgs, TSharedRef<FUICommandList> InCommandList, const TSharedRef<SDockTab>& ConstructUnderMajorTab, const TSharedPtr<SWindow>& ConstructUnderWindow)
 {
+	ForceRefreshTracks = InArgs._ForceRefreshTracks;
 	OnScrubPositionChanged = InArgs._OnScrubPositionChanged;
 	OnViewRangeChanged = InArgs._OnViewRangeChanged;
 	OnTrackSelectionChanged = InArgs._OnTrackSelectionChanged;
@@ -681,6 +704,31 @@ TSharedRef<SWidget> SRewindDebugger::MakeFilterMenu()
 		EUserInterfaceActionType::ToggleButton
 	);
 
+	Builder.AddSeparator();
+	
+	Builder.AddMenuEntry(
+		LOCTEXT("Sort Objects By Name", "Sort Objects By Name"),
+		LOCTEXT("Sort Objects By Name Tooltip", "Object Tracks will be sorted alphabetically"),
+		FSlateIcon(),
+		FUIAction(
+			FExecuteAction::CreateRaw(this, &SRewindDebugger::SortObjectsByName),
+			FCanExecuteAction(),
+			FIsActionChecked::CreateRaw(this, &SRewindDebugger::ShouldSortObjectsByName)),
+		NAME_None,
+		EUserInterfaceActionType::ToggleButton
+	);
+		Builder.AddMenuEntry(
+    		LOCTEXT("Sort Objects By Id", "Sort Objects By Time"),
+    		LOCTEXT("Sort Objects By Id Tooltip", "Object Tracks will be sorted by Id, making older objects display first"),
+    		FSlateIcon(),
+    		FUIAction(
+    			FExecuteAction::CreateRaw(this, &SRewindDebugger::SortObjectsByTime),
+    			FCanExecuteAction(),
+    			FIsActionChecked::CreateRaw(this, &SRewindDebugger::ShouldSortObjectsByTime)),
+    		NAME_None,
+    		EUserInterfaceActionType::ToggleButton
+    	);
+
 	return Builder.MakeWidget();
 }
 
diff --git a/Engine/Plugins/Animation/GameplayInsights/Source/RewindDebugger/Private/SRewindDebugger.h b/Engine/Plugins/Animation/GameplayInsights/Source/RewindDebugger/Private/SRewindDebugger.h
--- a/Engine/Plugins/Animation/GameplayInsights/Source/RewindDebugger/Private/SRewindDebugger.h
+++ b/Engine/Plugins/Animation/GameplayInsights/Source/RewindDebugger/Private/SRewindDebugger.h
@@ -28,6 +28,7 @@ public:
 	DECLARE_DELEGATE_OneParam(FOnTrackDoubleClicked, TSharedPtr<RewindDebugger::FRewindDebuggerTrack>)
 	DECLARE_DELEGATE_OneParam(FOnTrackSelectionChanged, TSharedPtr<RewindDebugger::FRewindDebuggerTrack>)
 	DECLARE_DELEGATE_RetVal(TSharedPtr<SWidget>, FBuildTrackContextMenu)
+	DECLARE_DELEGATE(FForceRefreshTracks)
 
 	SLATE_BEGIN_ARGS(SRewindDebugger) {}
 		SLATE_ARGUMENT(TArray< TSharedPtr< RewindDebugger::FRewindDebuggerTrack > >*, Tracks);
@@ -42,6 +43,7 @@ public:
 		SLATE_EVENT(FBuildTrackContextMenu, BuildTrackContextMenu);
 		SLATE_EVENT(FOnTrackDoubleClicked, OnTrackDoubleClicked);
 		SLATE_EVENT(FOnTrackSelectionChanged, OnTrackSelectionChanged);
+		SLATE_EVENT(FForceRefreshTracks, ForceRefreshTracks);
 	SLATE_END_ARGS()
 
 	/**
@@ -66,6 +68,11 @@ private:
 	bool ShouldHideTrackType(const FName& TrackType) const;
 	void ToggleDisplayEmptyTracks();
 	bool ShouldDisplayEmptyTracks() const;
+	
+	void SortObjectsByName();
+	bool ShouldSortObjectsByName() const;
+	void SortObjectsByTime();
+	bool ShouldSortObjectsByTime() const;
 
 	virtual FReply OnPreviewKeyDown(const FGeometry& MyGeometry, const FKeyEvent& InKeyEvent) override;
 	virtual bool SupportsKeyboardFocus() const override { return true; }
@@ -81,6 +88,7 @@ private:
 	TBindableProperty<double> TraceTime;
 	TBindableProperty<double> RecordingDuration;
 
+	FForceRefreshTracks ForceRefreshTracks;
 	TSharedPtr<FUICommandList> CommandList;
 	const FRewindDebuggerCommands& Commands;
 
@@ -95,6 +103,7 @@ private:
 	TSharedRef<SWidget> MakeMainMenu();
 	TSharedRef<SWidget> MakeFilterMenu();
 
+
 	/** Track tree view */
 	TArray<TSharedPtr<RewindDebugger::FRewindDebuggerTrack>>* Tracks = nullptr;
 	TSharedPtr<RewindDebugger::FRewindDebuggerTrack> SelectedComponent;
diff --git a/Engine/Plugins/Animation/GameplayInsights/Source/RewindDebugger/Public/RewindDebuggerSettings.h b/Engine/Plugins/Animation/GameplayInsights/Source/RewindDebugger/Public/RewindDebuggerSettings.h
--- a/Engine/Plugins/Animation/GameplayInsights/Source/RewindDebugger/Public/RewindDebuggerSettings.h
+++ b/Engine/Plugins/Animation/GameplayInsights/Source/RewindDebugger/Public/RewindDebuggerSettings.h
@@ -28,6 +28,13 @@ enum class ERewindDebuggerCameraMode : uint8
 	Disabled UMETA(Tooltip="Disable Camera On Playback"),
 };
 
+UENUM()
+enum class ERewindDebuggerObjectSortMode : uint8
+{
+	Time UMETA(Tooltip="Sort Object tracks By Id, making older objects display first"),
+	Name UMETA(Tooltip="Sort Object Tracks alphabetically by Name"),
+};
+
 /**
  * Implements project settings for the Rewind Debugger.
  */
@@ -94,6 +101,10 @@ public:
 	UPROPERTY(EditAnywhere, Config, Category = Camera)
 	ERewindDebuggerCameraMode CameraMode;
 
+	/** Rewind Debugger Object Track Sort Mode */
+	UPROPERTY(EditAnywhere, Config, Category = Other)
+	ERewindDebuggerObjectSortMode ObjectTrackSortMode = ERewindDebuggerObjectSortMode::Time;
+
 	/** If enabled, automatically detach player control when PIE is paused */
 	UPROPERTY(EditAnywhere, Config, Category = Other)
 	bool bShouldAutoEject;
