Fix rewind debugger details tab being empty if when closed and reopened

#Synced-CL 51932823

diff --git a/Engine/Plugins/Animation/GameplayInsights/Source/RewindDebugger/Private/RewindDebugger.cpp b/Engine/Plugins/Animation/GameplayInsights/Source/RewindDebugger/Private/RewindDebugger.cpp
--- a/Engine/Plugins/Animation/GameplayInsights/Source/RewindDebugger/Private/RewindDebugger.cpp
+++ b/Engine/Plugins/Animation/GameplayInsights/Source/RewindDebugger/Private/RewindDebugger.cpp
@@ -1470,6 +1470,10 @@ void FRewindDebugger::UpdateDetailsPanel(TSharedRef<SDockTab> DetailsTab)
 		if (!DetailsPanelWidget.IsValid())
 		{
 			DetailsPanelWidget = SNew(SBox);
+		}
+
+		if (DetailsTab->GetContent() != DetailsPanelWidget)
+		{
 			DetailsTab->SetContent(DetailsPanelWidget.ToSharedRef());
 		}
 		
diff --git a/Engine/Plugins/Animation/GameplayInsights/Source/RewindDebugger/Private/RewindDebuggerModule.h b/Engine/Plugins/Animation/GameplayInsights/Source/RewindDebugger/Private/RewindDebuggerModule.h
--- a/Engine/Plugins/Animation/GameplayInsights/Source/RewindDebugger/Private/RewindDebuggerModule.h
+++ b/Engine/Plugins/Animation/GameplayInsights/Source/RewindDebugger/Private/RewindDebuggerModule.h
@@ -51,7 +51,6 @@ private:
 
 private:
 	TSharedPtr<SRewindDebugger> RewindDebuggerWidget;
-	TSharedPtr<SRewindDebuggerDetails> RewindDebuggerDetailsWidget;
 
 	FRewindDebuggerCamera RewindDebuggerCameraExtension;
 	FRewindDebuggerAnimation RewindDebuggerAnimationExtension;
