[RewindDebugger]
Implement per-provider locking for RewindDebugger trace providers and migrate all call sites from session-level locks to provider-level locks to prevent deadlocks.

- Add FProviderLock implementation to FAnimationProvider, FVisualLoggerProvider, FPoseSearch::FTraceProvider, FUAFProvider, FAnimNextAnimGraphProvider, FUAFAnimNodeProvider, FChooserProvider, and FStateTreeTraceProvider
- Replace FAnalysisSessionReadScope with FProviderReadScopeLock at all RewindDebugger track call sites (~60 files)
- Replace FAnalysisSessionEditScope with FProviderEditScopeLock in all trace analyzers
- Replace Session.ReadAccessCheck()/WriteAccessCheck() with provider-level ReadAccessCheck()/EditAccessCheck()
- Add missing ReadAccessCheck/EditAccessCheck to provider methods that were unprotected
- Optimize lock scopes to minimize hold time (cache data inside tight locks, move widget creation/asset loading outside)
- Wrap Session.UpdateDurationSeconds() calls in FAnalysisSessionEditScope where needed

#Synced-CL 52251601

diff --git a/Engine/Plugins/Animation/GameplayInsights/Source/GameplayInsights/Private/AnimCurvesTrack.cpp b/Engine/Plugins/Animation/GameplayInsights/Source/GameplayInsights/Private/AnimCurvesTrack.cpp
--- a/Engine/Plugins/Animation/GameplayInsights/Source/GameplayInsights/Private/AnimCurvesTrack.cpp
+++ b/Engine/Plugins/Animation/GameplayInsights/Source/GameplayInsights/Private/AnimCurvesTrack.cpp
@@ -1,6 +1,7 @@
 // Copyright Epic Games, Inc. All Rights Reserved.
 
 #include "AnimCurvesTrack.h"
+#include "Common/ProviderLock.h"
 #include "SCurveTimelineView.h"
 #include "IRewindDebugger.h"
 #include "GameplayProvider.h"
@@ -52,6 +53,11 @@ bool FAnimCurvesTrack::UpdateInternal()
 
 	IRewindDebugger* RewindDebugger = IRewindDebugger::Instance();
 	const TraceServices::IAnalysisSession* AnalysisSession = RewindDebugger->GetAnalysisSession();
+	const FAnimationProvider* AnimationProvider = AnalysisSession->ReadProvider<FAnimationProvider>(FAnimationProvider::ProviderName);
+	if (AnimationProvider == nullptr)
+	{
+		return false;
+	}
 
 	bool bChanged = false;
 	
@@ -61,12 +67,8 @@ bool FAnimCurvesTrack::UpdateInternal()
 	double EndTime = TraceTimeRange.GetUpperBoundValue();
 	
 	// count number of unique animations in the current time range
-	const FGameplayProvider* GameplayProvider = AnalysisSession->ReadProvider<FGameplayProvider>(FGameplayProvider::ProviderName);
-	const FAnimationProvider* AnimationProvider = AnalysisSession->ReadProvider<FAnimationProvider>(FAnimationProvider::ProviderName);
-
-	if(GameplayProvider && AnimationProvider)
 	{
-		TraceServices::FAnalysisSessionReadScope SessionReadScope(*AnalysisSession);
+		TraceServices::FProviderReadScopeLock AnimationProviderReadScope(*AnimationProvider);
 		UniqueTrackIds.SetNum(0, EAllowShrinking::No);
 
 		AnimationProvider->ReadSkeletalMeshPoseTimeline(ObjectId, [&UniqueTrackIds,AnimationProvider, StartTime, EndTime](const FAnimationProvider::SkeletalMeshPoseTimeline& InTimeline, bool bHasCurves)
@@ -84,27 +86,27 @@ bool FAnimCurvesTrack::UpdateInternal()
 				return TraceServices::EEventEnumerate::Continue;
 			});
 		});
+	}
 
-		UniqueTrackIds.StableSort();
-		
-		const int32 CurveCount = UniqueTrackIds.Num();
+	UniqueTrackIds.StableSort();
 
-		if (Children.Num() != UniqueTrackIds.Num())
-			bChanged = true;
-		
-		Children.SetNum(CurveCount);
-		for(int i = 0; i < CurveCount; i++)
+	const int32 CurveCount = UniqueTrackIds.Num();
+
+	if (Children.Num() != UniqueTrackIds.Num())
+		bChanged = true;
+	
+	Children.SetNum(CurveCount);
+	for (int i = 0; i < CurveCount; i++)
+	{
+		if (!Children[i].IsValid() || Children[i].Get()->GetCurveId() != UniqueTrackIds[i])
 		{
-			if (!Children[i].IsValid() || Children[i].Get()->GetCurveId() != UniqueTrackIds[i])
-			{
-				Children[i] = MakeShared<FAnimCurveTrack>(ObjectId, UniqueTrackIds[i]);
-				bChanged = true;
-			}
+			Children[i] = MakeShared<FAnimCurveTrack>(ObjectId, UniqueTrackIds[i]);
+			bChanged = true;
+		}
 
-			if (Children[i]->Update())
-			{
-				bChanged = true;
-			}
+		if (Children[i]->Update())
+		{
+			bChanged = true;
 		}
 	}
 
@@ -137,7 +139,7 @@ bool FAnimCurveTrack::UpdateInternal()
 	const TraceServices::IAnalysisSession* AnalysisSession = RewindDebugger->GetAnalysisSession();
 	const FAnimationProvider* AnimationProvider = AnalysisSession->ReadProvider<FAnimationProvider>(FAnimationProvider::ProviderName);
 
-	if(CurvesUpdateRequested > 10 && AnimationProvider)
+	if (CurvesUpdateRequested > 10 && AnimationProvider)
 	{
 		CurvesUpdateRequested = 0;
 		UpdateCurvePointsInternal();
@@ -146,6 +148,7 @@ bool FAnimCurveTrack::UpdateInternal()
 	bool bChanged = false;
 	if (CurveName.IsEmpty() && AnimationProvider)
 	{
+		TraceServices::FProviderReadScopeLock AnimationProviderReadScope(*AnimationProvider);
 		CurveName = FText::FromString(AnimationProvider->GetName(CurveId));
 		bChanged = true;
 	}
@@ -156,21 +159,24 @@ bool FAnimCurveTrack::UpdateInternal()
 void FAnimCurveTrack::UpdateCurvePointsInternal()
 {
 	TRACE_CPUPROFILER_EVENT_SCOPE(FAnimCurveTrack::UpdateCurvePointsInternal);
-	
-	IRewindDebugger* RewindDebugger = IRewindDebugger::Instance();
+
+	auto& CurvePoints = CurveData->Points;
+	CurvePoints.SetNum(0, EAllowShrinking::No);
+
+	const IRewindDebugger* RewindDebugger = IRewindDebugger::Instance();
 	const TraceServices::IAnalysisSession* AnalysisSession = RewindDebugger->GetAnalysisSession();
 	const FAnimationProvider* AnimationProvider = AnalysisSession->ReadProvider<FAnimationProvider>(FAnimationProvider::ProviderName);
+	if (AnimationProvider == nullptr)
+	{
+		return;
+	}
 
 	// convert time range to from rewind debugger times to profiler times
-	TRange<double> TraceTimeRange = RewindDebugger->GetCurrentTraceRange();
+	const TRange<double> TraceTimeRange = RewindDebugger->GetCurrentTraceRange();
 	double StartTime = TraceTimeRange.GetLowerBoundValue();
 	double EndTime = TraceTimeRange.GetUpperBoundValue();
-	
-	auto& CurvePoints = CurveData->Points;
-	CurvePoints.SetNum(0,EAllowShrinking::No);
 
-	TraceServices::FAnalysisSessionReadScope SessionReadScope(*AnalysisSession);
-	
+	TraceServices::FProviderReadScopeLock ProviderReadScope(*AnimationProvider);
 	AnimationProvider->ReadSkeletalMeshPoseTimeline(ObjectId, [this, AnimationProvider, StartTime, EndTime, AnalysisSession, &CurvePoints](const FAnimationProvider::SkeletalMeshPoseTimeline& InTimeline, bool bHasCurves)
 	{
 		// this isn't very efficient, and it gets called every frame.  will need optimizing
@@ -179,7 +185,6 @@ void FAnimCurveTrack::UpdateCurvePointsInternal()
 			if (InEndTime > StartTime && InStartTime < EndTime)
 			{
 				double Time = InMessage.RecordingTime;
-			
 				AnimationProvider->EnumerateSkeletalMeshCurves(InMessage, [this, Time, &CurvePoints](const FSkeletalMeshNamedCurve& InCurve)
 				{
 					if (InCurve.Id == CurveId)
@@ -240,10 +245,10 @@ bool FAnimationCurvesTrackCreator::HasDebugInfoInternal(const FObjectId& InObjec
 {
 	const TraceServices::IAnalysisSession* AnalysisSession = IRewindDebugger::Instance()->GetAnalysisSession();
 	
-	TraceServices::FAnalysisSessionReadScope SessionReadScope(*AnalysisSession);
 	bool bHasData = false;
 	if (const FAnimationProvider* AnimationProvider = AnalysisSession->ReadProvider<FAnimationProvider>(FAnimationProvider::ProviderName))
 	{
+		TraceServices::FProviderReadScopeLock ProviderReadScope(*AnimationProvider);
 		AnimationProvider->ReadSkeletalMeshPoseTimeline(InObjectId.GetMainId(), [&bHasData](const FAnimationProvider::SkeletalMeshPoseTimeline& InTimeline, bool bInHasCurves)
 		{
 			bHasData = bInHasCurves;
diff --git a/Engine/Plugins/Animation/GameplayInsights/Source/GameplayInsights/Private/AnimNodesTrack.cpp b/Engine/Plugins/Animation/GameplayInsights/Source/GameplayInsights/Private/AnimNodesTrack.cpp
--- a/Engine/Plugins/Animation/GameplayInsights/Source/GameplayInsights/Private/AnimNodesTrack.cpp
+++ b/Engine/Plugins/Animation/GameplayInsights/Source/GameplayInsights/Private/AnimNodesTrack.cpp
@@ -4,6 +4,7 @@
 
 #include "AnimationProvider.h"
 #include "AnimationSharedData.h"
+#include "Common/ProviderLock.h"
 #include "Framework/MultiBox/MultiBoxBuilder.h"
 #include "GameplayProvider.h"
 #include "GameplayTimingViewExtender.h"
@@ -115,7 +116,7 @@ void FAnimNodesTrack::BuildDrawState(ITimingEventsTrackDrawStateBuilder& Builder
 
 	if(AnimationProvider)
 	{
-		TraceServices::FAnalysisSessionReadScope SessionReadScope(SharedData.GetAnalysisSession());
+		TraceServices::FProviderReadScopeLock ProviderReadScope(*AnimationProvider);
 
 		AnimationProvider->ReadAnimGraphTimeline(GetGameplayTrack().GetObjectId(), [&Context, &Builder](const FAnimationProvider::AnimGraphTimeline& InTimeline)
 		{
@@ -175,7 +176,7 @@ void FAnimNodesTrack::FindAnimGraphMessage(const FTimingEventSearchParameters& I
 
 			if(AnimationProvider)
 			{
-				TraceServices::FAnalysisSessionReadScope SessionReadScope(SharedData.GetAnalysisSession());
+				TraceServices::FProviderReadScopeLock ProviderReadScope(*AnimationProvider);
 
 				AnimationProvider->ReadAnimGraphTimeline(GetGameplayTrack().GetObjectId(), [&InContext](const FAnimationProvider::AnimGraphTimeline& InTimeline)
 				{
@@ -213,14 +214,21 @@ void FAnimNodesTrack::BuildContextMenu(FMenuBuilder& MenuBuilder)
 							GEditor->GetEditorSubsystem<UAssetEditorSubsystem>()->OpenEditorForAsset(AnimBlueprint);
 
 							const FGameplayProvider* GameplayProvider = SharedData.GetAnalysisSession().ReadProvider<FGameplayProvider>(FGameplayProvider::ProviderName);
-							if(GameplayProvider)
+							if (GameplayProvider)
 							{
-								TraceServices::FAnalysisSessionReadScope SessionReadScope(SharedData.GetAnalysisSession());
+								uint64 OuterUObjectId = 0;
+								{
+									TraceServices::FProviderReadScopeLock ProviderReadScope(*GameplayProvider);
+									if(const FObjectInfo* AnimInstanceObjectInfo = GameplayProvider->FindObjectInfo(GetGameplayTrack().GetObjectId()))
+									{
+										OuterUObjectId = AnimInstanceObjectInfo->GetOuterUObjectId();
+									}
+								}
 
-								if(const FObjectInfo* AnimInstanceObjectInfo = GameplayProvider->FindObjectInfo(GetGameplayTrack().GetObjectId()))
+								if (OuterUObjectId != 0)
 								{
 									// find a counterpart skeletal mesh pose track
-									TSharedPtr<FSkeletalMeshPoseTrack> SkeletalMeshPoseTrack = SharedData.FindSkeletalMeshPoseTrack(AnimInstanceObjectInfo->GetOuterUObjectId());
+									TSharedPtr<FSkeletalMeshPoseTrack> SkeletalMeshPoseTrack = SharedData.FindSkeletalMeshPoseTrack(OuterUObjectId);
 									if(SkeletalMeshPoseTrack.IsValid())
 									{
 										SkeletalMeshPoseTrack->MarkPotentiallyDebugged();
@@ -247,14 +255,21 @@ void FAnimNodesTrack::BuildContextMenu(FMenuBuilder& MenuBuilder)
 						if(UAnimBlueprint* AnimBlueprint = Cast<UAnimBlueprint>(InstanceClass.Get()->ClassGeneratedBy))
 						{
 							const FGameplayProvider* GameplayProvider = SharedData.GetAnalysisSession().ReadProvider<FGameplayProvider>(FGameplayProvider::ProviderName);
-							if(GameplayProvider)
+							if (GameplayProvider)
 							{
-								TraceServices::FAnalysisSessionReadScope SessionReadScope(SharedData.GetAnalysisSession());
+								uint64 OuterUObjectId = 0;
+								{
+									TraceServices::FProviderReadScopeLock ProviderReadScope(*GameplayProvider);
+									if(const FObjectInfo* AnimInstanceObjectInfo = GameplayProvider->FindObjectInfo(GetGameplayTrack().GetAssociatedObjectId()))
+									{
+										OuterUObjectId = AnimInstanceObjectInfo->GetOuterUObjectId();
+									}
+								}
 
-								if(const FObjectInfo* AnimInstanceObjectInfo = GameplayProvider->FindObjectInfo(GetGameplayTrack().GetAssociatedObjectId()))
+								if (OuterUObjectId != 0)
 								{
 									// find a counterpart skeletal mesh pose track
-									TSharedPtr<FSkeletalMeshPoseTrack> SkeletalMeshPoseTrack = SharedData.FindSkeletalMeshPoseTrack(AnimInstanceObjectInfo->GetOuterUObjectId());
+									TSharedPtr<FSkeletalMeshPoseTrack> SkeletalMeshPoseTrack = SharedData.FindSkeletalMeshPoseTrack(OuterUObjectId);
 									if(SkeletalMeshPoseTrack.IsValid())
 									{
 										if(USkeletalMeshComponent* SkeletalMeshComponent = SkeletalMeshPoseTrack->GetComponent(FGameplayTimingViewExtender::GetWorldToVisualize()))
@@ -320,7 +335,7 @@ void FAnimNodesTrack::UpdateDebugData(const TraceServices::FFrame& InFrame)
 
 				if(AnimationProvider && GameplayProvider)
 				{
-					TraceServices::FAnalysisSessionReadScope SessionReadScope(SharedData.GetAnalysisSession());
+					TraceServices::FProviderReadScopeLock AnimationProviderReadScope(*AnimationProvider);
 
 					const int32 NodeCount = InstanceClass.Get()->GetAnimNodeProperties().Num();
 
@@ -356,7 +371,7 @@ void FAnimNodesTrack::UpdateDebugData(const TraceServices::FFrame& InFrame)
 										});
 									});
 
-									AnimationProvider->ReadAnimSequencePlayersTimeline(GetGameplayTrack().GetObjectId(), [InGraphStartTime, InGraphEndTime, GameplayProvider, &DebugData](const FAnimationProvider::AnimSequencePlayersTimeline& InSequencePlayersTimeline)
+									AnimationProvider->ReadAnimSequencePlayersTimeline(GetGameplayTrack().GetObjectId(), [InGraphStartTime, InGraphEndTime, &DebugData](const FAnimationProvider::AnimSequencePlayersTimeline& InSequencePlayersTimeline)
 									{
 										InSequencePlayersTimeline.EnumerateEvents(InGraphStartTime, InGraphEndTime, [&DebugData](double InStartTime, double InEndTime, uint32 InDepth, const FAnimSequencePlayerMessage& InMessage)
 										{
@@ -367,6 +382,7 @@ void FAnimNodesTrack::UpdateDebugData(const TraceServices::FFrame& InFrame)
 
 									AnimationProvider->ReadAnimBlendSpacePlayersTimeline(GetGameplayTrack().GetObjectId(), [InGraphStartTime, InGraphEndTime, GameplayProvider, &DebugData](const FAnimationProvider::BlendSpacePlayersTimeline& InBlendSpacePlayersTimeline)
 									{
+										TraceServices::FProviderReadScopeLock GameplayProviderReadScope(*GameplayProvider);
 										InBlendSpacePlayersTimeline.EnumerateEvents(InGraphStartTime, InGraphEndTime, [GameplayProvider, &DebugData](double InStartTime, double InEndTime, uint32 InDepth, const FBlendSpacePlayerMessage& InMessage)
 										{
 											UBlendSpace* BlendSpace = nullptr;
@@ -441,14 +457,21 @@ void FAnimNodesTrack::GetCustomDebugObjects(const IAnimationBlueprintEditor& InA
 		if(UAnimBlueprint* AnimBlueprint = Cast<UAnimBlueprint>(InstanceClass.Get()->ClassGeneratedBy))
 		{
 			const FGameplayProvider* GameplayProvider = SharedData.GetAnalysisSession().ReadProvider<FGameplayProvider>(FGameplayProvider::ProviderName);
-			if(GameplayProvider)
+			if (GameplayProvider)
 			{
-				TraceServices::FAnalysisSessionReadScope SessionReadScope(SharedData.GetAnalysisSession());
+				uint64 OuterUObjectId = 0;
+				{
+					TraceServices::FProviderReadScopeLock ProviderReadScope(*GameplayProvider);
+					if(const FObjectInfo* AnimInstanceObjectInfo = GameplayProvider->FindObjectInfo(GetGameplayTrack().GetObjectId()))
+					{
+						OuterUObjectId = AnimInstanceObjectInfo->GetOuterUObjectId();
+					}
+				}
 
-				if(const FObjectInfo* AnimInstanceObjectInfo = GameplayProvider->FindObjectInfo(GetGameplayTrack().GetObjectId()))
+				if (OuterUObjectId != 0)
 				{
 					// find a counterpart skeletal mesh pose track
-					TSharedPtr<FSkeletalMeshPoseTrack> SkeletalMeshPoseTrack = SharedData.FindSkeletalMeshPoseTrack(AnimInstanceObjectInfo->GetOuterUObjectId());
+					TSharedPtr<FSkeletalMeshPoseTrack> SkeletalMeshPoseTrack = SharedData.FindSkeletalMeshPoseTrack(OuterUObjectId);
 					if(SkeletalMeshPoseTrack.IsValid())
 					{
 						SkeletalMeshPoseTrack->MarkPotentiallyDebugged();
diff --git a/Engine/Plugins/Animation/GameplayInsights/Source/GameplayInsights/Private/AnimNotifiesTrack.cpp b/Engine/Plugins/Animation/GameplayInsights/Source/GameplayInsights/Private/AnimNotifiesTrack.cpp
--- a/Engine/Plugins/Animation/GameplayInsights/Source/GameplayInsights/Private/AnimNotifiesTrack.cpp
+++ b/Engine/Plugins/Animation/GameplayInsights/Source/GameplayInsights/Private/AnimNotifiesTrack.cpp
@@ -4,6 +4,7 @@
 
 #include "AnimationProvider.h"
 #include "AnimationSharedData.h"
+#include "Common/ProviderLock.h"
 #include "Framework/MultiBox/MultiBoxBuilder.h"
 #include "GameplayProvider.h"
 #include "TraceServices/Model/Frames.h"
@@ -35,10 +36,9 @@ void FAnimNotifiesTrack::BuildDrawState(ITimingEventsTrackDrawStateBuilder& Buil
 
 	if(AnimationProvider)
 	{
-		TraceServices::FAnalysisSessionReadScope SessionReadScope(SharedData.GetAnalysisSession());
-
 		bHasNotifies = false;
 
+		TraceServices::FProviderReadScopeLock ProviderReadScope(*AnimationProvider);
 		AnimationProvider->ReadNotifyTimeline(GetGameplayTrack().GetObjectId(), [this, &Context, &Builder](const FAnimationProvider::AnimNotifyTimeline& InTimeline)
 		{
 			InTimeline.EnumerateEvents(Context.GetViewport().GetStartTime(), Context.GetViewport().GetEndTime(), [this, &Builder](double InStartTime, double InEndTime, uint32 InDepth, const FAnimNotifyMessage& InMessage)
@@ -143,11 +143,10 @@ void FAnimNotifiesTrack::InitTooltip(FTooltipDrawState& Tooltip, const ITimingEv
 		Tooltip.AddNameValueTextLine(LOCTEXT("NotifyType", "Type").ToString(), GetMessageType(InMessage.NotifyEventType, InFoundDepth == 0));
 
 		{
-			TraceServices::FAnalysisSessionReadScope SessionReadScope(SharedData.GetAnalysisSession());
-
 			const FAnimationProvider* AnimationProvider = SharedData.GetAnalysisSession().ReadProvider<FAnimationProvider>(FAnimationProvider::ProviderName);
 			if(AnimationProvider)
 			{
+				TraceServices::FProviderReadScopeLock ProviderReadScope(*AnimationProvider);
 				const TCHAR* Name = AnimationProvider->GetName(InMessage.NameId);
 				Tooltip.AddNameValueTextLine(LOCTEXT("NotifyName", "Name").ToString(), Name);
 			}
@@ -158,11 +157,10 @@ void FAnimNotifiesTrack::InitTooltip(FTooltipDrawState& Tooltip, const ITimingEv
 
 		if(InMessage.NotifyEventType != EAnimNotifyMessageType::SyncMarker)
 		{
-			TraceServices::FAnalysisSessionReadScope SessionReadScope(SharedData.GetAnalysisSession());
-
 			const FGameplayProvider* GameplayProvider = SharedData.GetAnalysisSession().ReadProvider<FGameplayProvider>(FGameplayProvider::ProviderName);
 			if(GameplayProvider)
 			{
+				TraceServices::FProviderReadScopeLock ProviderReadScope(*GameplayProvider);
 				const FObjectInfo& AssetInfo = GameplayProvider->GetObjectInfo(InMessage.AssetId);
 				Tooltip.AddNameValueTextLine(LOCTEXT("Asset", "Asset").ToString(), AssetInfo.PathName);
 
@@ -202,10 +200,9 @@ void FAnimNotifiesTrack::FindAnimNotifyMessage(const FTimingEventSearchParameter
 
 			if(AnimationProvider)
 			{
-				TraceServices::FAnalysisSessionReadScope SessionReadScope(SharedData.GetAnalysisSession());
-
 				if(bHasNotifies)
 				{
+					TraceServices::FProviderReadScopeLock ProviderReadScope(*AnimationProvider);
 					AnimationProvider->ReadNotifyTimeline(GetGameplayTrack().GetObjectId(), [this, &InContext](const FAnimationProvider::AnimNotifyTimeline& InTimeline)
 					{
 						InTimeline.EnumerateEvents(InContext.GetParameters().StartTime, InContext.GetParameters().EndTime, [this, &InContext](double InEventStartTime, double InEventEndTime, uint32 InDepth, const FAnimNotifyMessage& InMessage)
@@ -218,6 +215,7 @@ void FAnimNotifiesTrack::FindAnimNotifyMessage(const FTimingEventSearchParameter
 
 				if(IdToDepthMap.Num() > 0)
 				{
+					TraceServices::FProviderReadScopeLock ProviderReadScope(*AnimationProvider);
 					AnimationProvider->EnumerateNotifyStateTimelines(GetGameplayTrack().GetObjectId(), [this, &InContext](uint32 InNotifyNameId, const FAnimationProvider::AnimNotifyTimeline& InTimeline)
 					{
 						InTimeline.EnumerateEvents(InContext.GetParameters().StartTime, InContext.GetParameters().EndTime, [this, &InContext, InNotifyNameId](double InEventStartTime, double InEventEndTime, uint32 InDepth, const FAnimNotifyMessage& InMessage)
@@ -249,7 +247,7 @@ void FAnimNotifiesTrack::GetVariantsAtFrame(const TraceServices::FFrame& InFrame
 
 	if(AnimationProvider && GameplayProvider)
 	{
-		TraceServices::FAnalysisSessionReadScope SessionReadScope(SharedData.GetAnalysisSession());
+		TraceServices::FProviderReadScopeLock AnimationProviderReadScope(*AnimationProvider);
 
 		bool bForEvent = true;
 		auto ProcessEvent = [this, &AnimationProvider, &GameplayProvider, &Header, &bForEvent](double InStartTime, double InEndTime, uint32 InDepth, const FAnimNotifyMessage& InMessage)
@@ -261,13 +259,18 @@ void FAnimNotifiesTrack::GetVariantsAtFrame(const TraceServices::FFrame& InFrame
 			NotifyHeader->AddChild(FVariantTreeNode::MakeFloat(LOCTEXT("EventDuration", "Duration"), InEndTime - InStartTime));
 			NotifyHeader->AddChild(FVariantTreeNode::MakeFloat(LOCTEXT("EventTime", "Time"), InStartTime));
 
-			if(InMessage.NotifyEventType != EAnimNotifyMessageType::SyncMarker)
+			if (InMessage.NotifyEventType != EAnimNotifyMessageType::SyncMarker)
 			{
 				NotifyHeader->AddChild(FVariantTreeNode::MakeObject(LOCTEXT("Asset", "Asset"), InMessage.AssetId));
 
-				const FObjectInfo& NotifyInfo = GameplayProvider->GetObjectInfo(InMessage.NotifyId);
-				bool bIsState = InMessage.NotifyEventType == EAnimNotifyMessageType::Begin || InMessage.NotifyEventType == EAnimNotifyMessageType::End || InMessage.NotifyEventType == EAnimNotifyMessageType::Tick;
-				NotifyHeader->AddChild(FVariantTreeNode::MakeClass((bIsState ? LOCTEXT("Notify State Class", "Notify State Class") : LOCTEXT("NotifyClass", "Notify Class")), NotifyInfo.ClassId));
+				uint64 NotifyInfoClassId;
+				{
+					TraceServices::FProviderReadScopeLock GameplayProviderReadScope(*GameplayProvider);
+					const FObjectInfo& NotifyInfo = GameplayProvider->GetObjectInfo(InMessage.NotifyId);
+					NotifyInfoClassId = NotifyInfo.ClassId;
+				}
+				const bool bIsState = InMessage.NotifyEventType == EAnimNotifyMessageType::Begin || InMessage.NotifyEventType == EAnimNotifyMessageType::End || InMessage.NotifyEventType == EAnimNotifyMessageType::Tick;
+				NotifyHeader->AddChild(FVariantTreeNode::MakeClass((bIsState ? LOCTEXT("Notify State Class", "Notify State Class") : LOCTEXT("NotifyClass", "Notify Class")), NotifyInfoClassId));
 			}
 
 			return TraceServices::EEventEnumerate::Continue;
diff --git a/Engine/Plugins/Animation/GameplayInsights/Source/GameplayInsights/Private/AnimationAnalyzer.cpp b/Engine/Plugins/Animation/GameplayInsights/Source/GameplayInsights/Private/AnimationAnalyzer.cpp
--- a/Engine/Plugins/Animation/GameplayInsights/Source/GameplayInsights/Private/AnimationAnalyzer.cpp
+++ b/Engine/Plugins/Animation/GameplayInsights/Source/GameplayInsights/Private/AnimationAnalyzer.cpp
@@ -63,7 +63,7 @@ bool FAnimationAnalyzer::OnEvent(uint16 RouteId, EStyle Style, const FOnEventCon
 
 	LLM_SCOPE_BYNAME(TEXT("Insights/FAnimationAnalyzer"));
 
-	FAnalysisSessionEditScope _(Session);
+	FProviderEditScopeLock ProviderEditScope(AnimationProvider);
 
 	const auto& EventData = Context.EventData;
 	switch (RouteId)
diff --git a/Engine/Plugins/Animation/GameplayInsights/Source/GameplayInsights/Private/AnimationProvider.cpp b/Engine/Plugins/Animation/GameplayInsights/Source/GameplayInsights/Private/AnimationProvider.cpp
--- a/Engine/Plugins/Animation/GameplayInsights/Source/GameplayInsights/Private/AnimationProvider.cpp
+++ b/Engine/Plugins/Animation/GameplayInsights/Source/GameplayInsights/Private/AnimationProvider.cpp
@@ -5,6 +5,11 @@
 #include "Insights/ViewModels/TimingEventsTrack.h"
 #include "AnimationAnalyzer.h"
 
+namespace TraceServices
+{
+thread_local FProviderLock::FThreadLocalState GAnimationProviderLockState;
+}
+
 #define LOCTEXT_NAMESPACE "AnimationProvider"
 
 FName FAnimationProvider::ProviderName("AnimationProvider");
@@ -22,9 +27,15 @@ FAnimationProvider::FAnimationProvider(TraceServices::IAnalysisSession& InSessio
 	GameplayProvider.OnObjectEndPlay().AddRaw(this, &FAnimationProvider::HandleObjectEndPlay);
 }
 
+void FAnimationProvider::UpdateSessionDurationSeconds(const double InTime) const
+{
+	TraceServices::FAnalysisSessionEditScope SessionEditScope(Session);
+	Session.UpdateDurationSeconds(InTime);
+}
+
 void FAnimationProvider::EnumerateSkeletalMeshPoseTimelines(TFunctionRef<void(uint64 ObjectId, const SkeletalMeshPoseTimeline&)> Callback) const
 {
-	Session.ReadAccessCheck();
+	ReadAccessCheck();
 	
 	for(auto& IndexMapping : ObjectIdToSkeletalMeshPoseTimelines)
 	{
@@ -38,7 +49,7 @@ void FAnimationProvider::EnumerateSkeletalMeshPoseTimelines(TFunctionRef<void(ui
 
 bool FAnimationProvider::ReadSkeletalMeshPoseTimeline(uint64 InObjectId, TFunctionRef<void(const SkeletalMeshPoseTimeline&, bool)> Callback) const
 {
-	Session.ReadAccessCheck();
+	ReadAccessCheck();
 
 	const uint32* IndexPtr = ObjectIdToSkeletalMeshPoseTimelines.Find(InObjectId);
 	if(IndexPtr != nullptr)
@@ -56,7 +67,7 @@ bool FAnimationProvider::ReadSkeletalMeshPoseTimeline(uint64 InObjectId, TFuncti
 
 void FAnimationProvider::GetSkeletalMeshComponentSpacePose(const FSkeletalMeshPoseMessage& InMessage, const FSkeletalMeshInfo& InMeshInfo, FTransform& OutComponentToWorld, TArray<FTransform>& OutTransforms) const
 {
-	Session.ReadAccessCheck();
+	ReadAccessCheck();
 
 	if(InMeshInfo.BoneCount == InMessage.NumTransforms)
 	{
@@ -78,7 +89,7 @@ void FAnimationProvider::GetSkeletalMeshComponentSpacePose(const FSkeletalMeshPo
 
 void FAnimationProvider::GetPoseWatchData(const FPoseWatchMessage& InMessage, TArray<FTransform>& BoneTransforms, TArray<uint16>& RequiredBones) const
 {
-	Session.ReadAccessCheck();
+	ReadAccessCheck();
 
 	// Pre-alloc arrays
 	BoneTransforms.Empty();
@@ -112,7 +123,7 @@ void FAnimationProvider::GetPoseWatchData(const FPoseWatchMessage& InMessage, TA
 
 void FAnimationProvider::EnumerateSkeletalMeshCurveIds(uint64 InObjectId, TFunctionRef<void(uint32)> Callback) const
 {
-	Session.ReadAccessCheck();
+	ReadAccessCheck();
 
 	const uint32* IndexPtr = ObjectIdToSkeletalMeshPoseTimelines.Find(InObjectId);
 	if(IndexPtr != nullptr)
@@ -130,7 +141,7 @@ void FAnimationProvider::EnumerateSkeletalMeshCurveIds(uint64 InObjectId, TFunct
 
 void FAnimationProvider::EnumerateSkeletalMeshCurves(const FSkeletalMeshPoseMessage& InMessage, TFunctionRef<void(const FSkeletalMeshNamedCurve&)> Callback) const
 {
-	Session.ReadAccessCheck();
+	ReadAccessCheck();
 
 	uint64 StartCurveIndex = InMessage.CurveStartIndex;
 	uint64 EndCurveIndex = InMessage.CurveStartIndex + InMessage.NumCurves;
@@ -143,7 +154,7 @@ void FAnimationProvider::EnumerateSkeletalMeshCurves(const FSkeletalMeshPoseMess
 
 void FAnimationProvider::EnumeratePoseWatchCurves(const FPoseWatchMessage& InMessage, TFunctionRef<void(const FSkeletalMeshNamedCurve&)> Callback) const
 {
-	Session.ReadAccessCheck();
+	ReadAccessCheck();
 
 	uint64 StartCurveIndex = InMessage.CurveStartIndex;
 	uint64 EndCurveIndex = InMessage.CurveStartIndex + InMessage.NumCurves;
@@ -156,7 +167,7 @@ void FAnimationProvider::EnumeratePoseWatchCurves(const FPoseWatchMessage& InMes
 
 void FAnimationProvider::EnumerateExternalMorphSets(const FSkeletalMeshPoseMessage& InMessage, TFunctionRef<void(const FExternalMorphWeightMessage&)> Callback) const
 {
-	Session.ReadAccessCheck();
+	ReadAccessCheck();
 
 	const uint64 StartIndex = InMessage.ExternalMorphStartIndex;
 	const uint64 EndIndex = StartIndex + InMessage.NumExternalMorphSets;
@@ -168,7 +179,7 @@ void FAnimationProvider::EnumerateExternalMorphSets(const FSkeletalMeshPoseMessa
 
 void FAnimationProvider::EnumerateTickRecordIds(uint64 InObjectId, TFunctionRef<void(uint64, int32)> Callback) const
 {
-	Session.ReadAccessCheck();
+	ReadAccessCheck();
 
 	const uint32* IndexPtr = ObjectIdToTickRecordTimelineStorage.Find(InObjectId);
 	if (IndexPtr != nullptr)
@@ -186,7 +197,7 @@ void FAnimationProvider::EnumerateTickRecordIds(uint64 InObjectId, TFunctionRef<
 
 bool FAnimationProvider::ReadTickRecordTimeline(uint64 InObjectId, TFunctionRef<void(const TickRecordTimeline&)> Callback) const
 {
-	Session.ReadAccessCheck();
+	ReadAccessCheck();
 
 	const uint32* IndexPtr = ObjectIdToTickRecordTimelineStorage.Find(InObjectId);
 	if (IndexPtr != nullptr)
@@ -203,7 +214,7 @@ bool FAnimationProvider::ReadTickRecordTimeline(uint64 InObjectId, TFunctionRef<
 
 bool FAnimationProvider::ReadInertializationTimeline(uint64 InObjectId, TFunctionRef<void(const InertializationTimeline&)> Callback) const
 {
-	Session.ReadAccessCheck();
+	ReadAccessCheck();
 
 	const uint32* IndexPtr = ObjectIdToInertializationTimelines.Find(InObjectId);
 	if (IndexPtr != nullptr)
@@ -220,7 +231,7 @@ bool FAnimationProvider::ReadInertializationTimeline(uint64 InObjectId, TFunctio
 
 void FAnimationProvider::EnumerateInertializationNodes(uint64 InObjectId, TFunctionRef<void(int32 NodeId, EInertializationType Type)> Callback) const
 {
-	Session.ReadAccessCheck();
+	ReadAccessCheck();
 
 	const uint32* IndexPtr = ObjectIdToInertializationTimelines.Find(InObjectId);
 	if (IndexPtr != nullptr)
@@ -237,7 +248,7 @@ void FAnimationProvider::EnumerateInertializationNodes(uint64 InObjectId, TFunct
 
 void FAnimationProvider::EnumerateAnimGraphTimelines(TFunctionRef<void(uint64 ObjectId, const AnimGraphTimeline&)> Callback) const
 {
-	Session.ReadAccessCheck();
+	ReadAccessCheck();
 	
 	for(auto& IndexMapping : ObjectIdToAnimGraphTimelines)
 	{
@@ -247,7 +258,7 @@ void FAnimationProvider::EnumerateAnimGraphTimelines(TFunctionRef<void(uint64 Ob
 
 bool FAnimationProvider::ReadAnimGraphTimeline(uint64 InObjectId, TFunctionRef<void(const AnimGraphTimeline&)> Callback) const
 {
-	Session.ReadAccessCheck();
+	ReadAccessCheck();
 
 	const uint32* IndexPtr = ObjectIdToAnimGraphTimelines.Find(InObjectId);
 	if(IndexPtr != nullptr)
@@ -264,7 +275,7 @@ bool FAnimationProvider::ReadAnimGraphTimeline(uint64 InObjectId, TFunctionRef<v
 
 bool FAnimationProvider::ReadAnimNodesTimeline(uint64 InObjectId, TFunctionRef<void(const AnimNodesTimeline&)> Callback) const
 {
-	Session.ReadAccessCheck();
+	ReadAccessCheck();
 
 	const uint32* IndexPtr = ObjectIdToAnimNodeTimelines.Find(InObjectId);
 	if(IndexPtr != nullptr)
@@ -281,7 +292,7 @@ bool FAnimationProvider::ReadAnimNodesTimeline(uint64 InObjectId, TFunctionRef<v
 
 bool FAnimationProvider::ReadAnimNodeValuesTimeline(uint64 InObjectId, TFunctionRef<void(const AnimNodeValuesTimeline&)> Callback) const
 {
-	Session.ReadAccessCheck();
+	ReadAccessCheck();
 
 	const uint32* IndexPtr = ObjectIdToAnimNodeValueTimelines.Find(InObjectId);
 	if(IndexPtr != nullptr)
@@ -298,7 +309,7 @@ bool FAnimationProvider::ReadAnimNodeValuesTimeline(uint64 InObjectId, TFunction
 
 bool FAnimationProvider::ReadAnimAttributesTimeline(uint64 InObjectId, TFunctionRef<void(const AnimAttributeTimeline&)> Callback) const
 {
-	Session.ReadAccessCheck();
+	ReadAccessCheck();
 
 	const uint32* IndexPtr = ObjectIdToAnimAttributeTimelines.Find(InObjectId);
 	if(IndexPtr != nullptr)
@@ -315,7 +326,7 @@ bool FAnimationProvider::ReadAnimAttributesTimeline(uint64 InObjectId, TFunction
 
 bool FAnimationProvider::ReadAnimSequencePlayersTimeline(uint64 InObjectId, TFunctionRef<void(const AnimSequencePlayersTimeline&)> Callback) const
 {
-	Session.ReadAccessCheck();
+	ReadAccessCheck();
 
 	const uint32* IndexPtr = ObjectIdToAnimSequencePlayerTimelines.Find(InObjectId);
 	if(IndexPtr != nullptr)
@@ -332,7 +343,7 @@ bool FAnimationProvider::ReadAnimSequencePlayersTimeline(uint64 InObjectId, TFun
 
 bool FAnimationProvider::ReadAnimBlendSpacePlayersTimeline(uint64 InObjectId, TFunctionRef<void(const BlendSpacePlayersTimeline&)> Callback) const
 {
-	Session.ReadAccessCheck();
+	ReadAccessCheck();
 
 	const uint32* IndexPtr = ObjectIdToBlendSpacePlayerTimelines.Find(InObjectId);
 	if(IndexPtr != nullptr)
@@ -349,7 +360,7 @@ bool FAnimationProvider::ReadAnimBlendSpacePlayersTimeline(uint64 InObjectId, TF
 
 bool FAnimationProvider::ReadStateMachinesTimeline(uint64 InObjectId, TFunctionRef<void(const StateMachinesTimeline&)> Callback) const
 {
-	Session.ReadAccessCheck();
+	ReadAccessCheck();
 
 	const uint32* IndexPtr = ObjectIdToStateMachineTimelines.Find(InObjectId);
 	if(IndexPtr != nullptr)
@@ -366,7 +377,7 @@ bool FAnimationProvider::ReadStateMachinesTimeline(uint64 InObjectId, TFunctionR
 
 bool FAnimationProvider::ReadNotifyTimeline(uint64 InObjectId, TFunctionRef<void(const AnimNotifyTimeline&)> Callback) const
 {
-	Session.ReadAccessCheck();
+	ReadAccessCheck();
 
 	const uint32* IndexPtr = ObjectIdToAnimNotifyTimelines.Find(InObjectId);
 	if(IndexPtr != nullptr)
@@ -383,6 +394,8 @@ bool FAnimationProvider::ReadNotifyTimeline(uint64 InObjectId, TFunctionRef<void
 
 void FAnimationProvider::EnumerateNotifyStateTimelines(uint64 InObjectId, TFunctionRef<void(uint64, const AnimNotifyTimeline&)> Callback) const
 {
+	ReadAccessCheck();
+
 	const uint32* IndexPtr = ObjectIdToAnimNotifyStateTimelines.Find(InObjectId);
 	if(IndexPtr != nullptr)
 	{
@@ -399,7 +412,7 @@ void FAnimationProvider::EnumerateNotifyStateTimelines(uint64 InObjectId, TFunct
 
 bool FAnimationProvider::ReadMontageTimeline(uint64 InObjectId, TFunctionRef<void(const AnimMontageTimeline&)> Callback) const
 {
-	Session.ReadAccessCheck();
+	ReadAccessCheck();
 
 	const uint32* IndexPtr = ObjectIdToAnimMontageTimelines.Find(InObjectId);
 	if(IndexPtr != nullptr)
@@ -416,7 +429,7 @@ bool FAnimationProvider::ReadMontageTimeline(uint64 InObjectId, TFunctionRef<voi
 
 void FAnimationProvider::EnumerateMontageIds(uint64 InObjectId, TFunctionRef<void(uint64)> Callback) const
 {
-	Session.ReadAccessCheck();
+	ReadAccessCheck();
 
 	const uint32* IndexPtr = ObjectIdToAnimMontageTimelines.Find(InObjectId);
 	if(IndexPtr != nullptr)
@@ -434,7 +447,7 @@ void FAnimationProvider::EnumerateMontageIds(uint64 InObjectId, TFunctionRef<voi
 
 bool FAnimationProvider::ReadAnimSyncTimeline(uint64 InObjectId, TFunctionRef<void(const AnimSyncTimeline&)> Callback) const
 {
-	Session.ReadAccessCheck();
+	ReadAccessCheck();
 
 	const uint32* IndexPtr = ObjectIdToAnimSyncTimelines.Find(InObjectId);
 	if(IndexPtr != nullptr)
@@ -451,7 +464,7 @@ bool FAnimationProvider::ReadAnimSyncTimeline(uint64 InObjectId, TFunctionRef<vo
 
 bool FAnimationProvider::ReadPoseWatchTimeline(uint64 InObjectId, TFunctionRef<void(const PoseWatchTimeline&)> Callback) const
 {
-	Session.ReadAccessCheck();
+	ReadAccessCheck();
 
 	const uint32* IndexPtr = ObjectIdToPoseWatchTimelines.Find(InObjectId);
 	if (IndexPtr != nullptr)
@@ -468,7 +481,7 @@ bool FAnimationProvider::ReadPoseWatchTimeline(uint64 InObjectId, TFunctionRef<v
 
 const FSkeletalMeshInfo* FAnimationProvider::FindSkeletalMeshInfo(uint64 InObjectId) const
 {
-	Session.ReadAccessCheck();
+	ReadAccessCheck();
 
 	const uint32* IndexPtr = SkeletalMeshIdToIndexMap.Find(InObjectId);
 	if(IndexPtr != nullptr)
@@ -484,7 +497,7 @@ const FSkeletalMeshInfo* FAnimationProvider::FindSkeletalMeshInfo(uint64 InObjec
 
 const FAnimNodeInfo* FAnimationProvider::FindAnimNodeInfo(int32 InAnimNodeId, uint64 InAnimInstanceId) const
 {
-	Session.ReadAccessCheck();
+	ReadAccessCheck();
 
 	const uint32* IndexPtr = AnimNodeIdToIndexMap.Find({ InAnimNodeId, InAnimInstanceId });
 	if (IndexPtr != nullptr)
@@ -500,6 +513,8 @@ const FAnimNodeInfo* FAnimationProvider::FindAnimNodeInfo(int32 InAnimNodeId, ui
 
 const TCHAR* FAnimationProvider::GetName(uint32 InId) const
 {
+	ReadAccessCheck();
+
 	const TCHAR* const* FoundName = NameMap.Find(InId);
 	if(FoundName != nullptr)
 	{
@@ -512,11 +527,15 @@ const TCHAR* FAnimationProvider::GetName(uint32 InId) const
 
 FText FAnimationProvider::FormatNodeKeyValue(const FAnimNodeValueMessage& InMessage) const
 {
+	ReadAccessCheck();
+
 	return FText::Format(LOCTEXT("KeyValueFormat", "{0} = {1}"), FText::FromString(InMessage.Key), FormatNodeValue(InMessage));
 }
 
 FText FAnimationProvider::FormatNodeValue(const FAnimNodeValueMessage& InMessage) const
 {
+	ReadAccessCheck();
+
 	FText Text;
 
 	switch(InMessage.Value.Type)
@@ -541,12 +560,14 @@ FText FAnimationProvider::FormatNodeValue(const FAnimNodeValueMessage& InMessage
 		break;
 	case EAnimNodeValueType::Object:
 	{
+		TraceServices::FProviderReadScopeLock ProviderReadScope(GameplayProvider);
 		const FObjectInfo& ObjectInfo = GameplayProvider.GetObjectInfo(InMessage.Value.Object.Value);
 		Text = FText::FromString(ObjectInfo.Name);
 		break;
 	}
 	case EAnimNodeValueType::Class:
 	{
+		TraceServices::FProviderReadScopeLock ProviderReadScope(GameplayProvider);
 		const FClassInfo& ClassInfo = GameplayProvider.GetClassInfo(InMessage.Value.Class.Value);
 		Text = FText::FromString(ClassInfo.PathName);
 		break;
@@ -565,7 +586,7 @@ FText FAnimationProvider::FormatNodeValue(const FAnimNodeValueMessage& InMessage
 bool FAnimationProvider::HasExternalMorphSets(uint64 InObjectId) const
 {
 	//ObjectIdToSkeletalMeshPoseTimelines
-	Session.ReadAccessCheck();
+	ReadAccessCheck();
 
 	const uint32* IndexPtr = ObjectIdToSkeletalMeshPoseTimelines.Find(InObjectId);
 	if(IndexPtr != nullptr)
@@ -582,14 +603,14 @@ bool FAnimationProvider::HasExternalMorphSets(uint64 InObjectId) const
 
 bool FAnimationProvider::HasAnyData() const
 {
-	Session.ReadAccessCheck();
+	ReadAccessCheck();
 
 	return bHasAnyData;
 }
 
 void FAnimationProvider::AppendTickRecord(uint64 InAnimInstanceId, double InProfileTime, double InRecordingTime, uint64 InAssetId, int32 InNodeId, float InBlendWeight, float InPlaybackTime, float InRootMotionWeight, float InPlayRate, float InBlendSpacePositionX, float InBlendSpacePositionY, float InBlendSpaceFilteredPositionX, float InBlendSpaceFilteredPositionY, uint16 InFrameCounter, bool bInLooping, bool bInIsBlendSpace)
 {
-	Session.WriteAccessCheck();
+	EditAccessCheck();
 
 	bHasAnyData = true;
 
@@ -630,12 +651,12 @@ void FAnimationProvider::AppendTickRecord(uint64 InAnimInstanceId, double InProf
 
 	TimelineStorage->Timeline->AppendEvent(InProfileTime, Message);
 
-	Session.UpdateDurationSeconds(InProfileTime);
+	UpdateSessionDurationSeconds(InProfileTime);
 }
 
 void FAnimationProvider::AppendTickRecord(uint64 InAnimInstanceId, double InProfileTime, double InRecordingTime, uint64 InAssetId, int32 InNodeId, float InBlendWeight, float InPlaybackTime, float InRootMotionWeight, float InPlayRate, float InBlendSpacePositionX, float InBlendSpacePositionY, float InBlendSpaceFilteredPositionX, float InBlendSpaceFilteredPositionY, uint16 InFrameCounter, bool bInLooping, bool bInIsBlendSpace, uint32 InSyncGroupName, bool bInIsSyncGroupLeader, float InSyncLeaderScore, float TimeToPrevMarker, float TimeToNextMarker, int32 PrevMarkerIndex, int32 NextMarkerIndex)
 {
-	Session.WriteAccessCheck();
+	EditAccessCheck();
 
 	bHasAnyData = true;
 
@@ -683,12 +704,12 @@ void FAnimationProvider::AppendTickRecord(uint64 InAnimInstanceId, double InProf
 	
 	TimelineStorage->Timeline->AppendEvent(InProfileTime, Message);
 
-	Session.UpdateDurationSeconds(InProfileTime);
+	UpdateSessionDurationSeconds(InProfileTime);
 }
 
 void FAnimationProvider::AppendSkeletalMesh(uint64 InObjectId, uint64 SkeletonId, const TArrayView<const int32>& InParentIndices)
 {
-	Session.WriteAccessCheck();
+	EditAccessCheck();
 
 	if(SkeletalMeshIdToIndexMap.Find(InObjectId) == nullptr)
 	{
@@ -712,7 +733,7 @@ void FAnimationProvider::AppendSkeletalMesh(uint64 InObjectId, uint64 SkeletonId
 
 void FAnimationProvider::AppendSkeletalMeshComponent(uint64 InObjectId, uint64 InMeshId, double InProfileTime, double InRecordingTime, uint16 InLodIndex, uint16 InFrameCounter, const TArrayView<const FTransform>& InPose, const TArrayView<const FSkeletalMeshNamedCurve>& InCurves)
 {
-	Session.WriteAccessCheck();
+	EditAccessCheck();
 
 	bHasAnyData = true;
 
@@ -745,7 +766,10 @@ void FAnimationProvider::AppendSkeletalMeshComponent(uint64 InObjectId, uint64 I
 	Message.CurveStartIndex = SkeletalMeshCurves.Num();
 	Message.ComponentId = InObjectId;
 	Message.MeshId = InMeshId;
-	Message.MeshName = GameplayProvider.GetObjectInfo(Message.MeshId).Name;
+	{
+		TraceServices::FProviderReadScopeLock ProviderReadScope(GameplayProvider);
+		Message.MeshName = GameplayProvider.GetObjectInfo(Message.MeshId).Name;
+	}
 	Message.NumTransforms = (uint16)InPose.Num() - 1;
 	Message.NumCurves = (uint16)InCurves.Num();
 	Message.LodIndex = InLodIndex;
@@ -764,7 +788,7 @@ void FAnimationProvider::AppendSkeletalMeshComponent(uint64 InObjectId, uint64 I
 		TimelineStorage->AllCurveIds.Add(Curve.Id);
 	}
 
-	Session.UpdateDurationSeconds(InProfileTime);
+	UpdateSessionDurationSeconds(InProfileTime);
 }
 
 // support for deserializing LWC/non LWC Transforms from builds with non-matching LWC setting
@@ -825,7 +849,7 @@ void FAnimationProvider::AppendSkeletalMeshComponent(
 	const TArrayView<const float>& InCurveValues)
 
 {
-	Session.WriteAccessCheck();
+	EditAccessCheck();
 
 	bHasAnyData = true;
 
@@ -879,7 +903,10 @@ void FAnimationProvider::AppendSkeletalMeshComponent(
 	Message.CurveStartIndex = SkeletalMeshCurves.Num();
 	Message.ComponentId = InObjectId;
 	Message.MeshId = InMeshId;
-	Message.MeshName = GameplayProvider.GetObjectInfo(Message.MeshId).Name;
+	{
+		TraceServices::FProviderReadScopeLock ProviderReadScope(GameplayProvider);
+		Message.MeshName = GameplayProvider.GetObjectInfo(Message.MeshId).Name;
+	}
 	Message.NumTransforms = (uint16)PoseTransformCount;
 	Message.NumCurves = (uint16)NumCurves;
 	Message.LodIndex = InLodIndex;
@@ -908,7 +935,7 @@ void FAnimationProvider::AppendSkeletalMeshComponent(
 		TimelineStorage->AllCurveIds.Add(InCurveIds[CurveIndex]);
 	}
 
-	Session.UpdateDurationSeconds(InProfileTime);
+	UpdateSessionDurationSeconds(InProfileTime);
 }
 
 void FAnimationProvider::AppendSkeletalMeshComponent(
@@ -926,7 +953,7 @@ void FAnimationProvider::AppendSkeletalMeshComponent(
 	const TArrayView<const float>& InExternalMorphWeights,
 	const TArrayView<const int32>& InExternalMorphCounts)
 {
-	Session.WriteAccessCheck();
+	EditAccessCheck();
 
 	bHasAnyData = true;
 
@@ -981,7 +1008,10 @@ void FAnimationProvider::AppendSkeletalMeshComponent(
 	Message.CurveStartIndex = SkeletalMeshCurves.Num();
 	Message.ComponentId = InObjectId;
 	Message.MeshId = InMeshId;
-	Message.MeshName = GameplayProvider.GetObjectInfo(Message.MeshId).Name;
+	{
+		TraceServices::FProviderReadScopeLock ProviderReadScope(GameplayProvider);
+		Message.MeshName = GameplayProvider.GetObjectInfo(Message.MeshId).Name;
+	}
 	Message.NumTransforms = (uint16)PoseTransformCount;
 	Message.NumCurves = (uint16)NumCurves;
 	Message.LodIndex = InLodIndex;
@@ -1031,12 +1061,12 @@ void FAnimationProvider::AppendSkeletalMeshComponent(
 
 	TimelineStorage->NumExternalMorphSets = InExternalMorphCounts.Num();
 
-	Session.UpdateDurationSeconds(InProfileTime);
+	UpdateSessionDurationSeconds(InProfileTime);
 }
 	
 void FAnimationProvider::AppendName(uint32 InId, const TCHAR* InName)
 {
-	Session.WriteAccessCheck();
+	EditAccessCheck();
 
 	bHasAnyData = true;
 
@@ -1060,7 +1090,7 @@ void FAnimationProvider::HandleObjectEndPlay(uint64 InObjectId, double InTime, c
 
 void FAnimationProvider::AppendSkeletalMeshFrame(uint64 InObjectId, double InTime, uint16 InFrameCounter)
 {
-	Session.WriteAccessCheck();
+	EditAccessCheck();
 
 	bHasAnyData = true;
 
@@ -1092,12 +1122,12 @@ void FAnimationProvider::AppendSkeletalMeshFrame(uint64 InObjectId, double InTim
 
 	Timeline->AppendBeginEvent(InTime, Message);
 
-	Session.UpdateDurationSeconds(InTime);
+	UpdateSessionDurationSeconds(InTime);
 }
 
 void FAnimationProvider::AppendAnimGraph(uint64 InAnimInstanceId, double InStartTime, double InEndTime, int32 InNodeCount, uint16 InFrameCounter, uint8 InPhase)
 {
-	Session.WriteAccessCheck();
+	EditAccessCheck();
 
 	bHasAnyData = true;
 
@@ -1124,12 +1154,12 @@ void FAnimationProvider::AppendAnimGraph(uint64 InAnimInstanceId, double InStart
 	uint64 EventIndex = Timeline->AppendBeginEvent(InStartTime, Message);
 	Timeline->EndEvent(EventIndex, InEndTime);
 
-	Session.UpdateDurationSeconds(InStartTime);
+	UpdateSessionDurationSeconds(InStartTime);
 }
 
 void FAnimationProvider::AppendAnimNodeStart(uint64 InAnimInstanceId, double InStartTime, uint16 InFrameCounter, int32 InNodeId, int32 PreviousNodeId, const TCHAR* InTargetNodeName, float InWeight, float InRootMotionWeight, const TCHAR* InTargetNodeDisplayName, uint8 InPhase)
 {
-	Session.WriteAccessCheck();
+	EditAccessCheck();
 
 	bHasAnyData = true;
 
@@ -1174,12 +1204,12 @@ void FAnimationProvider::AppendAnimNodeStart(uint64 InAnimInstanceId, double InS
 		AnimNodeIdToIndexMap.Add({ InNodeId, InAnimInstanceId }, NewAnimNodeInfoIndex);
 	}
 
-	Session.UpdateDurationSeconds(InStartTime);
+	UpdateSessionDurationSeconds(InStartTime);
 }
 
 void FAnimationProvider::AppendAnimSequencePlayer(uint64 InAnimInstanceId, double InTime, int32 InNodeId, float InPosition, float InLength, uint16 InFrameCounter)
 {
-	Session.WriteAccessCheck();
+	EditAccessCheck();
 
 	bHasAnyData = true;
 
@@ -1206,12 +1236,12 @@ void FAnimationProvider::AppendAnimSequencePlayer(uint64 InAnimInstanceId, doubl
 
 	Timeline->AppendEvent(InTime, Message);
 
-	Session.UpdateDurationSeconds(InTime);
+	UpdateSessionDurationSeconds(InTime);
 }
 
 void FAnimationProvider::AppendBlendSpacePlayer(uint64 InAnimInstanceId, double InTime, int32 InNodeId, uint64 InBlendSpaceId, const FVector& InBlendPosition, const FVector& InFilteredBlendPosition)
 {
-	Session.WriteAccessCheck();
+	EditAccessCheck();
 
 	bHasAnyData = true;
 
@@ -1242,7 +1272,7 @@ void FAnimationProvider::AppendBlendSpacePlayer(uint64 InAnimInstanceId, double
 
 	Timeline->AppendEvent(InTime, Message);
 
-	Session.UpdateDurationSeconds(InTime);
+	UpdateSessionDurationSeconds(InTime);
 }
 
 void FAnimationProvider::AppendAnimNodeValue(uint64 InAnimInstanceId, double InTime, double InRecordingTime, uint16 InFrameCounter, int32 InNodeId, const TCHAR* InKey, bool bInValue)
@@ -1329,7 +1359,7 @@ void FAnimationProvider::AppendAnimNodeValueAnimNode(uint64 InAnimInstanceId, do
 
 void FAnimationProvider::AppendAnimNodeValue(uint64 InAnimInstanceId, double InTime, double InRecordingTime, uint16 InFrameCounter, int32 InNodeId, const TCHAR* InKey, FAnimNodeValueMessage& InMessage)
 {
-	Session.WriteAccessCheck();
+	EditAccessCheck();
 
 	bHasAnyData = true;
 
@@ -1354,12 +1384,12 @@ void FAnimationProvider::AppendAnimNodeValue(uint64 InAnimInstanceId, double InT
 
 	Timeline->AppendEvent(InTime, InMessage);
 
-	Session.UpdateDurationSeconds(InTime);
+	UpdateSessionDurationSeconds(InTime);
 }
 
 void FAnimationProvider::AppendAnimGraphAttribute(uint64 InSourceAnimInstanceId, uint64 InTargetAnimInstanceId, double InTime, int32 InSourceNodeId, int32 InTargetNodeId, uint32 InAttributeNameId)
 {
-	Session.WriteAccessCheck();
+	EditAccessCheck();
 
 	bHasAnyData = true;
 
@@ -1429,12 +1459,12 @@ void FAnimationProvider::AppendAnimGraphAttribute(uint64 InSourceAnimInstanceId,
 		SourceTimeline->AppendEvent(InTime, SourceMessage);
 	}
 
-	Session.UpdateDurationSeconds(InTime);
+	UpdateSessionDurationSeconds(InTime);
 }
 
 void FAnimationProvider::AppendStateMachineState(uint64 InAnimInstanceId, double InTime, int32 InNodeId, int32 InStateMachineIndex, int32 InStateIndex, float InStateWeight, float InElapsedTime)
 {
-	Session.WriteAccessCheck();
+	EditAccessCheck();
 
 	bHasAnyData = true;
 
@@ -1461,12 +1491,12 @@ void FAnimationProvider::AppendStateMachineState(uint64 InAnimInstanceId, double
 
 	Timeline->AppendEvent(InTime, Message);
 
-	Session.UpdateDurationSeconds(InTime);
+	UpdateSessionDurationSeconds(InTime);
 }
 
 void FAnimationProvider::AppendNotify(uint64 InAnimInstanceId, double InTime, double InRecordingTime, uint64 InAssetId, uint64 InNotifyId, uint32 InNameId, float InNotifyTime, float InNotifyDuration, EAnimNotifyMessageType InNotifyEventType)
 {
-	Session.WriteAccessCheck();
+	EditAccessCheck();
 
 	bHasAnyData = true;
 
@@ -1549,12 +1579,12 @@ void FAnimationProvider::AppendNotify(uint64 InAnimInstanceId, double InTime, do
 		Timeline->AppendEvent(InTime, Message);
 	}
 
-	Session.UpdateDurationSeconds(InTime);
+	UpdateSessionDurationSeconds(InTime);
 }
 
 void FAnimationProvider::AppendMontage(uint64 InAnimInstanceId, double InProfileTime, double InRecoringTime, uint64 InMontageId, uint32 InCurrentSectionNameId, uint32 InNextSectionNameId, float InWeight, float InDesiredWeight, float InPosition, uint16 InFrameCounter)
 {
-	Session.WriteAccessCheck();
+	EditAccessCheck();
 
 	bHasAnyData = true;
 
@@ -1588,12 +1618,12 @@ void FAnimationProvider::AppendMontage(uint64 InAnimInstanceId, double InProfile
 
 	TimelineStorage->Timeline->AppendEvent(InProfileTime, Message);
 
-	Session.UpdateDurationSeconds(InProfileTime);
+	UpdateSessionDurationSeconds(InProfileTime);
 }
 
 void FAnimationProvider::AppendSync(uint64 InAnimInstanceId, double InTime, int32 InSourceNodeId, uint32 InGroupNameId)
 {
-	Session.WriteAccessCheck();
+	EditAccessCheck();
 
 	bHasAnyData = true;
 
@@ -1617,12 +1647,12 @@ void FAnimationProvider::AppendSync(uint64 InAnimInstanceId, double InTime, int3
 
 	Timeline->AppendEvent(InTime, Message);
 
-	Session.UpdateDurationSeconds(InTime);
+	UpdateSessionDurationSeconds(InTime);
 }
 
 void FAnimationProvider::AppendPoseWatch(uint64 InAnimInstanceId, double InTime, double InRecordingTime, uint64 PoseWatchId, const TArrayView<const float>& BoneTransformsRaw, const TArrayView<const uint16>& RequiredBones, const TArrayView<const float>& WorldTransformRaw, const bool bIsEnabled)
 {
-	Session.WriteAccessCheck();
+	EditAccessCheck();
 
 	bHasAnyData = true;
 
@@ -1685,12 +1715,12 @@ void FAnimationProvider::AppendPoseWatch(uint64 InAnimInstanceId, double InTime,
 		PoseWatchRequiredBones.PushBack() = RequiredBone;
 	}
 
-	Session.UpdateDurationSeconds(InTime);
+	UpdateSessionDurationSeconds(InTime);
 }
 
 void FAnimationProvider::AppendPoseWatch(uint64 InComponentId, uint64 InAnimInstanceId, double InTime, double InRecordingTime, uint64 PoseWatchId, uint32 NameId, FColor Color, const TArrayView<const float>& BoneTransformsRaw, const TArrayView<const uint32>& CurveIds, const TArrayView<const float>& CurveValues, const TArrayView<const uint16>& RequiredBones, const TArrayView<const float>& WorldTransformRaw, const bool bIsEnabled)
 {
-	Session.WriteAccessCheck();
+	EditAccessCheck();
 
 	bHasAnyData = true;
 
@@ -1777,12 +1807,12 @@ void FAnimationProvider::AppendPoseWatch(uint64 InComponentId, uint64 InAnimInst
 		PoseWatchRequiredBones.PushBack() = RequiredBone;
 	}
 
-	Session.UpdateDurationSeconds(InTime);
+	UpdateSessionDurationSeconds(InTime);
 }
 
 void FAnimationProvider::AppendInertialization(uint64 InAnimInstanceId, double InProfileTime, double InRecordingTime, int32 InNodeId, float InWeight, EInertializationType InType)
 {
-	Session.WriteAccessCheck();
+	EditAccessCheck();
 	
 	TSharedPtr<FInertializationTimelineStorage> TimelineStorage;
 	uint32* IndexPtr = ObjectIdToInertializationTimelines.Find(InAnimInstanceId);
diff --git a/Engine/Plugins/Animation/GameplayInsights/Source/GameplayInsights/Private/AnimationProvider.h b/Engine/Plugins/Animation/GameplayInsights/Source/GameplayInsights/Private/AnimationProvider.h
--- a/Engine/Plugins/Animation/GameplayInsights/Source/GameplayInsights/Private/AnimationProvider.h
+++ b/Engine/Plugins/Animation/GameplayInsights/Source/GameplayInsights/Private/AnimationProvider.h
@@ -5,23 +5,60 @@
 #include "CoreMinimal.h"
 #include "IAnimationProvider.h"
 #include "Common/PagedArray.h"
+#include "Common/ProviderLock.h"
 #include "Model/IntervalTimeline.h"
 #include "Model/PointTimeline.h"
 #include "Containers/ArrayView.h"
 #include "Model/IntervalTimeline.h"
 
-namespace TraceServices { class IAnalysisSession; }
+namespace TraceServices
+{
+	class IAnalysisSession;
+	extern thread_local FProviderLock::FThreadLocalState GAnimationProviderLockState;
+}
 class FGameplayProvider;
 struct FObjectInfo;
 class USkeletalMesh;
 
-class FAnimationProvider : public IAnimationProvider
+class FAnimationProvider : public IAnimationProvider, public TraceServices::IEditableProvider
 {
 public:
 	static FName ProviderName;
 
 	FAnimationProvider(TraceServices::IAnalysisSession& InSession, FGameplayProvider& InGameplayProvider);
 
+	//~ IProvider (read lock)
+	virtual void BeginRead() const override
+	{
+		Lock.BeginRead(TraceServices::GAnimationProviderLockState);
+	}
+
+	virtual void EndRead() const override
+	{
+		Lock.EndRead(TraceServices::GAnimationProviderLockState);
+	}
+
+	virtual void ReadAccessCheck() const override
+	{
+		Lock.ReadAccessCheck(TraceServices::GAnimationProviderLockState);
+	}
+
+	//~ IEditableProvider (edit lock)
+	virtual void BeginEdit() const override
+	{
+		Lock.BeginWrite(TraceServices::GAnimationProviderLockState);
+	}
+
+	virtual void EndEdit() const override
+	{
+		Lock.EndWrite(TraceServices::GAnimationProviderLockState);
+	}
+
+	virtual void EditAccessCheck() const override
+	{
+		Lock.WriteAccessCheck(TraceServices::GAnimationProviderLockState);
+	}
+
 	/** IAnimationProvider interface */
 	virtual void EnumerateSkeletalMeshPoseTimelines(TFunctionRef<void(uint64 ObjectId, const SkeletalMeshPoseTimeline&)> Callback) const override;
 	virtual bool ReadSkeletalMeshPoseTimeline(uint64 InObjectId, TFunctionRef<void(const SkeletalMeshPoseTimeline&, bool)> Callback) const override;
@@ -134,8 +171,11 @@ private:
 	/** Handle an object ending play - terminate existing scopes */
 	void HandleObjectEndPlay(uint64 InObjectId, double InTime, const FObjectInfo& InObjectInfo);
 
+	void UpdateSessionDurationSeconds(double InTime) const;
+
 private:
 	TraceServices::IAnalysisSession& Session;
+	mutable TraceServices::FProviderLock Lock;
 	FGameplayProvider& GameplayProvider;
 
 	/** Maps into timeline arrays per-object */
diff --git a/Engine/Plugins/Animation/GameplayInsights/Source/GameplayInsights/Private/AnimationSharedData.cpp b/Engine/Plugins/Animation/GameplayInsights/Source/GameplayInsights/Private/AnimationSharedData.cpp
--- a/Engine/Plugins/Animation/GameplayInsights/Source/GameplayInsights/Private/AnimationSharedData.cpp
+++ b/Engine/Plugins/Animation/GameplayInsights/Source/GameplayInsights/Private/AnimationSharedData.cpp
@@ -6,6 +6,7 @@
 #include "AnimationTickRecordsTrack.h"
 #include "AnimNodesTrack.h"
 #include "AnimNotifiesTrack.h"
+#include "Common/ProviderLock.h"
 #include "Framework/MultiBox/MultiBoxBuilder.h"
 #include "GameplayInsightsModule.h"
 #include "GameplayProvider.h"
@@ -76,12 +77,13 @@ void FAnimationSharedData::Tick(UE::Insights::Timing::ITimingViewSession& InTimi
 
 	if(AnimationProvider && GameplayProvider)
 	{
-		TraceServices::FAnalysisSessionReadScope SessionReadScope(GetAnalysisSession());
+		TraceServices::FProviderReadScopeLock AnimationProviderReadScope(*AnimationProvider);
 
 		if(AnimationProvider->HasAnyData() && AreAnyAnimationTracksEnabled())
 		{
 			// Add tracks for each tracked object's animation data
-			GameplayProvider->EnumerateObjects([this, &InTimingViewSession, &InAnalysisSession, &AnimationProvider, &GameplayProvider](const FObjectInfo& InObjectInfo)
+			TraceServices::FProviderReadScopeLock ProviderReadScope(*GameplayProvider);
+			GameplayProvider->EnumerateObjects([this, &InTimingViewSession, &InAnalysisSession, &AnimationProvider](const FObjectInfo& InObjectInfo)
 			{
 				TSharedPtr<FObjectEventsTrack> ObjectEventsTrack;
 				const uint64 ObjectId = InObjectInfo.GetUObjectId();
@@ -136,7 +138,7 @@ void FAnimationSharedData::Tick(UE::Insights::Timing::ITimingViewSession& InTimi
 					}
 				});
 
-				AnimationProvider->ReadTickRecordTimeline(ObjectId, [this, ObjectId, &InObjectInfo, &ObjectEventsTrack, &InTimingViewSession, &GameplayProvider, &InAnalysisSession](const IAnimationProvider::TickRecordTimeline& InTimeline)
+				AnimationProvider->ReadTickRecordTimeline(ObjectId, [this, ObjectId, &InObjectInfo, &ObjectEventsTrack, &InTimingViewSession, &InAnalysisSession](const IAnimationProvider::TickRecordTimeline& InTimeline)
 				{
 					if(!ObjectEventsTrack.IsValid())
 					{
@@ -702,10 +704,12 @@ void FAnimationSharedData::OpenAnimGraphTab(uint64 InAnimInstanceId) const
 		const FGameplayProvider* GameplayProvider = AnalysisSession->ReadProvider<FGameplayProvider>(FGameplayProvider::ProviderName);
 		if(GameplayProvider)
 		{
-			TraceServices::FAnalysisSessionReadScope SessionReadScope(*AnalysisSession);
-
-			const FObjectInfo& ObjectInfo = GameplayProvider->GetObjectInfo(InAnimInstanceId);
-			Tab->SetLabel(FText::FromString(ObjectInfo.Name));
+			FString ObjectName;
+			{
+				TraceServices::FProviderReadScopeLock ProviderReadScope(*GameplayProvider);
+				ObjectName = GameplayProvider->GetObjectInfo(InAnimInstanceId).Name;
+			}
+			Tab->SetLabel(FText::FromString(ObjectName));
 		}
 	}
 }
diff --git a/Engine/Plugins/Animation/GameplayInsights/Source/GameplayInsights/Private/AnimationTickRecordsTrack.cpp b/Engine/Plugins/Animation/GameplayInsights/Source/GameplayInsights/Private/AnimationTickRecordsTrack.cpp
--- a/Engine/Plugins/Animation/GameplayInsights/Source/GameplayInsights/Private/AnimationTickRecordsTrack.cpp
+++ b/Engine/Plugins/Animation/GameplayInsights/Source/GameplayInsights/Private/AnimationTickRecordsTrack.cpp
@@ -4,6 +4,7 @@
 
 #include "AnimationProvider.h"
 #include "AnimationSharedData.h"
+#include "Common/ProviderLock.h"
 #include "Framework/MultiBox/MultiBoxBuilder.h"
 #include "GameplayProvider.h"
 #include "GameplaySharedData.h"
@@ -77,6 +78,7 @@ FAnimationTickRecordsTrack::FAnimationTickRecordsTrack(const FAnimationSharedDat
 
 	if(GameplayProvider)
 	{
+		TraceServices::FProviderReadScopeLock ProviderReadScope(*GameplayProvider);
 		const FObjectInfo* AnimInstanceInfo = GameplayProvider->FindObjectInfo(GetGameplayTrack().GetObjectId());
 		if(AnimInstanceInfo)
 		{
@@ -104,7 +106,8 @@ void FAnimationTickRecordsTrack::AddAllSeries()
 	const FAnimationProvider* AnimationProvider = SharedData.GetAnalysisSession().ReadProvider<FAnimationProvider>(FAnimationProvider::ProviderName);
 	if(GameplayProvider && AnimationProvider)
 	{
-		TraceServices::FAnalysisSessionReadScope SessionReadScope(SharedData.GetAnalysisSession());
+		TraceServices::FProviderReadScopeLock AnimationProviderReadScope(*AnimationProvider);
+		TraceServices::FProviderReadScopeLock GameplayProviderReadScope(*GameplayProvider);
 
 		AnimationProvider->EnumerateTickRecordIds(GetGameplayTrack().GetObjectId(), [this, &GameplayProvider](uint64 InAssetId, int32 InNodeId)
 		{
@@ -220,13 +223,12 @@ bool FAnimationTickRecordsTrack::UpdateSeriesBoundsHelper(FTickRecordSeries& InS
 
 	const FAnimationProvider* AnimationProvider = SharedData.GetAnalysisSession().ReadProvider<FAnimationProvider>(FAnimationProvider::ProviderName);
 
-	if(AnimationProvider)
+	if (AnimationProvider)
 	{
-		TraceServices::FAnalysisSessionReadScope SessionReadScope(SharedData.GetAnalysisSession());
-
 		InSeries.CurrentMin = 0.0;
 		InSeries.CurrentMax = 0.0;
 
+		TraceServices::FProviderReadScopeLock ProviderReadScope(*AnimationProvider);
 		AnimationProvider->ReadTickRecordTimeline(GetGameplayTrack().GetObjectId(), [&bFoundEvents, &InViewport, &InSeries, &Projection](const FAnimationProvider::TickRecordTimeline& InTimeline)
 		{
 			InTimeline.EnumerateEvents(InViewport.GetStartTime(), InViewport.GetEndTime(), [&bFoundEvents, &InSeries, &Projection](double InStartTime, double InEndTime, uint32 InDepth, const FTickRecordMessage& InMessage)
@@ -253,10 +255,9 @@ void FAnimationTickRecordsTrack::UpdateSeriesHelper(FTickRecordSeries& InSeries,
 
 	if(AnimationProvider)
 	{
-		TraceServices::FAnalysisSessionReadScope SessionReadScope(SharedData.GetAnalysisSession());
-
 		FGraphTrackBuilder Builder(*this, InSeries, InViewport);
 
+		TraceServices::FProviderReadScopeLock ProviderReadScope(*AnimationProvider);
 		AnimationProvider->ReadTickRecordTimeline(GetGameplayTrack().GetObjectId(), [this, &AnimationProvider, &Builder, &InViewport, &InSeries, &Projection](const FAnimationProvider::TickRecordTimeline& InTimeline)
 		{
 			uint16 FrameCounter = 0;
@@ -393,7 +394,7 @@ void FAnimationTickRecordsTrack::FindTickRecordMessage(const FTimingEventSearchP
 
 			if(AnimationProvider)
 			{
-				TraceServices::FAnalysisSessionReadScope SessionReadScope(SharedData.GetAnalysisSession());
+				TraceServices::FProviderReadScopeLock ProviderReadScope(*AnimationProvider);
 
 				AnimationProvider->ReadTickRecordTimeline(GetGameplayTrack().GetObjectId(), [this, &InContext](const FAnimationProvider::TickRecordTimeline& InTimeline)
 				{
@@ -473,7 +474,8 @@ void FAnimationTickRecordsTrack::GetVariantsAtFrame(const TraceServices::FFrame&
 	const FAnimationProvider* AnimationProvider = SharedData.GetAnalysisSession().ReadProvider<FAnimationProvider>(FAnimationProvider::ProviderName);
 	if(GameplayProvider && AnimationProvider)
 	{
-		TraceServices::FAnalysisSessionReadScope SessionReadScope(SharedData.GetAnalysisSession());
+		TraceServices::FProviderReadScopeLock AnimationProviderReadScope(*AnimationProvider);
+		TraceServices::FProviderReadScopeLock GameplayProviderReadScope(*GameplayProvider);
 
 		AnimationProvider->ReadTickRecordTimeline(GetGameplayTrack().GetObjectId(), [&GameplayProvider, &OutVariants, &InFrame](const FAnimationProvider::TickRecordTimeline& InTimeline)
 		{
diff --git a/Engine/Plugins/Animation/GameplayInsights/Source/GameplayInsights/Private/BlendWeightsTrack.cpp b/Engine/Plugins/Animation/GameplayInsights/Source/GameplayInsights/Private/BlendWeightsTrack.cpp
--- a/Engine/Plugins/Animation/GameplayInsights/Source/GameplayInsights/Private/BlendWeightsTrack.cpp
+++ b/Engine/Plugins/Animation/GameplayInsights/Source/GameplayInsights/Private/BlendWeightsTrack.cpp
@@ -1,6 +1,7 @@
 // Copyright Epic Games, Inc. All Rights Reserved.
 
 #include "BlendWeightsTrack.h"
+#include "Common/ProviderLock.h"
 #include "SCurveTimelineView.h"
 #include "IRewindDebugger.h"
 #include "GameplayProvider.h"
@@ -64,28 +65,29 @@ bool FBlendWeightsTrack::UpdateInternal()
 	double EndTime = TraceTimeRange.GetUpperBoundValue();
 	
 	// count number of unique animations in the current time range
-	const FGameplayProvider* GameplayProvider = AnalysisSession->ReadProvider<FGameplayProvider>(FGameplayProvider::ProviderName);
 	const FAnimationProvider* AnimationProvider = AnalysisSession->ReadProvider<FAnimationProvider>(FAnimationProvider::ProviderName);
 
 	int AnimationCount = 0;
 
-	if(GameplayProvider && AnimationProvider)
+	if (AnimationProvider)
 	{
-		TraceServices::FAnalysisSessionReadScope SessionReadScope(*AnalysisSession);
 		UniqueTrackIds.SetNum(0, EAllowShrinking::No);
 
-		AnimationProvider->ReadTickRecordTimeline(ObjectId, [&UniqueTrackIds,&GameplayProvider, StartTime, EndTime](const FAnimationProvider::TickRecordTimeline& InTimeline)
 		{
-			// this isn't very efficient, and it gets called every frame.  will need optimizing
-			InTimeline.EnumerateEvents(StartTime, EndTime, [&UniqueTrackIds, StartTime, EndTime](double InStartTime, double InEndTime, uint32 InDepth, const FTickRecordMessage& InMessage)
+			TraceServices::FProviderReadScopeLock AnimationProviderReadScope(*AnimationProvider);
+			AnimationProvider->ReadTickRecordTimeline(ObjectId, [&UniqueTrackIds, StartTime, EndTime](const FAnimationProvider::TickRecordTimeline& InTimeline)
 			{
-				if (InEndTime > StartTime && InStartTime < EndTime)
+				// this isn't very efficient, and it gets called every frame.  will need optimizing
+				InTimeline.EnumerateEvents(StartTime, EndTime, [&UniqueTrackIds, StartTime, EndTime](double InStartTime, double InEndTime, uint32 InDepth, const FTickRecordMessage& InMessage)
 				{
-					UniqueTrackIds.AddUnique(TrackId(InMessage.AssetId, InMessage.NodeId));
-				}
-				return TraceServices::EEventEnumerate::Continue;
+					if (InEndTime > StartTime && InStartTime < EndTime)
+					{
+						UniqueTrackIds.AddUnique(TrackId(InMessage.AssetId, InMessage.NodeId));
+					}
+					return TraceServices::EEventEnumerate::Continue;
+				});
 			});
-		});
+		}
 
 		UniqueTrackIds.StableSort([](const TrackId &A, const TrackId &B)
 			{
@@ -169,14 +171,13 @@ bool FBlendWeightTrack::UpdateInternal()
 
 	int AnimationCount = 0;
 
-	if(CurvesUpdateRequested > 10 && GameplayProvider && AnimationProvider)
+	if (CurvesUpdateRequested > 10 && GameplayProvider && AnimationProvider)
 	{
 		TRACE_CPUPROFILER_EVENT_SCOPE(FBlendWeightTrack::UpdateCurvePointsInternal);
 		auto& CurvePoints = CurveData->Points;
 		CurvePoints.SetNum(0,EAllowShrinking::No);
-		
-		TraceServices::FAnalysisSessionReadScope SessionReadScope(*AnalysisSession);
-		
+
+		TraceServices::FProviderReadScopeLock AnimationProviderReadScope(*AnimationProvider);
 		AnimationProvider->ReadTickRecordTimeline(ObjectId, [AnalysisSession, StartTime, EndTime, &CurvePoints, this](const FAnimationProvider::TickRecordTimeline& InTimeline)
 		{
 			InTimeline.EnumerateEvents(StartTime, EndTime, [&CurvePoints, AnalysisSession, this](double InStartTime, double InEndTime, uint32 InDepth, const FTickRecordMessage& InMessage)
@@ -215,6 +216,7 @@ bool FBlendWeightTrack::UpdateInternal()
 		// Blend Weight track gets name/icon of animation
 		if (CurveName.IsEmpty() && GameplayProvider)
 		{
+			TraceServices::FProviderReadScopeLock GameplayProviderReadScope(*GameplayProvider);
 			if (const FObjectInfo* ObjectInfo = GameplayProvider->FindObjectInfo(AssetId))
 			{
 				CurveName = FText::FromString(ObjectInfo->Name);
@@ -340,28 +342,42 @@ bool FBlendWeightTrack::HandleDoubleClickInternal()
 	IRewindDebugger* RewindDebugger = IRewindDebugger::Instance();
 	if (const TraceServices::IAnalysisSession* AnalysisSession = RewindDebugger->GetAnalysisSession())
 	{
-		TraceServices::FAnalysisSessionReadScope SessionReadScope(*AnalysisSession);
-
 		const FGameplayProvider* GameplayProvider = AnalysisSession->ReadProvider<FGameplayProvider>(FGameplayProvider::ProviderName);
+		const FAnimationProvider* AnimationProvider = AnalysisSession->ReadProvider<FAnimationProvider>(FAnimationProvider::ProviderName);
+
+		if (!GameplayProvider || !AnimationProvider)
+		{
+			return false;
+		}
 
-		const FObjectInfo& AssetInfo = GameplayProvider->GetObjectInfo(GetAssetId());
-		
-		UObject* Asset = nullptr;
 		bool bMessageFound = false;
 		float PlaybackTime = 0;
 		float BlendSpaceX = 0;
 		float BlendSpaceY = 0;
 
-		float CurrentTraceTime = IRewindDebugger::Instance()->CurrentTraceTime();
-		
-		const FAnimationProvider* AnimationProvider = AnalysisSession->ReadProvider<FAnimationProvider>(FAnimationProvider::ProviderName);
-		const TraceServices::IFrameProvider& FrameProvider = TraceServices::ReadFrameProvider(*AnalysisSession);
+		FString AssetPathName;
+		{
+			TraceServices::FProviderReadScopeLock GameplayProviderReadScope(*GameplayProvider);
+			const FObjectInfo& AssetInfo = GameplayProvider->GetObjectInfo(GetAssetId());
+			AssetPathName = AssetInfo.PathName;
+		}
+
+		const float CurrentTraceTime = IRewindDebugger::Instance()->CurrentTraceTime();
+
+		bool bFrameFound;
 		TraceServices::FFrame Frame;
-		if (FrameProvider.GetFrameFromTime(ETraceFrameType::TraceFrameType_Game, CurrentTraceTime, Frame))
 		{
-			AnimationProvider->ReadTickRecordTimeline(ObjectId, [this, &bMessageFound, &PlaybackTime, &BlendSpaceX, &BlendSpaceY, &GameplayProvider, &Frame](const FAnimationProvider::TickRecordTimeline& InTimeline)
+			TraceServices::FAnalysisSessionReadScope SessionReadScope(*AnalysisSession);
+			const TraceServices::IFrameProvider& FrameProvider = TraceServices::ReadFrameProvider(*AnalysisSession);
+			bFrameFound = FrameProvider.GetFrameFromTime(TraceFrameType_Game, CurrentTraceTime, Frame);
+		}
+
+		if (bFrameFound)
+		{
+			TraceServices::FProviderReadScopeLock AnimationProviderReadScope(*AnimationProvider);
+			AnimationProvider->ReadTickRecordTimeline(ObjectId, [this, &bMessageFound, &PlaybackTime, &BlendSpaceX, &BlendSpaceY, &Frame](const FAnimationProvider::TickRecordTimeline& InTimeline)
 			{
-				InTimeline.EnumerateEvents(Frame.StartTime, Frame.EndTime, [this, &bMessageFound, &PlaybackTime, &BlendSpaceX, &BlendSpaceY, &GameplayProvider, &Frame](double InStartTime, double InEndTime, uint32 InDepth, const FTickRecordMessage& InMessage)
+				InTimeline.EnumerateEvents(Frame.StartTime, Frame.EndTime, [this, &bMessageFound, &PlaybackTime, &BlendSpaceX, &BlendSpaceY, &Frame](double InStartTime, double InEndTime, uint32 InDepth, const FTickRecordMessage& InMessage)
 				{
 					if(InStartTime >= Frame.StartTime && InEndTime <= Frame.EndTime)
 					{
@@ -378,21 +394,22 @@ bool FBlendWeightTrack::HandleDoubleClickInternal()
 				});
 			});
 		}
-		
-		FString PackagePathString = FPackageName::ObjectPathToPackageName(FString(AssetInfo.PathName));
+
+		UObject* Asset = nullptr;
+		FString PackagePathString = FPackageName::ObjectPathToPackageName(AssetPathName);
 
 		UPackage* Package = LoadPackage(NULL, ToCStr(PackagePathString), LOAD_NoRedirects);
 		if (Package)
 		{
 			Package->FullyLoad();
                 
-			FString AssetName = FPaths::GetBaseFilename(AssetInfo.PathName);
+			FString AssetName = FPaths::GetBaseFilename(AssetPathName);
 			Asset = FindObject<UObject>(Package, *AssetName);
 		}
 		else
 		{
 			// fallback for unsaved assets
-			Asset = FindObject<UObject>(nullptr, AssetInfo.PathName);
+			Asset = FindObject<UObject>(nullptr, *AssetPathName);
 		}
                     	
 		if (Asset != nullptr)
@@ -452,10 +469,10 @@ bool FBlendWeightsTrackCreator::HasDebugInfoInternal(const FObjectId& InObjectId
 {
 	TRACE_CPUPROFILER_EVENT_SCOPE(FBlendWeightsTrack::HasDebugInfoInternal);
 	const TraceServices::IAnalysisSession* AnalysisSession = IRewindDebugger::Instance()->GetAnalysisSession();
-	TraceServices::FAnalysisSessionReadScope SessionReadScope(*AnalysisSession);
 	bool bHasData = false;
 	if (const FAnimationProvider* AnimationProvider = AnalysisSession->ReadProvider<FAnimationProvider>(FAnimationProvider::ProviderName))
 	{
+		TraceServices::FProviderReadScopeLock ProviderReadScope(*AnimationProvider);
 		AnimationProvider->ReadTickRecordTimeline(InObjectId.GetMainId(), [&bHasData](const FAnimationProvider::TickRecordTimeline& InGraphTimeline)
 		{
 			bHasData = true;
diff --git a/Engine/Plugins/Animation/GameplayInsights/Source/GameplayInsights/Private/BoolPropertyTrack.cpp b/Engine/Plugins/Animation/GameplayInsights/Source/GameplayInsights/Private/BoolPropertyTrack.cpp
--- a/Engine/Plugins/Animation/GameplayInsights/Source/GameplayInsights/Private/BoolPropertyTrack.cpp
+++ b/Engine/Plugins/Animation/GameplayInsights/Source/GameplayInsights/Private/BoolPropertyTrack.cpp
@@ -2,6 +2,7 @@
 
 #include "BoolPropertyTrack.h"
 
+#include "Common/ProviderLock.h"
 #include "GameplayProvider.h"
 #include "PropertyHelpers.h"
 
@@ -23,10 +24,9 @@ namespace RewindDebugger
 
 		if (const TraceServices::IAnalysisSession* AnalysisSession = RewindDebugger->GetAnalysisSession())
 		{
-			TraceServices::FAnalysisSessionReadScope SessionReadScope(*AnalysisSession);
-		
 			if (const FGameplayProvider* GameplayProvider = AnalysisSession->ReadProvider<FGameplayProvider>(FGameplayProvider::ProviderName))
 			{
+				TraceServices::FProviderReadScopeLock ProviderReadScope(*GameplayProvider);
 				return FObjectPropertyHelpers::IsBoolProperty(InObjectProperty.Property.Value, GameplayProvider->GetPropertyName(InObjectProperty.Property.TypeStringId));
 			}
 		}
diff --git a/Engine/Plugins/Animation/GameplayInsights/Source/GameplayInsights/Private/ExternalMorphTrack.cpp b/Engine/Plugins/Animation/GameplayInsights/Source/GameplayInsights/Private/ExternalMorphTrack.cpp
--- a/Engine/Plugins/Animation/GameplayInsights/Source/GameplayInsights/Private/ExternalMorphTrack.cpp
+++ b/Engine/Plugins/Animation/GameplayInsights/Source/GameplayInsights/Private/ExternalMorphTrack.cpp
@@ -1,6 +1,7 @@
 // Copyright Epic Games, Inc. All Rights Reserved.
 
 #include "ExternalMorphTrack.h"
+#include "Common/ProviderLock.h"
 #include "SCurveTimelineView.h"
 #include "IRewindDebugger.h"
 #include "GameplayProvider.h"
@@ -30,22 +31,21 @@ namespace RewindDebugger
 	{
 		IRewindDebugger* RewindDebugger = IRewindDebugger::Instance();
 		const TraceServices::IAnalysisSession* AnalysisSession = RewindDebugger->GetAnalysisSession();
+		const FAnimationProvider* AnimationProvider = AnalysisSession->ReadProvider<FAnimationProvider>(FAnimationProvider::ProviderName);
+		if (AnimationProvider == nullptr)
+		{
+			return false;
+		}
 
 		// Convert time range to from rewind debugger times to profiler times.
 		TRange<double> TraceTimeRange = RewindDebugger->GetCurrentTraceRange();
 		double StartTime = TraceTimeRange.GetLowerBoundValue();
 		double EndTime = TraceTimeRange.GetUpperBoundValue();
-	
-		const FGameplayProvider* GameplayProvider = AnalysisSession->ReadProvider<FGameplayProvider>(FGameplayProvider::ProviderName);
-		const FAnimationProvider* AnimationProvider = AnalysisSession->ReadProvider<FAnimationProvider>(FAnimationProvider::ProviderName);
 
 		TArray<int32> MorphSetIndexValues;
-
 		bool bChanged = false;
-		if (GameplayProvider && AnimationProvider)
 		{
-			TraceServices::FAnalysisSessionReadScope SessionReadScope(*AnalysisSession);
-
+			TraceServices::FProviderReadScopeLock AnimationProviderReadScope(*AnimationProvider);
 			AnimationProvider->ReadSkeletalMeshPoseTimeline(ObjectId, [this, AnimationProvider, StartTime, EndTime, &MorphSetIndexValues](const FAnimationProvider::SkeletalMeshPoseTimeline& InTimeline, bool bHasCurves)
 			{
 				InTimeline.EnumerateEvents(StartTime, EndTime, [this, StartTime, EndTime, AnimationProvider, &MorphSetIndexValues](double InStartTime, double InEndTime, uint32 InDepth, const FSkeletalMeshPoseMessage& InMessage)
@@ -60,27 +60,26 @@ namespace RewindDebugger
 					return TraceServices::EEventEnumerate::Continue;
 				});
 			});
+		}
 
-			bChanged = false;
-			if (MorphSetIndexValues.Num() != Children.Num() || PrevMorphSetIndexValues != MorphSetIndexValues)
-			{
-				bChanged = true;
-				Children.SetNum(MorphSetIndexValues.Num());
-				for (int32 Index = 0; Index < MorphSetIndexValues.Num(); ++Index)
-				{
-					TSharedPtr<FExternalMorphSetTrack> NewTrack = MakeShared<FExternalMorphSetTrack>(ObjectId, MorphSetIndexValues[Index]);
-					Children[Index] = NewTrack;
-				}
-			}
-
-			for (int32 Index = 0; Index < Children.Num(); ++Index)
+		if (MorphSetIndexValues.Num() != Children.Num() || PrevMorphSetIndexValues != MorphSetIndexValues)
+		{
+			bChanged = true;
+			Children.SetNum(MorphSetIndexValues.Num());
+			for (int32 Index = 0; Index < MorphSetIndexValues.Num(); ++Index)
 			{
-				bChanged |= Children[Index]->Update();
+				TSharedPtr<FExternalMorphSetTrack> NewTrack = MakeShared<FExternalMorphSetTrack>(ObjectId, MorphSetIndexValues[Index]);
+				Children[Index] = NewTrack;
 			}
+		}
 
-			PrevMorphSetIndexValues = MorphSetIndexValues;
+		for (int32 Index = 0; Index < Children.Num(); ++Index)
+		{
+			bChanged |= Children[Index]->Update();
 		}
 
+		PrevMorphSetIndexValues = MorphSetIndexValues;
+
 		return bChanged;
 	}
 
@@ -111,21 +110,21 @@ namespace RewindDebugger
 	{
 		IRewindDebugger* RewindDebugger = IRewindDebugger::Instance();
 		const TraceServices::IAnalysisSession* AnalysisSession = RewindDebugger->GetAnalysisSession();
+		const FAnimationProvider* AnimationProvider = AnalysisSession->ReadProvider<FAnimationProvider>(FAnimationProvider::ProviderName);
+		if (AnimationProvider == nullptr)
+		{
+			return false;
+		}
 
 		// Convert time range to from rewind debugger times to profiler times.
 		TRange<double> TraceTimeRange = RewindDebugger->GetCurrentTraceRange();
 		double StartTime = TraceTimeRange.GetLowerBoundValue();
 		double EndTime = TraceTimeRange.GetUpperBoundValue();
 	
-		const FGameplayProvider* GameplayProvider = AnalysisSession->ReadProvider<FGameplayProvider>(FGameplayProvider::ProviderName);
-		const FAnimationProvider* AnimationProvider = AnalysisSession->ReadProvider<FAnimationProvider>(FAnimationProvider::ProviderName);
-
 		int32 NumMorphTracks = 0;
 		bool bChanged = false;
-		if (GameplayProvider && AnimationProvider)
 		{
-			TraceServices::FAnalysisSessionReadScope SessionReadScope(*AnalysisSession);
-
+			TraceServices::FProviderReadScopeLock AnimationProviderReadScope(*AnimationProvider);
 			AnimationProvider->ReadSkeletalMeshPoseTimeline(ObjectId, [this, AnimationProvider, StartTime, EndTime, &NumMorphTracks](const FAnimationProvider::SkeletalMeshPoseTimeline& InTimeline, bool bHasCurves)
 			{
 				InTimeline.EnumerateEvents(StartTime, EndTime, [this, StartTime, EndTime, AnimationProvider, &NumMorphTracks](double InStartTime, double InEndTime, uint32 InDepth, const FSkeletalMeshPoseMessage& InMessage)
@@ -143,26 +142,25 @@ namespace RewindDebugger
 					return TraceServices::EEventEnumerate::Continue;
 				});
 			});
+		}
 
-			bChanged = false;
-			if (NumMorphTracks != Children.Num())
-			{
-				bChanged = true;
-				Children.SetNum(NumMorphTracks);
-				for (int32 Index = 0; Index < NumMorphTracks; ++Index)
-				{
-					TSharedPtr<FExternalMorphTrack> NewTrack = MakeShared<FExternalMorphTrack>(ObjectId, Index, MorphSetIndex);
-					Children[Index] = NewTrack;
-				}
-			}
-
-			check(Children.Num() == NumMorphTracks);
+		if (NumMorphTracks != Children.Num())
+		{
+			bChanged = true;
+			Children.SetNum(NumMorphTracks);
 			for (int32 Index = 0; Index < NumMorphTracks; ++Index)
 			{
-				bChanged |= Children[Index]->Update();
+				TSharedPtr<FExternalMorphTrack> NewTrack = MakeShared<FExternalMorphTrack>(ObjectId, Index, MorphSetIndex);
+				Children[Index] = NewTrack;
 			}
 		}
 
+		check(Children.Num() == NumMorphTracks);
+		for (int32 Index = 0; Index < NumMorphTracks; ++Index)
+		{
+			bChanged |= Children[Index]->Update();
+		}
+
 		return bChanged;
 	}
 
@@ -224,7 +222,7 @@ namespace RewindDebugger
 		CurvePoints.Reset();
 		CurvePoints.Reserve(256);
 
-		TraceServices::FAnalysisSessionReadScope SessionReadScope(*AnalysisSession);
+		TraceServices::FProviderReadScopeLock ProviderReadScope(*AnimationProvider);
 		AnimationProvider->ReadSkeletalMeshPoseTimeline(ObjectId, [this, AnimationProvider, StartTime, EndTime, AnalysisSession, &CurvePoints](const FAnimationProvider::SkeletalMeshPoseTimeline& InTimeline, bool bHasCurves)
 		{
 			InTimeline.EnumerateEvents(StartTime, EndTime, [this, StartTime, EndTime, AnimationProvider, AnalysisSession, &CurvePoints](double InStartTime, double InEndTime, uint32 InDepth, const FSkeletalMeshPoseMessage& InMessage)
@@ -291,9 +289,9 @@ namespace RewindDebugger
 	bool FExternalMorphSetGroupTrackCreator::HasDebugInfoInternal(const FObjectId& InObjectId) const
 	{
 		const TraceServices::IAnalysisSession* AnalysisSession = IRewindDebugger::Instance()->GetAnalysisSession();
-		TraceServices::FAnalysisSessionReadScope SessionReadScope(*AnalysisSession);
 		if (const FAnimationProvider* AnimationProvider = AnalysisSession->ReadProvider<FAnimationProvider>(FAnimationProvider::ProviderName))
 		{
+			TraceServices::FProviderReadScopeLock ProviderReadScope(*AnimationProvider);
 			return AnimationProvider->HasExternalMorphSets(InObjectId.GetMainId());
 		}
 		return false;
diff --git a/Engine/Plugins/Animation/GameplayInsights/Source/GameplayInsights/Private/GameplayAnalyzer.cpp b/Engine/Plugins/Animation/GameplayInsights/Source/GameplayInsights/Private/GameplayAnalyzer.cpp
--- a/Engine/Plugins/Animation/GameplayInsights/Source/GameplayInsights/Private/GameplayAnalyzer.cpp
+++ b/Engine/Plugins/Animation/GameplayInsights/Source/GameplayInsights/Private/GameplayAnalyzer.cpp
@@ -42,7 +42,7 @@ bool FGameplayAnalyzer::OnEvent(uint16 RouteId, EStyle Style, const FOnEventCont
 {
 	LLM_SCOPE_BYNAME(TEXT("Insights/FGameplayAnalyzer"));
 
-	TraceServices::FAnalysisSessionEditScope _(Session);
+	TraceServices::FProviderEditScopeLock ProviderEditScope(GameplayProvider);
 
 	const auto& EventData = Context.EventData;
 	switch (RouteId)
diff --git a/Engine/Plugins/Animation/GameplayInsights/Source/GameplayInsights/Private/GameplayProvider.cpp b/Engine/Plugins/Animation/GameplayInsights/Source/GameplayInsights/Private/GameplayProvider.cpp
--- a/Engine/Plugins/Animation/GameplayInsights/Source/GameplayInsights/Private/GameplayProvider.cpp
+++ b/Engine/Plugins/Animation/GameplayInsights/Source/GameplayInsights/Private/GameplayProvider.cpp
@@ -106,6 +106,10 @@ static UClass* FindClassByPathName(const FString& RawTargetPathName, const bool
 }
 #endif
 
+namespace TraceServices
+{
+thread_local FProviderLock::FThreadLocalState GGameplayProviderLockState;
+}
 
 FGameplayProvider::FGameplayProvider(TraceServices::IAnalysisSession& InSession)
 	: Session(InSession)
@@ -119,7 +123,7 @@ FGameplayProvider::FGameplayProvider(TraceServices::IAnalysisSession& InSession)
 
 bool FGameplayProvider::ReadObjectEventsTimeline(const RewindDebugger::FObjectId& InObjectId, TFunctionRef<void(const ObjectEventsTimeline&)> Callback) const
 {
-	Session.ReadAccessCheck();
+	ReadAccessCheck();
 
 	const uint32* IndexPtr = ObjectIdToEventTimelines.Find(InObjectId);
 	if (IndexPtr != nullptr)
@@ -136,7 +140,7 @@ bool FGameplayProvider::ReadObjectEventsTimeline(const RewindDebugger::FObjectId
 
 bool FGameplayProvider::ReadObjectEvent(const RewindDebugger::FObjectId& InObjectId, uint64 InMessageId, TFunctionRef<void(const FObjectEventMessage&)> Callback) const
 {
-	Session.ReadAccessCheck();
+	ReadAccessCheck();
 
 	return ReadObjectEventsTimeline(InObjectId, [&Callback, &InMessageId](const ObjectEventsTimeline& InTimeline)
 		{
@@ -149,7 +153,7 @@ bool FGameplayProvider::ReadObjectEvent(const RewindDebugger::FObjectId& InObjec
 
 bool FGameplayProvider::ReadObjectPropertiesTimeline(const uint64 InObjectId, TFunctionRef<void(const ObjectPropertiesTimeline&)> Callback) const
 {
-	Session.ReadAccessCheck();
+	ReadAccessCheck();
 
 	const uint32* IndexPtr = ObjectIdToPropertiesStorage.Find(InObjectId);
 	if (IndexPtr != nullptr)
@@ -166,7 +170,7 @@ bool FGameplayProvider::ReadObjectPropertiesTimeline(const uint64 InObjectId, TF
 
 bool FGameplayProvider::ReadObjectPropertiesStorage(const uint64 InObjectId, const FObjectPropertiesMessage& InMessage, TFunctionRef<void(const TConstArrayView<FObjectPropertyValue>&)> Callback) const
 {
-	Session.ReadAccessCheck();
+	ReadAccessCheck();
 
 	const uint32* IndexPtr = ObjectIdToPropertiesStorage.Find(InObjectId);
 	if (IndexPtr != nullptr)
@@ -184,7 +188,7 @@ bool FGameplayProvider::ReadObjectPropertiesStorage(const uint64 InObjectId, con
 
 void FGameplayProvider::EnumerateObjectPropertyValues(const uint64 InObjectId, const FObjectPropertiesMessage& InMessage, TFunctionRef<void(const FObjectPropertyValue&)> Callback) const
 {
-	Session.ReadAccessCheck();
+	ReadAccessCheck();
 
 	const uint32* IndexPtr = ObjectIdToPropertiesStorage.Find(InObjectId);
 	if (IndexPtr != nullptr)
@@ -202,7 +206,7 @@ void FGameplayProvider::EnumerateObjectPropertyValues(const uint64 InObjectId, c
 
 void FGameplayProvider::EnumerateObjects(TFunctionRef<void(const FObjectInfo&)> Callback) const
 {
-	Session.ReadAccessCheck();
+	ReadAccessCheck();
 
 	for (const FObjectInfo& ObjectInfo : ObjectInfos)
 	{
@@ -212,7 +216,7 @@ void FGameplayProvider::EnumerateObjects(TFunctionRef<void(const FObjectInfo&)>
 
 void FGameplayProvider::EnumerateObjects(double StartTime, double EndTime, TFunctionRef<void(const FObjectInfo&)> Callback) const
 {
-	Session.ReadAccessCheck();
+	ReadAccessCheck();
 
 	ObjectLifetimes.EnumerateEvents(StartTime, EndTime,
 		[this, Callback](double InStartTime, double InEndTime, uint32 InDepth, const FObjectExistsMessage& ExistsMessage)
@@ -228,7 +232,7 @@ void FGameplayProvider::EnumerateObjects(double StartTime, double EndTime, TFunc
 
 void FGameplayProvider::EnumerateWorlds(TFunctionRef<void(const FWorldInfo&)> Callback) const
 {
-	Session.ReadAccessCheck();
+	ReadAccessCheck();
 
 	for (const FWorldInfo& WorldInfo : WorldInfos)
 	{
@@ -238,7 +242,7 @@ void FGameplayProvider::EnumerateWorlds(TFunctionRef<void(const FWorldInfo&)> Ca
 
 const FObjectPropertyValue* FGameplayProvider::FindPropertyValueFromStorageIndex(const uint64 InObjectId, int64 InStorageIndex) const
 {
-	Session.ReadAccessCheck();
+	ReadAccessCheck();
 
 	const uint32* IndexPtr = ObjectIdToPropertiesStorage.Find(InObjectId);
 	if (IndexPtr != nullptr)
@@ -259,6 +263,7 @@ const FObjectPropertyValue* FGameplayProvider::FindPropertyValueFromStorageIndex
 
 void FGameplayProvider::EnumerateSubobjects(const RewindDebugger::FObjectId& InObjectId, TFunctionRef<void(const RewindDebugger::FObjectId& SubObjectId)> Callback) const
 {
+	ReadAccessCheck();
 	TArray<RewindDebugger::FObjectId, TInlineAllocator<32>> SubObjectIds;
 	ObjectHierarchy.MultiFind(InObjectId, SubObjectIds);
 
@@ -271,7 +276,7 @@ void FGameplayProvider::EnumerateSubobjects(const RewindDebugger::FObjectId& InO
 
 const FClassInfo* FGameplayProvider::FindClassInfo(uint64 InClassId) const
 {
-	Session.ReadAccessCheck();
+	ReadAccessCheck();
 
 	const int32* ClassIndex = ClassIdToIndexMap.Find(InClassId);
 	if (ClassIndex != nullptr)
@@ -297,7 +302,7 @@ const UClass* FGameplayProvider::FindClass(uint64 InClassId, bool bSearchBluepri
 
 const FClassInfo* FGameplayProvider::FindClassInfo(const TCHAR* InClassPath) const
 {
-	Session.ReadAccessCheck();
+	ReadAccessCheck();
 
 	const int32* ClassIndex = ClassPathNameToIndexMap.Find(InClassPath);
 	if (ClassIndex != nullptr)
@@ -310,7 +315,7 @@ const FClassInfo* FGameplayProvider::FindClassInfo(const TCHAR* InClassPath) con
 
 bool FGameplayProvider::IsSubClassOf(uint64 InSubClassId, uint64 InParentClassId) const
 {
-	Session.ReadAccessCheck();
+	ReadAccessCheck();
 
 	uint64 ClassId = InSubClassId;
 	while (true)
@@ -341,7 +346,7 @@ const FObjectInfo* FGameplayProvider::FindObjectInfo(uint64 InObjectId) const
 
 const FObjectInfo* FGameplayProvider::FindObjectInfo(const RewindDebugger::FObjectId& InObjectId) const
 {
-	Session.ReadAccessCheck();
+	ReadAccessCheck();
 
 	const int32* ObjectIndex = ObjectIdToIndexMap.Find(InObjectId);
 	if (ObjectIndex != nullptr)
@@ -354,7 +359,7 @@ const FObjectInfo* FGameplayProvider::FindObjectInfo(const RewindDebugger::FObje
 
 const FWorldInfo* FGameplayProvider::FindWorldInfo(const uint64 InObjectId) const
 {
-	Session.ReadAccessCheck();
+	ReadAccessCheck();
 
 	const int32* WorldIndex = WorldIdToIndexMap.Find(InObjectId);
 	if (WorldIndex != nullptr)
@@ -451,6 +456,7 @@ const FObjectInfo& FGameplayProvider::GetObjectInfo(const RewindDebugger::FObjec
 
 const TCHAR* FGameplayProvider::GetPropertyName(uint32 InPropertyStringId) const
 {
+	ReadAccessCheck();
 	if (const TCHAR* const* FoundString = PropertyStrings.Find(InPropertyStringId))
 	{
 		return *FoundString;
@@ -462,7 +468,7 @@ const TCHAR* FGameplayProvider::GetPropertyName(uint32 InPropertyStringId) const
 
 void FGameplayProvider::AppendClass(uint64 InClassId, uint64 InSuperId, const TCHAR* InClassName, const TCHAR* InClassPathName)
 {
-	Session.WriteAccessCheck();
+	EditAccessCheck();
 
 	bHasAnyData = true;
 
@@ -485,7 +491,7 @@ void FGameplayProvider::AppendClass(uint64 InClassId, uint64 InSuperId, const TC
 
 void FGameplayProvider::AppendObject(const RewindDebugger::FObjectId& InObjectId, const RewindDebugger::FObjectId& InOuterId, uint64 InClassId, const TCHAR* InObjectName, const TCHAR* InObjectPathName, EObjectInfoFlags Flags)
 {
-	Session.WriteAccessCheck();
+	EditAccessCheck();
 
 	bHasAnyData = true;
 
@@ -502,7 +508,7 @@ void FGameplayProvider::AppendObject(const RewindDebugger::FObjectId& InObjectId
 
 void FGameplayProvider::AppendObjectTransform(uint64 InObjectId, double InProfileTime, double InRecordingTime, const FTransform& Transform)
 {
-	Session.WriteAccessCheck();
+	EditAccessCheck();
 
 	bHasAnyData = true;
 
@@ -516,7 +522,7 @@ void FGameplayProvider::AppendObjectTransform(uint64 InObjectId, double InProfil
 
 void FGameplayProvider::AppendObjectLifetimeBegin(const RewindDebugger::FObjectId& InObjectId, double InProfileTime, double InRecordingTime)
 {
-	Session.WriteAccessCheck();
+	EditAccessCheck();
 	bHasAnyData = true;
 
 	if (InObjectId.IsSet())
@@ -530,7 +536,7 @@ void FGameplayProvider::AppendObjectLifetimeBegin(const RewindDebugger::FObjectI
 
 void FGameplayProvider::AppendObjectLifetimeEnd(const RewindDebugger::FObjectId& InObjectId, double InProfileTime, double InRecordingTime)
 {
-	Session.WriteAccessCheck();
+	EditAccessCheck();
 	bHasAnyData = true;
 
 	if (const uint64* FoundIndex = ActiveObjectLifetimes.Find(InObjectId))
@@ -556,7 +562,7 @@ void FGameplayProvider::AppendObjectLifetimeEnd(const RewindDebugger::FObjectId&
 
 void FGameplayProvider::AppendObjectEvent(const RewindDebugger::FObjectId& InObjectId, double InTime, const TCHAR* InEventName)
 {
-	Session.WriteAccessCheck();
+	EditAccessCheck();
 
 	bHasAnyData = true;
 
@@ -594,12 +600,15 @@ void FGameplayProvider::AppendObjectEvent(const RewindDebugger::FObjectId& InObj
 
 	Timeline->AppendEvent(InTime, Message);
 
-	Session.UpdateDurationSeconds(InTime);
+	{
+		TraceServices::FAnalysisSessionEditScope SessionEditScope(Session);
+		Session.UpdateDurationSeconds(InTime);
+	}
 }
 
 void FGameplayProvider::AppendPawnPossess(uint64 InControllerId, uint64 InPawnId, double InTime)
 {
-	Session.WriteAccessCheck();
+	EditAccessCheck();
 	bHasAnyData = true;
 
 	// End any active controller attachment interval for this controller
@@ -620,6 +629,7 @@ void FGameplayProvider::AppendPawnPossess(uint64 InControllerId, uint64 InPawnId
 
 uint64 FGameplayProvider::FindPossessingController(uint64 PawnId, double Time) const
 {
+	ReadAccessCheck();
 	uint64 ControllerId = 0;
 	PawnPossession.EnumerateEvents(Time, Time, [&ControllerId, PawnId](double StartTime, double EndTime, const FPawnPossessMessage Message)
 		{
@@ -635,6 +645,7 @@ uint64 FGameplayProvider::FindPossessingController(uint64 PawnId, double Time) c
 
 bool FGameplayProvider::GetObjectTransform(uint64 InObjectId, double InStartTime, double InEndTime, FTransform& OutTransform) const
 {
+	ReadAccessCheck();
 	bool bSuccess = false;
 	TransformTimeline.EnumerateEvents(InStartTime, InEndTime, [&OutTransform, InObjectId, &bSuccess](bool bStart, double Time, const FObjectTransformMessage& TransformMessage)
 		{
@@ -651,7 +662,7 @@ bool FGameplayProvider::GetObjectTransform(uint64 InObjectId, double InStartTime
 
 TRange<double> FGameplayProvider::GetObjectTraceLifetime(const RewindDebugger::FObjectId& InObjectId) const
 {
-	Session.ReadAccessCheck();
+	ReadAccessCheck();
 
 	if (const uint64* FoundIndex = ActiveObjectLifetimes.Find(InObjectId))
 	{
@@ -666,7 +677,7 @@ TRange<double> FGameplayProvider::GetObjectTraceLifetime(const RewindDebugger::F
 TRange<double> FGameplayProvider::GetObjectRecordingLifetime(const RewindDebugger::FObjectId& InObjectId) const
 {
 	TRACE_CPUPROFILER_EVENT_SCOPE(FGameplayProvider::GetObjectRecordingLifetime);
-	Session.ReadAccessCheck();
+	ReadAccessCheck();
 
 	if (const uint64 *FoundIndex = ActiveObjectRecordingLifetimes.Find(InObjectId))
 	{
@@ -680,7 +691,7 @@ TRange<double> FGameplayProvider::GetObjectRecordingLifetime(const RewindDebugge
 
 void FGameplayProvider::ReadViewTimeline(TFunctionRef<void(const IGameplayProvider::ViewTimeline&)> Callback) const
 {
-	Session.ReadAccessCheck();
+	ReadAccessCheck();
 
 	if (ViewTimeline.IsValid())
 	{
@@ -690,6 +701,7 @@ void FGameplayProvider::ReadViewTimeline(TFunctionRef<void(const IGameplayProvid
 
 double FGameplayProvider::GetRecordingDuration() const
 {
+	ReadAccessCheck();
 	return RecordingDuration;
 }
 
@@ -723,7 +735,7 @@ FSlateIcon FGameplayProvider::FindIconForClass(uint64 ClassId) const
 
 void FGameplayProvider::AppendView(uint64 InPlayerId, double InTime, const FVector& InPosition, const FRotator& InRotation, float InFov, float InAspectRatio)
 {
-	Session.WriteAccessCheck();
+	EditAccessCheck();
 
 	if (!ViewTimeline.IsValid())
 	{
@@ -741,12 +753,15 @@ void FGameplayProvider::AppendView(uint64 InPlayerId, double InTime, const FVect
 
 	ViewTimeline->AppendEvent(InTime, Message);
 
-	Session.UpdateDurationSeconds(InTime);
+	{
+		TraceServices::FAnalysisSessionEditScope SessionEditScope(Session);
+		Session.UpdateDurationSeconds(InTime);
+	}
 }
 
 void FGameplayProvider::AppendWorld(uint64 InWorldId, int32 InPIEInstanceId, uint8 InType, uint8 InNetMode, bool bInIsSimulating)
 {
-	Session.WriteAccessCheck();
+	EditAccessCheck();
 
 	bHasAnyData = true;
 
@@ -766,11 +781,13 @@ void FGameplayProvider::AppendWorld(uint64 InWorldId, int32 InPIEInstanceId, uin
 
 void FGameplayProvider::AppendRecordingInfo(uint64 InWorldId, double InProfileTime, uint32 InRecordingIndex, uint32 InFrameIndex, double InElapsedTime)
 {
-	Session.WriteAccessCheck();
+	EditAccessCheck();
 
 	if (!TargetType.IsSet())
 	{
 		TargetType = EBuildTargetType::Unknown;
+
+		TraceServices::FAnalysisSessionReadScope SessionReadScope(Session);
 		Session.EnumerateMetadata(
 			[&TargetType = TargetType](const TraceServices::FTraceSessionMetadata& Metadata)
 			{
@@ -817,7 +834,7 @@ void FGameplayProvider::AppendRecordingInfo(uint64 InWorldId, double InProfileTi
 
 const FGameplayProvider::RecordingInfoTimeline* FGameplayProvider::GetRecordingInfo(uint32 RecordingId) const
 {
-	Session.ReadAccessCheck();
+	ReadAccessCheck();
 
 	if (const TSharedRef<TraceServices::TPointTimeline<FRecordingInfoMessage>>* Recording = Recordings.Find(RecordingId))
 	{
@@ -829,7 +846,7 @@ const FGameplayProvider::RecordingInfoTimeline* FGameplayProvider::GetRecordingI
 
 void FGameplayProvider::AppendClassPropertyStringId(uint32 InStringId, const FStringView& InString)
 {
-	Session.WriteAccessCheck();
+	EditAccessCheck();
 
 	bHasAnyData = true;
 
@@ -840,7 +857,7 @@ void FGameplayProvider::AppendClassPropertyStringId(uint32 InStringId, const FSt
 
 void FGameplayProvider::AppendPropertiesStart(const uint64 InObjectId, double InTime, uint64 InEventId, double InRecordingTime)
 {
-	Session.WriteAccessCheck();
+	EditAccessCheck();
 
 	bHasAnyData = true;
 	bHasObjectProperties = true;
@@ -871,7 +888,7 @@ void FGameplayProvider::AppendPropertiesStart(const uint64 InObjectId, double In
 
 void FGameplayProvider::AppendPropertiesEnd(const uint64 InObjectId, double InTime)
 {
-	Session.WriteAccessCheck();
+	EditAccessCheck();
 
 	bHasAnyData = true;
 	bHasObjectProperties = true;
@@ -901,7 +918,7 @@ void FGameplayProvider::AppendPropertiesEnd(const uint64 InObjectId, double InTi
 
 void FGameplayProvider::AppendPropertyValue(const uint64 InObjectId, double InTime, uint64 InEventId, int32 InParentId, uint32 InTypeStringId, uint32 InNameId, uint32 InParentNameId, const FStringView& InValue)
 {
-	Session.WriteAccessCheck();
+	EditAccessCheck();
 
 	bHasAnyData = true;
 	bHasObjectProperties = true;
@@ -936,14 +953,14 @@ void FGameplayProvider::AppendPropertyValue(const uint64 InObjectId, double InTi
 
 bool FGameplayProvider::HasAnyData() const
 {
-	Session.ReadAccessCheck();
+	ReadAccessCheck();
 
 	return bHasAnyData;
 }
 
 bool FGameplayProvider::HasObjectProperties() const
 {
-	Session.ReadAccessCheck();
+	ReadAccessCheck();
 
 	return bHasObjectProperties;
 }
diff --git a/Engine/Plugins/Animation/GameplayInsights/Source/GameplayInsights/Private/GameplayProvider.h b/Engine/Plugins/Animation/GameplayInsights/Source/GameplayInsights/Private/GameplayProvider.h
--- a/Engine/Plugins/Animation/GameplayInsights/Source/GameplayInsights/Private/GameplayProvider.h
+++ b/Engine/Plugins/Animation/GameplayInsights/Source/GameplayInsights/Private/GameplayProvider.h
@@ -1,8 +1,9 @@
-﻿// Copyright Epic Games, Inc. All Rights Reserved.
+// Copyright Epic Games, Inc. All Rights Reserved.
 
 #pragma once
 
 #include "IGameplayProvider.h"
+#include "Common/ProviderLock.h"
 #include "Containers/Map.h"
 #include "Containers/StringView.h"
 #include "Math/TransformVectorized.h"
@@ -18,6 +19,7 @@ struct FObjectId;
 namespace TraceServices
 {
 class IAnalysisSession;
+extern thread_local FProviderLock::FThreadLocalState GGameplayProviderLockState;
 }
 
 class FGameplayProvider : public IGameplayProvider
@@ -25,9 +27,41 @@ class FGameplayProvider : public IGameplayProvider
 public:
 	static FName ProviderName;
 
-	FGameplayProvider(TraceServices::IAnalysisSession& InSession);
+	explicit FGameplayProvider(TraceServices::IAnalysisSession& InSession);
 
-	/** IGameplayProvider interface */
+	//~ IProvider
+	virtual void BeginRead() const override
+	{
+		Lock.BeginRead(TraceServices::GGameplayProviderLockState);
+	}
+
+	virtual void EndRead() const override
+	{
+		Lock.EndRead(TraceServices::GGameplayProviderLockState);
+	}
+
+	virtual void ReadAccessCheck() const override
+	{
+		Lock.ReadAccessCheck(TraceServices::GGameplayProviderLockState);
+	}
+
+	//~ IEditableProvider
+	virtual void BeginEdit() const override
+	{
+		Lock.BeginWrite(TraceServices::GGameplayProviderLockState);
+	}
+
+	virtual void EndEdit() const override
+	{
+		Lock.EndWrite(TraceServices::GGameplayProviderLockState);
+	}
+
+	virtual void EditAccessCheck() const override
+	{
+		Lock.WriteAccessCheck(TraceServices::GGameplayProviderLockState);
+	}
+
+	//~ IGameplayProvider
 	virtual bool ReadObjectEventsTimeline(const RewindDebugger::FObjectId& InObjectId, TFunctionRef<void(const ObjectEventsTimeline&)> Callback) const override;
 	virtual bool ReadObjectEvent(const RewindDebugger::FObjectId& InObjectId, uint64 InMessageId, TFunctionRef<void(const FObjectEventMessage&)> Callback) const override;
 	virtual bool ReadObjectPropertiesTimeline(uint64 InObjectId, TFunctionRef<void(const ObjectPropertiesTimeline&)> Callback) const override;
@@ -122,6 +156,7 @@ public:
 	virtual TRange<double> GetObjectRecordingLifetime(const RewindDebugger::FObjectId& InObjectId) const override;
 
 private:
+	mutable TraceServices::FProviderLock Lock;
 	TraceServices::IAnalysisSession& Session;
 
 	/** All class info, grow only for stable indices */
diff --git a/Engine/Plugins/Animation/GameplayInsights/Source/GameplayInsights/Private/GameplaySharedData.cpp b/Engine/Plugins/Animation/GameplayInsights/Source/GameplayInsights/Private/GameplaySharedData.cpp
--- a/Engine/Plugins/Animation/GameplayInsights/Source/GameplayInsights/Private/GameplaySharedData.cpp
+++ b/Engine/Plugins/Animation/GameplayInsights/Source/GameplayInsights/Private/GameplaySharedData.cpp
@@ -3,6 +3,7 @@
 #include "GameplaySharedData.h"
 
 #include "Algo/Sort.h"
+#include "Common/ProviderLock.h"
 #include "Framework/MultiBox/MultiBoxBuilder.h"
 #include "Modules/ModuleManager.h"
 #include "Widgets/Docking/SDockTab.h"
@@ -152,7 +153,7 @@ void FGameplaySharedData::Tick(UE::Insights::Timing::ITimingViewSession& InTimin
 
 	if(GameplayProvider)
 	{
-		TraceServices::FAnalysisSessionReadScope SessionReadScope(GetAnalysisSession());
+		TraceServices::FProviderReadScopeLock ProviderReadScope(*GameplayProvider);
 
 		if(GameplayProvider->HasAnyData())
 		{
@@ -163,7 +164,7 @@ void FGameplaySharedData::Tick(UE::Insights::Timing::ITimingViewSession& InTimin
 				{
 					TSharedPtr<FObjectEventsTrack> ObjectEventsTrack;
 
-					GameplayProvider->ReadObjectEventsTimeline(InObjectInfo.GetId(), [this, &InTimingViewSession, &InAnalysisSession, &InObjectInfo, &GameplayProvider, &ObjectEventsTrack](const IGameplayProvider::ObjectEventsTimeline& InTimeline)
+					GameplayProvider->ReadObjectEventsTimeline(InObjectInfo.GetId(), [this, &InTimingViewSession, &InAnalysisSession, &InObjectInfo, &ObjectEventsTrack](const IGameplayProvider::ObjectEventsTimeline& InTimeline)
 					{
 						if(bObjectTracksEnabled)
 						{
@@ -171,7 +172,7 @@ void FGameplaySharedData::Tick(UE::Insights::Timing::ITimingViewSession& InTimin
 						}
 					});
 
-					GameplayProvider->ReadObjectPropertiesTimeline(InObjectInfo.GetUObjectId(), [this, &InObjectInfo, &ObjectEventsTrack, &InTimingViewSession, &InAnalysisSession, &GameplayProvider](const IGameplayProvider::ObjectPropertiesTimeline& InTimeline)
+					GameplayProvider->ReadObjectPropertiesTimeline(InObjectInfo.GetUObjectId(), [this, &InObjectInfo, &ObjectEventsTrack, &InTimingViewSession, &InAnalysisSession](const IGameplayProvider::ObjectPropertiesTimeline& InTimeline)
 					{
 						if(!ObjectEventsTrack.IsValid())
 						{
