[RewindDebugger]

diff --git a/Engine/Plugins/Animation/GameplayInsights/Source/GameplayInsights/Private/GameplayTraceModule.cpp b/Engine/Plugins/Animation/GameplayInsights/Source/GameplayInsights/Private/GameplayTraceModule.cpp
--- a/Engine/Plugins/Animation/GameplayInsights/Source/GameplayInsights/Private/GameplayTraceModule.cpp
+++ b/Engine/Plugins/Animation/GameplayInsights/Source/GameplayInsights/Private/GameplayTraceModule.cpp
@@ -17,7 +17,7 @@ void FGameplayTraceModule::GetModuleInfo(TraceServices::FModuleInfo& OutModuleIn
 void FGameplayTraceModule::OnAnalysisBegin(TraceServices::IAnalysisSession& InSession)
 {
 	TSharedPtr<FGameplayProvider> GameplayProvider = MakeShared<FGameplayProvider>(InSession);
-	InSession.AddProvider(FGameplayProvider::ProviderName, GameplayProvider);
+	InSession.AddProvider(FGameplayProvider::ProviderName, GameplayProvider, GameplayProvider);
 	TSharedPtr<FAnimationProvider> AnimationProvider = MakeShared<FAnimationProvider>(InSession, *GameplayProvider);
 	InSession.AddProvider(FAnimationProvider::ProviderName, AnimationProvider);
 
diff --git a/Engine/Plugins/Animation/GameplayInsights/Source/GameplayInsights/Private/GameplayTrack.cpp b/Engine/Plugins/Animation/GameplayInsights/Source/GameplayInsights/Private/GameplayTrack.cpp
--- a/Engine/Plugins/Animation/GameplayInsights/Source/GameplayInsights/Private/GameplayTrack.cpp
+++ b/Engine/Plugins/Animation/GameplayInsights/Source/GameplayInsights/Private/GameplayTrack.cpp
@@ -3,6 +3,7 @@
 #include "GameplayTrack.h"
 
 #include "Algo/Sort.h"
+#include "Common/ProviderLock.h"
 #include "Fonts/FontMeasure.h"
 #include "Framework/Application/SlateApplication.h"
 #include "Framework/MultiBox/MultiBoxBuilder.h"
@@ -146,7 +147,7 @@ FText FGameplayTrack::GetWorldName(const TraceServices::IAnalysisSession& InAnal
 	const FGameplayProvider* GameplayProvider = InAnalysisSession.ReadProvider<FGameplayProvider>(FGameplayProvider::ProviderName);
 	if (GameplayProvider)
 	{
-		TraceServices::FAnalysisSessionReadScope SessionReadScope(InAnalysisSession);
+		TraceServices::FProviderReadScopeLock ProviderReadScope(*GameplayProvider);
 
 		if(const FWorldInfo* WorldInfo = GameplayProvider->FindWorldInfoFromObject(ObjectId))
 		{
diff --git a/Engine/Plugins/Animation/GameplayInsights/Source/GameplayInsights/Private/InertializationsTrack.cpp b/Engine/Plugins/Animation/GameplayInsights/Source/GameplayInsights/Private/InertializationsTrack.cpp
--- a/Engine/Plugins/Animation/GameplayInsights/Source/GameplayInsights/Private/InertializationsTrack.cpp
+++ b/Engine/Plugins/Animation/GameplayInsights/Source/GameplayInsights/Private/InertializationsTrack.cpp
@@ -1,6 +1,7 @@
 // Copyright Epic Games, Inc. All Rights Reserved.
 
 #include "InertializationsTrack.h"
+#include "Common/ProviderLock.h"
 #include "SCurveTimelineView.h"
 #include "IRewindDebugger.h"
 #include "GameplayProvider.h"
@@ -47,15 +48,13 @@ bool FInertializationsTrack::UpdateInternal()
 	double EndTime = TraceTimeRange.GetUpperBoundValue();
 	
 	// count number of unique inertialization nodes in the current time range
-	const FGameplayProvider* GameplayProvider = AnalysisSession->ReadProvider<FGameplayProvider>(FGameplayProvider::ProviderName);
 	const FAnimationProvider* AnimationProvider = AnalysisSession->ReadProvider<FAnimationProvider>(FAnimationProvider::ProviderName);
 
-	if (GameplayProvider && AnimationProvider)
+	if (AnimationProvider)
 	{
-		TraceServices::FAnalysisSessionReadScope SessionReadScope(*AnalysisSession);
-
 		TArray<TPair<int32,EInertializationType>, TInlineAllocator<8>> Nodes;
 		{
+			TraceServices::FProviderReadScopeLock AnimationProviderReadScope(*AnimationProvider);
 			AnimationProvider->EnumerateInertializationNodes(ObjectId,  [&Nodes] (uint32 NodeId, EInertializationType Type)
 			{
 				Nodes.Add({ NodeId, Type });
@@ -126,16 +125,14 @@ bool FInertializationTrack::UpdateInternal()
 	double EndTime = TraceTimeRange.GetUpperBoundValue();
 
 	const TraceServices::IAnalysisSession* AnalysisSession = RewindDebugger->GetAnalysisSession();
-	const FGameplayProvider* GameplayProvider = AnalysisSession->ReadProvider<FGameplayProvider>(FGameplayProvider::ProviderName);
 	const FAnimationProvider* AnimationProvider = AnalysisSession->ReadProvider<FAnimationProvider>(FAnimationProvider::ProviderName);
 
-	if (CurvesUpdateRequested > 10 && GameplayProvider && AnimationProvider)
+	if (CurvesUpdateRequested > 10 && AnimationProvider)
 	{
 		auto& CurvePoints = CurveData->Points;
 		CurvePoints.SetNum(0,EAllowShrinking::No);
 		
-		TraceServices::FAnalysisSessionReadScope SessionReadScope(*AnalysisSession);
-		
+		TraceServices::FProviderReadScopeLock AnimationProviderReadScope(*AnimationProvider);
 		AnimationProvider->ReadInertializationTimeline(ObjectId, [AnalysisSession, StartTime, EndTime, &CurvePoints, this](const IAnimationProvider::InertializationTimeline& InTimeline)
 		{
 			InTimeline.EnumerateEvents(StartTime, EndTime, [&CurvePoints, AnalysisSession, this](double InStartTime, double InEndTime, uint32 InDepth, const FInertializationMessage& InMessage)
@@ -202,40 +199,47 @@ bool FInertializationTrack::HandleDoubleClickInternal()
 	IRewindDebugger* RewindDebugger = IRewindDebugger::Instance();
 	if (const TraceServices::IAnalysisSession* AnalysisSession = RewindDebugger->GetAnalysisSession())
 	{
-		TraceServices::FAnalysisSessionReadScope SessionReadScope(*AnalysisSession);
-
 		if (const FGameplayProvider* GameplayProvider = AnalysisSession->ReadProvider<FGameplayProvider>(FGameplayProvider::ProviderName))
 		{
-			if (const FObjectInfo* AnimInstanceInfo = GameplayProvider->FindObjectInfo(ObjectId))
+			FString AnimInstanceClassPathName;
 			{
-				if (const FClassInfo* AnimInstanceClassInfo = GameplayProvider->FindClassInfo(AnimInstanceInfo->ClassId))
+				TraceServices::FProviderReadScopeLock ProviderReadScope(*GameplayProvider);
+				if (const FObjectInfo* AnimInstanceInfo = GameplayProvider->FindObjectInfo(ObjectId))
 				{
-					TSoftObjectPtr<UAnimBlueprintGeneratedClass> InstanceClass;
-					InstanceClass = FSoftObjectPath(AnimInstanceClassInfo->PathName);
+					if (const FClassInfo* AnimInstanceClassInfo = GameplayProvider->FindClassInfo(AnimInstanceInfo->ClassId))
+					{
+						AnimInstanceClassPathName = AnimInstanceClassInfo->PathName;
+					}
+				}
+			}
 
-					if (InstanceClass.LoadSynchronous())
+			if (!AnimInstanceClassPathName.IsEmpty())
+			{
+				TSoftObjectPtr<UAnimBlueprintGeneratedClass> InstanceClass;
+				InstanceClass = FSoftObjectPath(AnimInstanceClassPathName);
+
+				if (InstanceClass.LoadSynchronous())
+				{
+					if (UAnimBlueprint* AnimBlueprint = Cast<UAnimBlueprint>(InstanceClass.Get()->ClassGeneratedBy))
 					{
-						if (UAnimBlueprint* AnimBlueprint = Cast<UAnimBlueprint>(InstanceClass.Get()->ClassGeneratedBy))
-						{
-							GEditor->GetEditorSubsystem<UAssetEditorSubsystem>()->OpenEditorForAsset(AnimBlueprint);
+						GEditor->GetEditorSubsystem<UAssetEditorSubsystem>()->OpenEditorForAsset(AnimBlueprint);
 
-							if (UObject* SelectedInstance = FObjectTrace::GetObjectFromId(ObjectId))
-							{
-								AnimBlueprint->SetObjectBeingDebugged(SelectedInstance);
-							}
+						if (UObject* SelectedInstance = FObjectTrace::GetObjectFromId(ObjectId))
+						{
+							AnimBlueprint->SetObjectBeingDebugged(SelectedInstance);
+						}
 
-							if (IAnimationBlueprintEditor* AnimBlueprintEditor = static_cast<IAnimationBlueprintEditor*>(GEditor->GetEditorSubsystem<UAssetEditorSubsystem>()->FindEditorForAsset(AnimBlueprint, true)))
+						if (IAnimationBlueprintEditor* AnimBlueprintEditor = static_cast<IAnimationBlueprintEditor*>(GEditor->GetEditorSubsystem<UAssetEditorSubsystem>()->FindEditorForAsset(AnimBlueprint, true)))
+						{
+							int32 AnimNodeIndex = InstanceClass.Get()->GetAnimNodeProperties().Num() - NodeId - 1;
+							TWeakObjectPtr<const UEdGraphNode>* GraphNode = InstanceClass.Get()->AnimBlueprintDebugData.NodePropertyIndexToNodeMap.Find(AnimNodeIndex);
+							if (GraphNode != nullptr && GraphNode->Get())
 							{
-								int32 AnimNodeIndex = InstanceClass.Get()->GetAnimNodeProperties().Num() - NodeId - 1;
-								TWeakObjectPtr<const UEdGraphNode>* GraphNode = InstanceClass.Get()->AnimBlueprintDebugData.NodePropertyIndexToNodeMap.Find(AnimNodeIndex);
-								if (GraphNode != nullptr && GraphNode->Get())
-								{
-									AnimBlueprintEditor->JumpToHyperlink(GraphNode->Get());
-								}
+								AnimBlueprintEditor->JumpToHyperlink(GraphNode->Get());
 							}
-
-							return true;
 						}
+
+						return true;
 					}
 				}
 			}
@@ -272,10 +276,9 @@ bool FInertializationsTrackCreator::HasDebugInfoInternal(const FObjectId& InObje
 {
 	TRACE_CPUPROFILER_EVENT_SCOPE(FInertializationsTrack::HasDebugInfoInternal);
 	const TraceServices::IAnalysisSession* AnalysisSession = IRewindDebugger::Instance()->GetAnalysisSession();
-	TraceServices::FAnalysisSessionReadScope SessionReadScope(*AnalysisSession);
-	
 	if (const FAnimationProvider* AnimationProvider = AnalysisSession->ReadProvider<FAnimationProvider>(FAnimationProvider::ProviderName))
 	{
+		TraceServices::FProviderReadScopeLock ProviderReadScope(*AnimationProvider);
 		if (AnimationProvider->ReadInertializationTimeline(InObjectId.GetMainId(), [](const IAnimationProvider::InertializationTimeline& InTimeline){}))
 		{
 			return true;
diff --git a/Engine/Plugins/Animation/GameplayInsights/Source/GameplayInsights/Private/MontageTrack.cpp b/Engine/Plugins/Animation/GameplayInsights/Source/GameplayInsights/Private/MontageTrack.cpp
--- a/Engine/Plugins/Animation/GameplayInsights/Source/GameplayInsights/Private/MontageTrack.cpp
+++ b/Engine/Plugins/Animation/GameplayInsights/Source/GameplayInsights/Private/MontageTrack.cpp
@@ -1,6 +1,7 @@
 // Copyright Epic Games, Inc. All Rights Reserved.
 
 #include "MontageTrack.h"
+#include "Common/ProviderLock.h"
 #include "GameplayProvider.h"
 #include "AnimationProvider.h"
 #include "Insights/ViewModels/TimingTrackViewport.h"
@@ -42,7 +43,8 @@ void FMontageTrack::AddAllSeries()
 
 	if(AnimationProvider && GameplayProvider)
 	{
-		TraceServices::FAnalysisSessionReadScope SessionReadScope(SharedData.GetAnalysisSession());
+		TraceServices::FProviderReadScopeLock AnimationProviderReadScope(*AnimationProvider);
+		TraceServices::FProviderReadScopeLock GameplayProviderReadScope(*GameplayProvider);
 
 		AnimationProvider->EnumerateMontageIds(GetGameplayTrack().GetObjectId(), [this, &AnimationProvider, &GameplayProvider](uint64 InMontageId)
 		{
@@ -113,11 +115,10 @@ bool FMontageTrack::UpdateSeriesBoundsHelper(FMontageSeries& InSeries, const FTi
 
 	if (AnimationProvider)
 	{
-		TraceServices::FAnalysisSessionReadScope SessionReadScope(SharedData.GetAnalysisSession());
-
 		InSeries.CurrentMin = 0.0;
 		InSeries.CurrentMax = 0.0;
 
+		TraceServices::FProviderReadScopeLock ProviderReadScope(*AnimationProvider);
 		AnimationProvider->ReadMontageTimeline(GetGameplayTrack().GetObjectId(), [&bFoundEvents, &InViewport, &InSeries, &Projection](const FAnimationProvider::AnimMontageTimeline& InTimeline)
 		{
 			InTimeline.EnumerateEvents(InViewport.GetStartTime(), InViewport.GetEndTime(), [&bFoundEvents, &InSeries, &Projection](double InStartTime, double InEndTime, uint32 InDepth, const FAnimMontageMessage& InMessage)
@@ -158,10 +159,9 @@ void FMontageTrack::UpdateSeriesHelper(FMontageSeries& InSeries, const FTimingTr
 
 	if (AnimationProvider)
 	{
-		TraceServices::FAnalysisSessionReadScope SessionReadScope(SharedData.GetAnalysisSession());
-
 		FGraphTrackBuilder Builder(*this, InSeries, InViewport);
 
+		TraceServices::FProviderReadScopeLock ProviderReadScope(*AnimationProvider);
 		AnimationProvider->ReadMontageTimeline(GetGameplayTrack().GetObjectId(), [this, &AnimationProvider, &Builder, &InViewport, &InSeries, &Projection](const FAnimationProvider::AnimMontageTimeline& InTimeline)
 		{
 			uint16 FrameCounter = 0;
@@ -214,24 +214,22 @@ void FMontageTrack::InitTooltip(FTooltipDrawState& Tooltip, const ITimingEvent&
 		Tooltip.AddNameValueTextLine(LOCTEXT("EventDesiredWeight", "Desired Weight").ToString(), FText::AsNumber(InMessage.DesiredWeight).ToString());
 		Tooltip.AddNameValueTextLine(LOCTEXT("EventPosition", "Position").ToString(), FText::AsNumber(InMessage.Position).ToString()); 
 
+		const FGameplayProvider* GameplayProvider = SharedData.GetAnalysisSession().ReadProvider<FGameplayProvider>(FGameplayProvider::ProviderName);
+		if(GameplayProvider)
 		{
-			TraceServices::FAnalysisSessionReadScope SessionReadScope(SharedData.GetAnalysisSession());
-
-			const FGameplayProvider* GameplayProvider = SharedData.GetAnalysisSession().ReadProvider<FGameplayProvider>(FGameplayProvider::ProviderName);
-			if(GameplayProvider)
-			{
-				const FObjectInfo& MontageInfo = GameplayProvider->GetObjectInfo(InMessage.MontageId);
-				Tooltip.AddNameValueTextLine(LOCTEXT("MontageName", "Montage").ToString(), MontageInfo.PathName);
-			}
+			TraceServices::FProviderReadScopeLock ProviderReadScope(*GameplayProvider);
+			const FObjectInfo& MontageInfo = GameplayProvider->GetObjectInfo(InMessage.MontageId);
+			Tooltip.AddNameValueTextLine(LOCTEXT("MontageName", "Montage").ToString(), MontageInfo.PathName);
+		}
 
-			const FAnimationProvider* AnimationProvider = SharedData.GetAnalysisSession().ReadProvider<FAnimationProvider>(FAnimationProvider::ProviderName);
-			if(AnimationProvider)
-			{
-				const TCHAR* CurrentSectionName = AnimationProvider->GetName(InMessage.CurrentSectionNameId);
-				Tooltip.AddNameValueTextLine(LOCTEXT("CurrentSectionName", "Current Section").ToString(), FText::FromString(CurrentSectionName).ToString());
-				const TCHAR* NextSectionName = AnimationProvider->GetName(InMessage.NextSectionNameId);
-				Tooltip.AddNameValueTextLine(LOCTEXT("NextSectionName", "Next Section").ToString(), FText::FromString(NextSectionName).ToString());
-			}
+		const FAnimationProvider* AnimationProvider = SharedData.GetAnalysisSession().ReadProvider<FAnimationProvider>(FAnimationProvider::ProviderName);
+		if (AnimationProvider)
+		{
+			TraceServices::FProviderReadScopeLock ProviderReadScope(*AnimationProvider);
+			const TCHAR* CurrentSectionName = AnimationProvider->GetName(InMessage.CurrentSectionNameId);
+			Tooltip.AddNameValueTextLine(LOCTEXT("CurrentSectionName", "Current Section").ToString(), FText::FromString(CurrentSectionName).ToString());
+			const TCHAR* NextSectionName = AnimationProvider->GetName(InMessage.NextSectionNameId);
+			Tooltip.AddNameValueTextLine(LOCTEXT("NextSectionName", "Next Section").ToString(), FText::FromString(NextSectionName).ToString());
 		}
 
 		Tooltip.AddNameValueTextLine(LOCTEXT("EventWorld", "World").ToString(), GetGameplayTrack().GetWorldName(SharedData.GetAnalysisSession()).ToString());
@@ -263,7 +261,7 @@ void FMontageTrack::FindMontageMessage(const FTimingEventSearchParameters& InPar
 
 			if(AnimationProvider)
 			{
-				TraceServices::FAnalysisSessionReadScope SessionReadScope(SharedData.GetAnalysisSession());
+				TraceServices::FProviderReadScopeLock ProviderReadScope(*AnimationProvider);
 
 				AnimationProvider->ReadMontageTimeline(GetGameplayTrack().GetObjectId(), [&InContext](const FAnimationProvider::AnimMontageTimeline& InTimeline)
 				{
@@ -300,7 +298,8 @@ void FMontageTrack::GetVariantsAtFrame(const TraceServices::FFrame& InFrame, TAr
 
 	if(AnimationProvider && GameplayProvider)
 	{
-		TraceServices::FAnalysisSessionReadScope SessionReadScope(SharedData.GetAnalysisSession());
+		TraceServices::FProviderReadScopeLock AnimationProviderReadScope(*AnimationProvider);
+		TraceServices::FProviderReadScopeLock GameplayProviderReadScope(*GameplayProvider);
 
 		AnimationProvider->ReadMontageTimeline(GetGameplayTrack().GetObjectId(), [this, &InFrame, &AnimationProvider, &GameplayProvider, &Header](const FAnimationProvider::AnimMontageTimeline& InTimeline)
 		{
diff --git a/Engine/Plugins/Animation/GameplayInsights/Source/GameplayInsights/Private/MontagesTrack.cpp b/Engine/Plugins/Animation/GameplayInsights/Source/GameplayInsights/Private/MontagesTrack.cpp
--- a/Engine/Plugins/Animation/GameplayInsights/Source/GameplayInsights/Private/MontagesTrack.cpp
+++ b/Engine/Plugins/Animation/GameplayInsights/Source/GameplayInsights/Private/MontagesTrack.cpp
@@ -1,6 +1,7 @@
 // Copyright Epic Games, Inc. All Rights Reserved.
 
 #include "MontagesTrack.h"
+#include "Common/ProviderLock.h"
 #include "SCurveTimelineView.h"
 #include "IRewindDebugger.h"
 #include "GameplayProvider.h"
@@ -50,17 +51,16 @@ bool FMontagesTrack::UpdateInternal()
 	double EndTime = TraceTimeRange.GetUpperBoundValue();
 	
 	// count number of unique animations in the current time range
-	const FGameplayProvider* GameplayProvider = AnalysisSession->ReadProvider<FGameplayProvider>(FGameplayProvider::ProviderName);
 	const FAnimationProvider* AnimationProvider = AnalysisSession->ReadProvider<FAnimationProvider>(FAnimationProvider::ProviderName);
 
 	int AnimationCount = 0;
 
-	if(GameplayProvider && AnimationProvider)
+	if (AnimationProvider)
 	{
-		TraceServices::FAnalysisSessionReadScope SessionReadScope(*AnalysisSession);
+		TraceServices::FProviderReadScopeLock AnimationProviderReadScope(*AnimationProvider);
 		UniqueTrackIds.SetNum(0, EAllowShrinking::No);
 
-		AnimationProvider->ReadMontageTimeline(ObjectId, [&UniqueTrackIds,&GameplayProvider, StartTime, EndTime](const FAnimationProvider::AnimMontageTimeline& InTimeline)
+		AnimationProvider->ReadMontageTimeline(ObjectId, [&UniqueTrackIds, StartTime, EndTime](const FAnimationProvider::AnimMontageTimeline& InTimeline)
 		{
 			// this isn't very efficient, and it gets called every frame.  will need optimizing
 			InTimeline.EnumerateEvents(StartTime, EndTime, [&UniqueTrackIds, StartTime, EndTime](double InStartTime, double InEndTime, uint32 InDepth, const FAnimMontageMessage& InMessage)
@@ -152,8 +152,7 @@ bool FMontageTrack::UpdateInternal()
 		auto& CurvePoints = CurveData->Points;
 		CurvePoints.SetNum(0,EAllowShrinking::No);
 		
-		TraceServices::FAnalysisSessionReadScope SessionReadScope(*AnalysisSession);
-		
+		TraceServices::FProviderReadScopeLock AnimationProviderReadScope(*AnimationProvider);
 		AnimationProvider->ReadMontageTimeline(ObjectId, [AnalysisSession, StartTime, EndTime, &CurvePoints, this](const FAnimationProvider::AnimMontageTimeline& InTimeline)
 		{
 			InTimeline.EnumerateEvents(StartTime, EndTime, [&CurvePoints, AnalysisSession, this](double InStartTime, double InEndTime, uint32 InDepth, const FAnimMontageMessage& InMessage)
@@ -187,6 +186,7 @@ bool FMontageTrack::UpdateInternal()
 		// Blend Weight track gets name/icon of animation
 		if (CurveName.IsEmpty() && GameplayProvider)
 		{
+			TraceServices::FProviderReadScopeLock GameplayProviderReadScope(*GameplayProvider);
 			if (const FObjectInfo* ObjectInfo = GameplayProvider->FindObjectInfo(AssetId))
 			{
 				CurveName = FText::FromString(ObjectInfo->Name);
@@ -294,26 +294,40 @@ bool FMontageTrack::HandleDoubleClickInternal()
 	IRewindDebugger* RewindDebugger = IRewindDebugger::Instance();
 	if (const TraceServices::IAnalysisSession* AnalysisSession = RewindDebugger->GetAnalysisSession())
 	{
-		TraceServices::FAnalysisSessionReadScope SessionReadScope(*AnalysisSession);
-
 		const FGameplayProvider* GameplayProvider = AnalysisSession->ReadProvider<FGameplayProvider>(FGameplayProvider::ProviderName);
+		const FAnimationProvider* AnimationProvider = AnalysisSession->ReadProvider<FAnimationProvider>(FAnimationProvider::ProviderName);
+
+		if (!GameplayProvider || !AnimationProvider)
+		{
+			return false;
+		}
 
-		const FObjectInfo& AssetInfo = GameplayProvider->GetObjectInfo(GetAssetId());
-	
-		UObject* Asset = nullptr;
 		bool bMessageFound = false;
 		float PlaybackTime = 0;
 
-		float CurrentTraceTime = IRewindDebugger::Instance()->CurrentTraceTime();
-		
-		const FAnimationProvider* AnimationProvider = AnalysisSession->ReadProvider<FAnimationProvider>(FAnimationProvider::ProviderName);
-		const TraceServices::IFrameProvider& FrameProvider = TraceServices::ReadFrameProvider(*AnalysisSession);
+		FString AssetPathName;
+		{
+			TraceServices::FProviderReadScopeLock GameplayProviderReadScope(*GameplayProvider);
+			const FObjectInfo& AssetInfo = GameplayProvider->GetObjectInfo(GetAssetId());
+			AssetPathName = AssetInfo.PathName;
+		}
+
+		const float CurrentTraceTime = RewindDebugger->CurrentTraceTime();
+
+		bool bFrameFound;
 		TraceServices::FFrame Frame;
-		if (FrameProvider.GetFrameFromTime(ETraceFrameType::TraceFrameType_Game, CurrentTraceTime, Frame))
 		{
-			AnimationProvider->ReadMontageTimeline(ObjectId, [this, &bMessageFound, &PlaybackTime, &GameplayProvider, &Frame](const FAnimationProvider::AnimMontageTimeline& InTimeline)
+			TraceServices::FAnalysisSessionReadScope SessionReadScope(*AnalysisSession);
+			const TraceServices::IFrameProvider& FrameProvider = TraceServices::ReadFrameProvider(*AnalysisSession);
+			bFrameFound = FrameProvider.GetFrameFromTime(TraceFrameType_Game, CurrentTraceTime, Frame);
+		}
+
+		if (bFrameFound)
+		{
+			TraceServices::FProviderReadScopeLock AnimationProviderReadScope(*AnimationProvider);
+			AnimationProvider->ReadMontageTimeline(ObjectId, [this, &bMessageFound, &PlaybackTime, &Frame](const FAnimationProvider::AnimMontageTimeline& InTimeline)
 			{
-				InTimeline.EnumerateEvents(Frame.StartTime, Frame.EndTime, [this, &bMessageFound, &PlaybackTime, &GameplayProvider, &Frame](double InStartTime, double InEndTime, uint32 InDepth, const FAnimMontageMessage& InMessage)
+				InTimeline.EnumerateEvents(Frame.StartTime, Frame.EndTime, [this, &bMessageFound, &PlaybackTime, &Frame](double InStartTime, double InEndTime, uint32 InDepth, const FAnimMontageMessage& InMessage)
 				{
 					if(InStartTime >= Frame.StartTime && InEndTime <= Frame.EndTime)
 					{
@@ -328,21 +342,22 @@ bool FMontageTrack::HandleDoubleClickInternal()
 				});
 			});
 		}
-		
-		FString PackagePathString = FPackageName::ObjectPathToPackageName(FString(AssetInfo.PathName));
+
+		UObject* Asset = nullptr;
+		FString PackagePathString = FPackageName::ObjectPathToPackageName(AssetPathName);
 
 		UPackage* Package = LoadPackage(NULL, ToCStr(PackagePathString), LOAD_NoRedirects);
 		if (Package)
 		{
 			Package->FullyLoad();
                 
-			FString AssetName = FPaths::GetBaseFilename(AssetInfo.PathName);
+			FString AssetName = FPaths::GetBaseFilename(AssetPathName);
 			Asset = FindObject<UObject>(Package, *AssetName);
 		}
 		else
 		{
 			// fallback for unsaved assets
-			Asset = FindObject<UObject>(nullptr, AssetInfo.PathName);
+			Asset = FindObject<UObject>(nullptr, *AssetPathName);
 		}
                     	
 		if (Asset != nullptr)
@@ -407,10 +422,10 @@ bool FMontagesTrackCreator::HasDebugInfoInternal(const FObjectId& InObjectId) co
 {
 	TRACE_CPUPROFILER_EVENT_SCOPE(FMontagesTrack::HasDebugInfoInternal);
 	const TraceServices::IAnalysisSession* AnalysisSession = IRewindDebugger::Instance()->GetAnalysisSession();
-	TraceServices::FAnalysisSessionReadScope SessionReadScope(*AnalysisSession);
 	bool bHasData = false;
 	if (const FAnimationProvider* AnimationProvider = AnalysisSession->ReadProvider<FAnimationProvider>(FAnimationProvider::ProviderName))
 	{
+		TraceServices::FProviderReadScopeLock ProviderReadScope(*AnimationProvider);
 		AnimationProvider->ReadMontageTimeline(InObjectId.GetMainId(), [&bHasData](const FAnimationProvider::AnimMontageTimeline& InTimeline)
 		{
 			bHasData = true;
diff --git a/Engine/Plugins/Animation/GameplayInsights/Source/GameplayInsights/Private/NotifiesTrack.cpp b/Engine/Plugins/Animation/GameplayInsights/Source/GameplayInsights/Private/NotifiesTrack.cpp
--- a/Engine/Plugins/Animation/GameplayInsights/Source/GameplayInsights/Private/NotifiesTrack.cpp
+++ b/Engine/Plugins/Animation/GameplayInsights/Source/GameplayInsights/Private/NotifiesTrack.cpp
@@ -1,6 +1,7 @@
 // Copyright Epic Games, Inc. All Rights Reserved.
 
 #include "NotifiesTrack.h"
+#include "Common/ProviderLock.h"
 #include "SCurveTimelineView.h"
 #include "IRewindDebugger.h"
 #include "GameplayProvider.h"
@@ -72,10 +73,9 @@ bool FNotifyTrack::UpdateInternal()
 	double EndTime = TraceTimeRange.GetUpperBoundValue();
 	
 	const TraceServices::IAnalysisSession* AnalysisSession = RewindDebugger->GetAnalysisSession();
-	const FGameplayProvider* GameplayProvider = AnalysisSession->ReadProvider<FGameplayProvider>(FGameplayProvider::ProviderName);
 	const FAnimationProvider* AnimationProvider = AnalysisSession->ReadProvider<FAnimationProvider>(FAnimationProvider::ProviderName);
 	
-	if(EventUpdateRequested > 10 && GameplayProvider && AnimationProvider)
+	if (EventUpdateRequested > 10 && AnimationProvider)
 	{
 		TRACE_CPUPROFILER_EVENT_SCOPE(FNotifiesTrack::UpdateEventPointsInternal);
 		EventUpdateRequested = 0;
@@ -84,7 +84,7 @@ bool FNotifyTrack::UpdateInternal()
 		EventPoints.SetNum(0,EAllowShrinking::No);
 		EventData->Windows.SetNum(0);
 
-		TraceServices::FAnalysisSessionReadScope SessionReadScope(*AnalysisSession);
+		TraceServices::FProviderReadScopeLock AnimationProviderReadScope(*AnimationProvider);
 
 		AnimationProvider->ReadNotifyTimeline(ObjectId, [this, StartTime, EndTime, AnimationProvider, AnalysisSession, &EventPoints](const FAnimationProvider::AnimNotifyTimeline& InTimeline)
 		{
@@ -193,16 +193,15 @@ bool FNotifiesTrack::UpdateInternal()
 	double EndTime = TraceTimeRange.GetUpperBoundValue();
 	
 	const TraceServices::IAnalysisSession* AnalysisSession = RewindDebugger->GetAnalysisSession();
-	const FGameplayProvider* GameplayProvider = AnalysisSession->ReadProvider<FGameplayProvider>(FGameplayProvider::ProviderName);
 	const FAnimationProvider* AnimationProvider = AnalysisSession->ReadProvider<FAnimationProvider>(FAnimationProvider::ProviderName);
 	
 	bool bChanged = false;
 	
-	if(GameplayProvider && AnimationProvider)
+	if (AnimationProvider)
 	{
 		TArray<FNotifyTrack::FNotifyTrackId> UniqueTrackIds;
 		
-		TraceServices::FAnalysisSessionReadScope SessionReadScope(*AnalysisSession);
+		TraceServices::FProviderReadScopeLock AnimationProviderReadScope(*AnimationProvider);
 
 		AnimationProvider->ReadNotifyTimeline(ObjectId, [this, StartTime, EndTime, AnimationProvider, AnalysisSession, &UniqueTrackIds](const FAnimationProvider::AnimNotifyTimeline& InTimeline)
 		{
@@ -267,11 +266,11 @@ bool FNotifiesTrackCreator::HasDebugInfoInternal(const FObjectId& InObjectId) co
 {
 	TRACE_CPUPROFILER_EVENT_SCOPE(FNotifiesTrack::HasDebugInfoInternal);
 	const TraceServices::IAnalysisSession* AnalysisSession = IRewindDebugger::Instance()->GetAnalysisSession();
-	
-	TraceServices::FAnalysisSessionReadScope SessionReadScope(*AnalysisSession);
+
 	bool bHasData = false;
 	if (const FAnimationProvider* AnimationProvider = AnalysisSession->ReadProvider<FAnimationProvider>(FAnimationProvider::ProviderName))
 	{
+		TraceServices::FProviderReadScopeLock ProviderReadScope(*AnimationProvider);
 		AnimationProvider->ReadNotifyTimeline(InObjectId.GetMainId(), [&bHasData](const FAnimationProvider::AnimNotifyTimeline& InTimeline)
 		{
 			bHasData = true;
diff --git a/Engine/Plugins/Animation/GameplayInsights/Source/GameplayInsights/Private/ObjectAsTraceIdProxyArchiveReader.cpp b/Engine/Plugins/Animation/GameplayInsights/Source/GameplayInsights/Private/ObjectAsTraceIdProxyArchiveReader.cpp
--- a/Engine/Plugins/Animation/GameplayInsights/Source/GameplayInsights/Private/ObjectAsTraceIdProxyArchiveReader.cpp
+++ b/Engine/Plugins/Animation/GameplayInsights/Source/GameplayInsights/Private/ObjectAsTraceIdProxyArchiveReader.cpp
@@ -15,6 +15,7 @@
 #include "Misc/PackageName.h"
 #include "ObjectTrace.h"
 #include "IGameplayProvider.h"
+#include "Common/ProviderLock.h"
 #include "UObject/ConstructorHelpers.h"
 
 
@@ -41,33 +42,46 @@ FArchive& FObjectAsTraceIdProxyArchiveReader::operator<<(UObject*& Obj)
 		uint64 Id = 0;
 		InnerArchive << Id;
 
-		
-		if (const FClassInfo* ClassInfo = GameplayProvider->FindClassInfo(Id))
+		const TCHAR* ClassPathName = nullptr;
+		const TCHAR* ObjectPathName = nullptr;
+		{
+			TraceServices::FProviderReadScopeLock ProviderReadScope(*GameplayProvider);
+			if (const FClassInfo* ClassInfo = GameplayProvider->FindClassInfo(Id))
+			{
+				ClassPathName = ClassInfo->PathName;
+			}
+			else if (const FObjectInfo* ObjectInfo = GameplayProvider->FindObjectInfo(Id))
+			{
+				ObjectPathName = ObjectInfo->PathName;
+			}
+		}
+
+		if (ClassPathName != nullptr)
 		{
 			TSoftObjectPtr<UClass> Class;
-			Class = FSoftObjectPath(ClassInfo->PathName);
+			Class = FSoftObjectPath(ClassPathName);
 
 			if (Class.LoadSynchronous())
 			{
 				Obj = Class.Get();
 			}
 		}
-		else if (const FObjectInfo* ObjectInfo = GameplayProvider->FindObjectInfo(Id))
+		else if (ObjectPathName != nullptr)
 		{
 			// look up the object by fully qualified pathname
-			Obj = FindObject<UObject>(nullptr, ObjectInfo->PathName);
+			Obj = FindObject<UObject>(nullptr, ObjectPathName);
 			// If we couldn't find it, and we want to load it, do that
-			if(!Obj)
+			if (!Obj)
 			{
-				Obj = LoadObject<UObject>(nullptr, ObjectInfo->PathName);
-			}			
+				Obj = LoadObject<UObject>(nullptr, ObjectPathName);
+			}
 		}
 		else
 		{
 			Obj = nullptr;
 		}
 	}
-	else 
+	else
 	{
 		Write(Obj);
 	}
diff --git a/Engine/Plugins/Animation/GameplayInsights/Source/GameplayInsights/Private/ObjectEventsTrack.cpp b/Engine/Plugins/Animation/GameplayInsights/Source/GameplayInsights/Private/ObjectEventsTrack.cpp
--- a/Engine/Plugins/Animation/GameplayInsights/Source/GameplayInsights/Private/ObjectEventsTrack.cpp
+++ b/Engine/Plugins/Animation/GameplayInsights/Source/GameplayInsights/Private/ObjectEventsTrack.cpp
@@ -1,6 +1,7 @@
 // Copyright Epic Games, Inc. All Rights Reserved.
 
 #include "ObjectEventsTrack.h"
+#include "Common/ProviderLock.h"
 #include "GameplayProvider.h"
 #include "Insights/ViewModels/TimingTrackViewport.h"
 #include "Insights/ViewModels/TimingEvent.h"
@@ -33,9 +34,7 @@ void FObjectEventsTrack::BuildDrawState(ITimingEventsTrackDrawStateBuilder& Buil
 
 	if(GameplayProvider)
 	{
-		TraceServices::FAnalysisSessionReadScope SessionReadScope(SharedData.GetAnalysisSession());
-
-		// object events
+		TraceServices::FProviderReadScopeLock ProviderReadScope(*GameplayProvider);
 		GameplayProvider->ReadObjectEventsTimeline(GetGameplayTrack().GetAssociatedObjectId(), [&Context, &Builder](const FGameplayProvider::ObjectEventsTimeline& InTimeline)
 		{
 			InTimeline.EnumerateEvents(Context.GetViewport().GetStartTime(), Context.GetViewport().GetEndTime(), [&Builder](double InStartTime, double InEndTime, uint32 InDepth, const FObjectEventMessage& InMessage)
@@ -94,8 +93,7 @@ void FObjectEventsTrack::FindObjectEvent(const FTimingEventSearchParameters& InP
 
 			if(GameplayProvider)
 			{
-				TraceServices::FAnalysisSessionReadScope SessionReadScope(SharedData.GetAnalysisSession());
-
+				TraceServices::FProviderReadScopeLock ProviderReadScope(*GameplayProvider);
 				GameplayProvider->ReadObjectEventsTimeline(GetGameplayTrack().GetAssociatedObjectId(), [&InContext](const FGameplayProvider::ObjectEventsTimeline& InTimeline)
 				{
 					InTimeline.EnumerateEvents(InContext.GetParameters().StartTime, InContext.GetParameters().EndTime, [&InContext](double InEventStartTime, double InEventEndTime, uint32 InDepth, const FObjectEventMessage& InMessage)
@@ -120,20 +118,24 @@ FText FObjectEventsTrack::MakeTrackName(const FGameplaySharedData& InSharedData,
 	FText ClassName = LOCTEXT("UnknownClass", "Unknown");
 
 	const FGameplayProvider* GameplayProvider = InSharedData.GetAnalysisSession().ReadProvider<FGameplayProvider>(FGameplayProvider::ProviderName);
-	if(GameplayProvider)
+	if (GameplayProvider)
 	{
-		TraceServices::FAnalysisSessionReadScope SessionReadScope(InSharedData.GetAnalysisSession());
-
-		if(const FObjectInfo* ObjectInfo = GameplayProvider->FindObjectInfo(InObjectID))
+		bool bIsWorld = false;
 		{
-			if(const FClassInfo* ClassInfo = GameplayProvider->FindClassInfo(ObjectInfo->ClassId))
+			TraceServices::FProviderReadScopeLock ProviderReadScope(*GameplayProvider);
+			if (const FObjectInfo* ObjectInfo = GameplayProvider->FindObjectInfo(InObjectID))
 			{
-				ClassName = FText::FromString(ClassInfo->Name);
+				if (const FClassInfo* ClassInfo = GameplayProvider->FindClassInfo(ObjectInfo->ClassId))
+				{
+					ClassName = FText::FromString(ClassInfo->Name);
+				}
 			}
+
+			bIsWorld = GameplayProvider->IsWorld(InObjectID);
 		}
 
 		FText ObjectName;
-		if (GameplayProvider->IsWorld(InObjectID))
+		if (bIsWorld)
 		{
 			ObjectName = GetGameplayTrack().GetWorldName(InSharedData.GetAnalysisSession());
 		}
@@ -151,13 +153,12 @@ FText FObjectEventsTrack::MakeTrackName(const FGameplaySharedData& InSharedData,
 void FObjectEventsTrack::GetVariantsAtFrame(const TraceServices::FFrame& InFrame, TArray<TSharedRef<FVariantTreeNode>>& OutVariants) const
 {
 	const FGameplayProvider* GameplayProvider = SharedData.GetAnalysisSession().ReadProvider<FGameplayProvider>(FGameplayProvider::ProviderName);
-	if(GameplayProvider)
+	if (GameplayProvider)
 	{
-		TraceServices::FAnalysisSessionReadScope SessionReadScope(SharedData.GetAnalysisSession());
-
 		TSharedRef<FVariantTreeNode> Header = OutVariants.Add_GetRef(FVariantTreeNode::MakeHeader(FText::FromString(GetName()), 0));
 
 		// object events
+		TraceServices::FProviderReadScopeLock ProviderReadScope(*GameplayProvider);
 		GameplayProvider->ReadObjectEventsTimeline(GetGameplayTrack().GetAssociatedObjectId(), [&InFrame, &Header](const FGameplayProvider::ObjectEventsTimeline& InTimeline)
 		{
 			InTimeline.EnumerateEvents(InFrame.StartTime, InFrame.EndTime, [&Header](double InStartTime, double InEndTime, uint32 InDepth, const FObjectEventMessage& InMessage)
@@ -181,11 +182,18 @@ void FObjectEventsTrack::BuildContextMenu(FMenuBuilder& MenuBuilder)
 		const FGameplayProvider* GameplayProvider = SharedData.GetAnalysisSession().ReadProvider<FGameplayProvider>(FGameplayProvider::ProviderName);
 		if (GameplayProvider)
 		{
-			TraceServices::FAnalysisSessionReadScope SessionReadScope(SharedData.GetAnalysisSession());
+			FString ObjectPathName;
+			{
+				TraceServices::FProviderReadScopeLock ProviderReadScope(*GameplayProvider);
+				if (const FObjectInfo* ObjectInfo = GameplayProvider->FindObjectInfo(GetGameplayTrack().GetObjectId()))
+				{
+					ObjectPathName = ObjectInfo->PathName;
+				}
+			}
 
-			if (const FObjectInfo* ObjectInfo = GameplayProvider->FindObjectInfo(GetGameplayTrack().GetObjectId()))
+			if (!ObjectPathName.IsEmpty())
 			{
-				WeakObject = FindObject<UObject>(nullptr, ObjectInfo->PathName);
+				WeakObject = FindObject<UObject>(nullptr, *ObjectPathName);
 			}
 		}
 
diff --git a/Engine/Plugins/Animation/GameplayInsights/Source/GameplayInsights/Private/ObjectPropertiesTrack.cpp b/Engine/Plugins/Animation/GameplayInsights/Source/GameplayInsights/Private/ObjectPropertiesTrack.cpp
--- a/Engine/Plugins/Animation/GameplayInsights/Source/GameplayInsights/Private/ObjectPropertiesTrack.cpp
+++ b/Engine/Plugins/Animation/GameplayInsights/Source/GameplayInsights/Private/ObjectPropertiesTrack.cpp
@@ -1,6 +1,7 @@
 // Copyright Epic Games, Inc. All Rights Reserved.
 
 #include "ObjectPropertiesTrack.h"
+#include "Common/ProviderLock.h"
 #include "GameplayProvider.h"
 #include "Insights/ViewModels/TimingTrackViewport.h"
 #include "Insights/ViewModels/TimingEvent.h"
@@ -28,12 +29,11 @@ void FObjectPropertiesTrack::BuildDrawState(ITimingEventsTrackDrawStateBuilder&
 
 	if(GameplayProvider)
 	{
-		TraceServices::FAnalysisSessionReadScope SessionReadScope(SharedData.GetAnalysisSession());
-
 		const FText PropertiesText(LOCTEXT("PropertiesEventLabel", "Properties"));
 		const TCHAR* PropertiesString = *PropertiesText.ToString();
 
 		// object events
+		TraceServices::FProviderReadScopeLock ProviderReadScope(*GameplayProvider);
 		GameplayProvider->ReadObjectPropertiesTimeline(GetGameplayTrack().GetObjectId(), [&Context, &Builder, &PropertiesString](const FGameplayProvider::ObjectPropertiesTimeline& InTimeline)
 		{
 			InTimeline.EnumerateEvents(Context.GetViewport().GetStartTime(), Context.GetViewport().GetEndTime(), [&Builder, &PropertiesString](double InStartTime, double InEndTime, uint32 InDepth, const FObjectPropertiesMessage& InMessage)
@@ -54,24 +54,18 @@ void FObjectPropertiesTrack::Draw(const ITimingTrackDrawContext& Context) const
 
 void FObjectPropertiesTrack::InitTooltip(FTooltipDrawState& Tooltip, const ITimingEvent& HoveredTimingEvent) const
 {
-	const FGameplayProvider* GameplayProvider = SharedData.GetAnalysisSession().ReadProvider<FGameplayProvider>(FGameplayProvider::ProviderName);
-	if(GameplayProvider)
-	{
-		TraceServices::FAnalysisSessionReadScope SessionReadScope(SharedData.GetAnalysisSession());
+	FTimingEventSearchParameters SearchParameters(HoveredTimingEvent.GetStartTime(), HoveredTimingEvent.GetEndTime(), ETimingEventSearchFlags::StopAtFirstMatch);
 
-		FTimingEventSearchParameters SearchParameters(HoveredTimingEvent.GetStartTime(), HoveredTimingEvent.GetEndTime(), ETimingEventSearchFlags::StopAtFirstMatch);
-
-		FindObjectEvent(SearchParameters, [this, &Tooltip, &HoveredTimingEvent](double InFoundStartTime, double InFoundEndTime, uint32 InFoundDepth, const FObjectPropertiesMessage& InMessage)
-		{
-			Tooltip.ResetContent();
+	FindObjectEvent(SearchParameters, [this, &Tooltip, &HoveredTimingEvent](double InFoundStartTime, double InFoundEndTime, uint32 InFoundDepth, const FObjectPropertiesMessage& InMessage)
+	{
+		Tooltip.ResetContent();
 
-			Tooltip.AddTitle(FText::Format(LOCTEXT("PropertiesEventFormat", "{0} Properties"), FText::AsNumber(InMessage.PropertyValueEndIndex - InMessage.PropertyValueStartIndex)).ToString());	
-			Tooltip.AddNameValueTextLine(LOCTEXT("EventTime", "Time").ToString(), FText::AsNumber(HoveredTimingEvent.GetStartTime()).ToString());
-			Tooltip.AddNameValueTextLine(LOCTEXT("EventWorld", "World").ToString(), GetGameplayTrack().GetWorldName(SharedData.GetAnalysisSession()).ToString());
+		Tooltip.AddTitle(FText::Format(LOCTEXT("PropertiesEventFormat", "{0} Properties"), FText::AsNumber(InMessage.PropertyValueEndIndex - InMessage.PropertyValueStartIndex)).ToString());	
+		Tooltip.AddNameValueTextLine(LOCTEXT("EventTime", "Time").ToString(), FText::AsNumber(HoveredTimingEvent.GetStartTime()).ToString());
+		Tooltip.AddNameValueTextLine(LOCTEXT("EventWorld", "World").ToString(), GetGameplayTrack().GetWorldName(SharedData.GetAnalysisSession()).ToString());
 
-			Tooltip.UpdateLayout();
-		});
-	}
+		Tooltip.UpdateLayout();
+	});
 }
 
 const TSharedPtr<const ITimingEvent> FObjectPropertiesTrack::SearchEvent(const FTimingEventSearchParameters& InSearchParameters) const
@@ -98,7 +92,7 @@ void FObjectPropertiesTrack::FindObjectEvent(const FTimingEventSearchParameters&
 
 			if(GameplayProvider)
 			{
-				TraceServices::FAnalysisSessionReadScope SessionReadScope(SharedData.GetAnalysisSession());
+				TraceServices::FProviderReadScopeLock ProviderReadScope(*GameplayProvider);
 
 				GameplayProvider->ReadObjectPropertiesTimeline(GetGameplayTrack().GetObjectId(), [&InContext](const FGameplayProvider::ObjectPropertiesTimeline& InTimeline)
 				{
@@ -126,10 +120,14 @@ FText FObjectPropertiesTrack::MakeTrackName(const FGameplaySharedData& InSharedD
 	const FGameplayProvider* GameplayProvider = InSharedData.GetAnalysisSession().ReadProvider<FGameplayProvider>(FGameplayProvider::ProviderName);
 	if(GameplayProvider)
 	{
-		TraceServices::FAnalysisSessionReadScope SessionReadScope(InSharedData.GetAnalysisSession());
+		bool bIsWorld = false;
+		{
+			TraceServices::FProviderReadScopeLock ProviderReadScope(*GameplayProvider);
+			bIsWorld = GameplayProvider->IsWorld(InObjectID);
+		}
 
 		FText ObjectName;
-		if (GameplayProvider->IsWorld(InObjectID))
+		if (bIsWorld)
 		{
 			ObjectName = GetGameplayTrack().GetWorldName(InSharedData.GetAnalysisSession());
 		}
@@ -148,11 +146,10 @@ void FObjectPropertiesTrack::GetVariantsAtFrame(const TraceServices::FFrame& InF
 {
 	if(const FGameplayProvider* GameplayProvider = SharedData.GetAnalysisSession().ReadProvider<FGameplayProvider>(FGameplayProvider::ProviderName))
 	{
-		TraceServices::FAnalysisSessionReadScope SessionReadScope(SharedData.GetAnalysisSession());
-
 		TSharedRef<FVariantTreeNode> Header = OutVariants.Add_GetRef(FVariantTreeNode::MakeHeader(FText::FromString(GetName()), INDEX_NONE));
 		TArray<TSharedPtr<FVariantTreeNode>> PropertyVariants;
 
+		TraceServices::FProviderReadScopeLock ProviderReadScope(*GameplayProvider);
 		GameplayProvider->ReadObjectPropertiesTimeline(GetGameplayTrack().GetObjectId(), [this, &InFrame, GameplayProvider, &Header, &PropertyVariants](const FGameplayProvider::ObjectPropertiesTimeline& InTimeline)
 		{
 			InTimeline.EnumerateEvents(InFrame.StartTime, InFrame.EndTime, [this, GameplayProvider, &Header, &PropertyVariants](double InStartTime, double InEndTime, uint32 InDepth, const FObjectPropertiesMessage& InMessage)
diff --git a/Engine/Plugins/Animation/GameplayInsights/Source/GameplayInsights/Private/PoseWatchTrack.cpp b/Engine/Plugins/Animation/GameplayInsights/Source/GameplayInsights/Private/PoseWatchTrack.cpp
--- a/Engine/Plugins/Animation/GameplayInsights/Source/GameplayInsights/Private/PoseWatchTrack.cpp
+++ b/Engine/Plugins/Animation/GameplayInsights/Source/GameplayInsights/Private/PoseWatchTrack.cpp
@@ -3,6 +3,7 @@
 #if WITH_EDITOR
 
 #include "PoseWatchTrack.h"
+#include "Common/ProviderLock.h"
 #include "SCurveTimelineView.h"
 #include "IRewindDebugger.h"
 #include "GameplayProvider.h"
@@ -41,7 +42,7 @@ void FPoseWatchCurveTrack::UpdateCurvePointsInternal()
 	auto& CurvePoints = CurveData->Points;
 	CurvePoints.SetNum(0,EAllowShrinking::No);
 
-	TraceServices::FAnalysisSessionReadScope SessionReadScope(*AnalysisSession);
+	TraceServices::FProviderReadScopeLock ProviderReadScope(*AnimationProvider);
 
 	AnimationProvider->ReadPoseWatchTimeline(ObjectId, [this, StartTime, EndTime, AnimationProvider, &CurvePoints](const FAnimationProvider::PoseWatchTimeline& InPoseWatchTimeline)
 	{
@@ -107,12 +108,11 @@ bool FPoseWatchTrack::UpdateInternal()
 	double EndTime = TraceTimeRange.GetUpperBoundValue();
 	
 	const TraceServices::IAnalysisSession* AnalysisSession = RewindDebugger->GetAnalysisSession();
-	const FGameplayProvider* GameplayProvider = AnalysisSession->ReadProvider<FGameplayProvider>(FGameplayProvider::ProviderName);
 	const FAnimationProvider* AnimationProvider = AnalysisSession->ReadProvider<FAnimationProvider>(FAnimationProvider::ProviderName);
 
 	bool bChanged = false;
 
-	if(GameplayProvider && AnimationProvider && EnabledSegments.IsValid())
+	if (AnimationProvider && EnabledSegments.IsValid())
 	{
 		EnabledSegments->Segments.SetNum(0, EAllowShrinking::No);
 
@@ -129,33 +129,41 @@ bool FPoseWatchTrack::UpdateInternal()
 
 		TArray<FPoseWatchEnabledTime> RecordingTimes;
 
-		TraceServices::FAnalysisSessionReadScope SessionReadScope(*AnalysisSession);
-
 		TArray<uint64> UniqueCurveIds;
 
-		AnimationProvider->ReadPoseWatchTimeline(ObjectId, [this, StartTime, EndTime, &RecordingTimes, &UniqueCurveIds, AnimationProvider](const FAnimationProvider::PoseWatchTimeline& InPoseWatchTimeline)
 		{
-			InPoseWatchTimeline.EnumerateEvents(StartTime, EndTime, [this, StartTime, EndTime, &RecordingTimes, &UniqueCurveIds, AnimationProvider](double InStartTime, double InEndTime, uint32 InDepth, const FPoseWatchMessage& InMessage)
+			TraceServices::FProviderReadScopeLock ProviderReadScope(*AnimationProvider);
+	
+			AnimationProvider->ReadPoseWatchTimeline(ObjectId, [this, StartTime, EndTime, &RecordingTimes, &UniqueCurveIds, AnimationProvider](const FAnimationProvider::PoseWatchTimeline& InPoseWatchTimeline)
 			{
-				if (InMessage.PoseWatchId == PoseWatchTrackId)
+				InPoseWatchTimeline.EnumerateEvents(StartTime, EndTime, [this, StartTime, EndTime, &RecordingTimes, &UniqueCurveIds, AnimationProvider](double InStartTime, double InEndTime, uint32 InDepth, const FPoseWatchMessage& InMessage)
 				{
-					Color = InMessage.Color;
-					NameId = InMessage.NameId;
-
-					FPoseWatchEnabledTime EnabledTime;
-					EnabledTime.RecordingTime = InMessage.RecordingTime;
-					EnabledTime.bIsEnabled = InMessage.bIsEnabled;
-
-					RecordingTimes.Add(EnabledTime);
-
-					AnimationProvider->EnumeratePoseWatchCurves(InMessage, [this, &UniqueCurveIds](const FSkeletalMeshNamedCurve& InCurve)
+					if (InMessage.PoseWatchId == PoseWatchTrackId)
 					{
-						UniqueCurveIds.AddUnique(InCurve.Id);
-					});
-				}
-				return TraceServices::EEventEnumerate::Continue;
+						Color = InMessage.Color;
+						NameId = InMessage.NameId;
+	
+						FPoseWatchEnabledTime EnabledTime;
+						EnabledTime.RecordingTime = InMessage.RecordingTime;
+						EnabledTime.bIsEnabled = InMessage.bIsEnabled;
+	
+						RecordingTimes.Add(EnabledTime);
+	
+						AnimationProvider->EnumeratePoseWatchCurves(InMessage, [this, &UniqueCurveIds](const FSkeletalMeshNamedCurve& InCurve)
+						{
+							UniqueCurveIds.AddUnique(InCurve.Id);
+						});
+					}
+					return TraceServices::EEventEnumerate::Continue;
+				});
 			});
-		});
+
+			if(const TCHAR* FoundName = AnimationProvider->GetName(NameId))
+			{
+				bChanged = bChanged || TrackName.ToString().Compare(FoundName) != 0;
+				TrackName = FText::FromString(FoundName);
+			}
+		}
 
 		if (!RecordingTimes.IsEmpty())
 		{
@@ -210,12 +218,6 @@ bool FPoseWatchTrack::UpdateInternal()
 				bChanged = true;
 			}
 		}
-
-		if(const TCHAR* FoundName = AnimationProvider->GetName(NameId))
-		{
-			bChanged = bChanged || TrackName.ToString().Compare(FoundName) != 0;
-			TrackName = FText::FromString(FoundName);
-		}
 	}
 
 	return bChanged;
@@ -275,12 +277,11 @@ bool FPoseWatchesTrack::UpdateInternal()
 	double EndTime = TraceTimeRange.GetUpperBoundValue();
 	
 	const TraceServices::IAnalysisSession* AnalysisSession = RewindDebugger->GetAnalysisSession();
-	const FGameplayProvider* GameplayProvider = AnalysisSession->ReadProvider<FGameplayProvider>(FGameplayProvider::ProviderName);
 	const FAnimationProvider* AnimationProvider = AnalysisSession->ReadProvider<FAnimationProvider>(FAnimationProvider::ProviderName);
 	
 	bool bChanged = false;
 	
-	if(GameplayProvider && AnimationProvider)
+	if (AnimationProvider)
 	{
 		struct FUniquePoseWatch
 		{
@@ -301,17 +302,18 @@ bool FPoseWatchesTrack::UpdateInternal()
 		
 		TArray<FUniquePoseWatch> UniquePoseWatches;
 
-		TraceServices::FAnalysisSessionReadScope SessionReadScope(*AnalysisSession);
-
-		AnimationProvider->ReadPoseWatchTimeline(ObjectId, [this, StartTime, EndTime, &UniquePoseWatches](const FAnimationProvider::PoseWatchTimeline& InPoseWatchTimeline)
 		{
-			InPoseWatchTimeline.EnumerateEvents(StartTime, EndTime, [this, &UniquePoseWatches](double InStartTime, double InEndTime, uint32 InDepth, const FPoseWatchMessage& InMessage)
+			TraceServices::FProviderReadScopeLock ProviderReadScope(*AnimationProvider);
+			AnimationProvider->ReadPoseWatchTimeline(ObjectId, [this, StartTime, EndTime, &UniquePoseWatches](const FAnimationProvider::PoseWatchTimeline& InPoseWatchTimeline)
 			{
-				UniquePoseWatches.AddUnique({ InMessage.PoseWatchId, FColor(InMessage.Color), InMessage.NameId });
-				return TraceServices::EEventEnumerate::Continue;
+				InPoseWatchTimeline.EnumerateEvents(StartTime, EndTime, [this, &UniquePoseWatches](double InStartTime, double InEndTime, uint32 InDepth, const FPoseWatchMessage& InMessage)
+				{
+					UniquePoseWatches.AddUnique({ InMessage.PoseWatchId, FColor(InMessage.Color), InMessage.NameId });
+					return TraceServices::EEventEnumerate::Continue;
+				});
 			});
-		});
-		
+		}
+
 		UniquePoseWatches.StableSort();
 		const int32 TrackCount = UniquePoseWatches.Num();
 
@@ -347,10 +349,10 @@ bool FPoseWatchesTrackCreator::HasDebugInfoInternal(const FObjectId& InObjectId)
 	TRACE_CPUPROFILER_EVENT_SCOPE(FPoseWatchesTrack::HasDebugInfoInternal);
 	const TraceServices::IAnalysisSession* AnalysisSession = IRewindDebugger::Instance()->GetAnalysisSession();
 	
-	TraceServices::FAnalysisSessionReadScope SessionReadScope(*AnalysisSession);
 	bool bHasData = false;
 	if (const FAnimationProvider* AnimationProvider = AnalysisSession->ReadProvider<FAnimationProvider>(FAnimationProvider::ProviderName))
 	{
+		TraceServices::FProviderReadScopeLock ProviderReadScope(*AnimationProvider);
 		AnimationProvider->ReadPoseWatchTimeline(InObjectId.GetMainId(), [&bHasData](const FAnimationProvider::PoseWatchTimeline& InTimeline)
 		{
 			bHasData = true;
diff --git a/Engine/Plugins/Animation/GameplayInsights/Source/GameplayInsights/Private/PropertiesTrack.cpp b/Engine/Plugins/Animation/GameplayInsights/Source/GameplayInsights/Private/PropertiesTrack.cpp
--- a/Engine/Plugins/Animation/GameplayInsights/Source/GameplayInsights/Private/PropertiesTrack.cpp
+++ b/Engine/Plugins/Animation/GameplayInsights/Source/GameplayInsights/Private/PropertiesTrack.cpp
@@ -3,6 +3,7 @@
 #include "PropertiesTrack.h"
 
 #include "AnimationProvider.h"
+#include "Common/ProviderLock.h"
 #include "GameplayProvider.h"
 #include "IGameplayProvider.h"
 #include "IRewindDebugger.h"
@@ -319,10 +320,9 @@ namespace RewindDebugger
 
 		if (const TraceServices::IAnalysisSession* AnalysisSession = IRewindDebugger::Instance()->GetAnalysisSession())
 		{
-			TraceServices::FAnalysisSessionReadScope SessionReadScope(*AnalysisSession);
-		
 			if (const FGameplayProvider* GameplayProvider = AnalysisSession->ReadProvider<FGameplayProvider>(FGameplayProvider::ProviderName))
 			{
+				TraceServices::FProviderReadScopeLock ProviderReadScope(*GameplayProvider);
 				GameplayProvider->ReadObjectPropertiesTimeline(InObjectId.GetMainId(), [&bHasData](const IGameplayProvider::ObjectPropertiesTimeline & ObjPropsTimeline)
 				{
 					bHasData = true;
diff --git a/Engine/Plugins/Animation/GameplayInsights/Source/GameplayInsights/Private/PropertyHelpers.cpp b/Engine/Plugins/Animation/GameplayInsights/Source/GameplayInsights/Private/PropertyHelpers.cpp
--- a/Engine/Plugins/Animation/GameplayInsights/Source/GameplayInsights/Private/PropertyHelpers.cpp
+++ b/Engine/Plugins/Animation/GameplayInsights/Source/GameplayInsights/Private/PropertyHelpers.cpp
@@ -1,6 +1,7 @@
 ﻿// Copyright Epic Games, Inc. All Rights Reserved.
 
 #include "PropertyHelpers.h"
+#include "Common/ProviderLock.h"
 #include "PropertyTrack.h"
 #include "GameplayProvider.h"
 #include "VariantTreeNode.h"
@@ -16,10 +17,9 @@ FLinearColor FObjectPropertyHelpers::GetPropertyColor(const FObjectPropertyValue
 
 	if (const TraceServices::IAnalysisSession* AnalysisSession = RewindDebugger->GetAnalysisSession())
 	{
-		TraceServices::FAnalysisSessionReadScope SessionReadScope(*AnalysisSession);
-		
 		if (const FGameplayProvider* GameplayProvider = AnalysisSession->ReadProvider<FGameplayProvider>(FGameplayProvider::ProviderName))
 		{
+			TraceServices::FProviderReadScopeLock ProviderReadScope(*GameplayProvider);
 			// Cached type strings
 			static const FString FloatString = TEXT("float");
 			static const FString DoubleString = TEXT("double");
@@ -165,10 +165,9 @@ FSlateIcon FObjectPropertyHelpers::GetPropertyIcon(const FObjectPropertyValue &
 
 	if (const TraceServices::IAnalysisSession* AnalysisSession = RewindDebugger->GetAnalysisSession())
 	{
-		TraceServices::FAnalysisSessionReadScope SessionReadScope(*AnalysisSession);
-		
 		if (const FGameplayProvider* GameplayProvider = AnalysisSession->ReadProvider<FGameplayProvider>(FGameplayProvider::ProviderName))
 		{
+			TraceServices::FProviderReadScopeLock ProviderReadScope(*GameplayProvider);
 			const FString PropertyTypeString = GameplayProvider->GetPropertyName(InProperty.TypeStringId);
 		
 			if (PropertyTypeString.StartsWith(TEXT("TArray")))
@@ -238,10 +237,9 @@ bool FObjectPropertyHelpers::FindPropertyValueFromNameId(uint32 InPropertyNameId
 {
 	bool bFoundProperty = false;
 
-	TraceServices::FAnalysisSessionReadScope SessionReadScope(InSession);
-
 	if (const FGameplayProvider* GameplayProvider = InSession.ReadProvider<FGameplayProvider>(FGameplayProvider::ProviderName))
 	{
+		TraceServices::FProviderReadScopeLock ProviderReadScope(*GameplayProvider);
 		GameplayProvider->ReadObjectPropertiesTimeline(InObjectId, [InStartTime, InEndTime, GameplayProvider, InObjectId, InPropertyNameId, &bFoundProperty, &InOutProperty](const IGameplayProvider::ObjectPropertiesTimeline & InObjPropsTimeline)
 		{
 			InObjPropsTimeline.EnumerateEvents(InStartTime, InEndTime, [GameplayProvider, InObjectId, InPropertyNameId, &bFoundProperty, &InOutProperty](double InStartTimeProcedure, double InEndTimeProcedure, uint32 InDepth, const FObjectPropertiesMessage& InMessage) -> TraceServices::EEventEnumerate
diff --git a/Engine/Plugins/Animation/GameplayInsights/Source/GameplayInsights/Private/SAnimGraphSchematicView.cpp b/Engine/Plugins/Animation/GameplayInsights/Source/GameplayInsights/Private/SAnimGraphSchematicView.cpp
--- a/Engine/Plugins/Animation/GameplayInsights/Source/GameplayInsights/Private/SAnimGraphSchematicView.cpp
+++ b/Engine/Plugins/Animation/GameplayInsights/Source/GameplayInsights/Private/SAnimGraphSchematicView.cpp
@@ -2,6 +2,7 @@
 
 #include "SAnimGraphSchematicView.h"
 #include "AnimationProvider.h"
+#include "Common/ProviderLock.h"
 #include "Widgets/Input/SSearchBox.h"
 #include "TraceServices/Model/AnalysisSession.h"
 #include "GameplayProvider.h"
@@ -237,125 +238,134 @@ class SAnimGraphSchematicNode : public SMultiColumnTableRow<TSharedRef<FAnimGrap
 		{
 			IRewindDebugger* RewindDebugger = IRewindDebugger::Instance();
 
+			bool bHasValidInfo = false;
+			FString AnimInstanceClassPathName;
+			const int32 AnimNodeId = Node->NodeId;
+			const uint64 AnimInstanceId = Node->AnimInstanceId;
+
 			if (const TraceServices::IAnalysisSession* AnalysisSession = RewindDebugger->GetAnalysisSession())
 			{
 				if (const FGameplayProvider* GameplayProvider = AnalysisSession->ReadProvider<FGameplayProvider>(FGameplayProvider::ProviderName))
 				{
-					TraceServices::FAnalysisSessionReadScope SessionReadScope(*AnalysisSession);
+					TraceServices::FProviderReadScopeLock ProviderReadScope(*GameplayProvider);
 
 					const FObjectInfo* ObjectInfo = GameplayProvider->FindObjectInfo(Node->AnimInstanceId);
 					const FClassInfo* AnimInstanceClassInfo = ObjectInfo ? GameplayProvider->FindClassInfo(ObjectInfo->ClassId) : nullptr;
-					const int32 AnimNodeId = Node->NodeId;
-					const uint64 AnimInstanceId = Node->AnimInstanceId;
 
 					if (ObjectInfo && AnimInstanceClassInfo)
 					{
+						bHasValidInfo = true;
+						AnimInstanceClassPathName = AnimInstanceClassInfo->PathName;
+					}
+				}
+			}
+
+			if (bHasValidInfo)
+			{
 #if WITH_EDITOR
-						return
-							SNew(SBorder)
-							.BorderImage(bIsRoot ? FGameplayInsightsStyle::Get().GetBrush("SchematicViewRootLeft") : FCoreStyle::Get().GetBrush("NoBorder"))
-							[
-								SNew(SHorizontalBox)
-								+SHorizontalBox::Slot()
-								.AutoWidth()
-								.VAlign(VAlign_Center)
-								.Padding(6, 0, 0, 0)
-								[
-									SNew(SExpanderArrow, SharedThis(this))
-									.IndentAmount(12)
-								]
-								+SHorizontalBox::Slot()
-								.AutoWidth()
-								.VAlign(VAlign_Center)
-								[
-									SNew(SHyperlink)
-									.Text(Node->Type)
-									.TextStyle(&FCoreStyle::Get().GetWidgetStyle<FTextBlockStyle>("SmallText"))
-									.ToolTipText(FText::Format(LOCTEXT("AssetHyperlinkTooltipFormat", "Open node '{0}'"), Node->Type))
-									.HighlightText(FilterText)
-									.OnNavigate_Lambda([AnimNodeId, AnimInstanceClassInfo, AnimInstanceId]()
+				return
+					SNew(SBorder)
+					.BorderImage(bIsRoot ? FGameplayInsightsStyle::Get().GetBrush("SchematicViewRootLeft") : FCoreStyle::Get().GetBrush("NoBorder"))
+					[
+						SNew(SHorizontalBox)
+						+SHorizontalBox::Slot()
+						.AutoWidth()
+						.VAlign(VAlign_Center)
+						.Padding(6, 0, 0, 0)
+						[
+							SNew(SExpanderArrow, SharedThis(this))
+							.IndentAmount(12)
+						]
+						+SHorizontalBox::Slot()
+						.AutoWidth()
+						.VAlign(VAlign_Center)
+						[
+							SNew(SHyperlink)
+							.Text(Node->Type)
+							.TextStyle(&FCoreStyle::Get().GetWidgetStyle<FTextBlockStyle>("SmallText"))
+							.ToolTipText(FText::Format(LOCTEXT("AssetHyperlinkTooltipFormat", "Open node '{0}'"), Node->Type))
+							.HighlightText(FilterText)
+							.OnNavigate_Lambda([AnimNodeId, AnimInstanceClassPathName, AnimInstanceId]()
+							{
+								TSoftObjectPtr<UAnimBlueprintGeneratedClass> InstanceClass;
+								InstanceClass = FSoftObjectPath(AnimInstanceClassPathName);
+
+								if (InstanceClass.LoadSynchronous())
+								{
+									if (UAnimBlueprint* AnimBlueprint = Cast<UAnimBlueprint>(InstanceClass.Get()->ClassGeneratedBy))
 									{
-										TSoftObjectPtr<UAnimBlueprintGeneratedClass> InstanceClass;
-										InstanceClass = FSoftObjectPath(AnimInstanceClassInfo->PathName);
+										GEditor->GetEditorSubsystem<UAssetEditorSubsystem>()->OpenEditorForAsset(AnimBlueprint);
 
-										if (InstanceClass.LoadSynchronous())
+										if (UObject* SelectedInstance = FObjectTrace::GetObjectFromId(AnimInstanceId))
 										{
-											if (UAnimBlueprint* AnimBlueprint = Cast<UAnimBlueprint>(InstanceClass.Get()->ClassGeneratedBy))
+											AnimBlueprint->SetObjectBeingDebugged(SelectedInstance);
+										}
+
+										if (IAnimationBlueprintEditor* AnimBlueprintEditor = static_cast<IAnimationBlueprintEditor*>(GEditor->GetEditorSubsystem<UAssetEditorSubsystem>()->FindEditorForAsset(AnimBlueprint, true)))
+										{
+											int32 AnimNodeIndex = InstanceClass.Get()->GetAnimNodeProperties().Num() - AnimNodeId - 1;
+											TWeakObjectPtr<const UEdGraphNode>* GraphNode = InstanceClass.Get()->AnimBlueprintDebugData.NodePropertyIndexToNodeMap.Find(AnimNodeIndex);
+											if (GraphNode != nullptr && GraphNode->Get())
 											{
-												GEditor->GetEditorSubsystem<UAssetEditorSubsystem>()->OpenEditorForAsset(AnimBlueprint);
-												
-												if (UObject* SelectedInstance = FObjectTrace::GetObjectFromId(AnimInstanceId))
-												{
-													AnimBlueprint->SetObjectBeingDebugged(SelectedInstance);
-												}
-
-												if (IAnimationBlueprintEditor* AnimBlueprintEditor = static_cast<IAnimationBlueprintEditor*>(GEditor->GetEditorSubsystem<UAssetEditorSubsystem>()->FindEditorForAsset(AnimBlueprint, true)))
-												{
-													int32 AnimNodeIndex = InstanceClass.Get()->GetAnimNodeProperties().Num() - AnimNodeId - 1;
-													TWeakObjectPtr<const UEdGraphNode>* GraphNode = InstanceClass.Get()->AnimBlueprintDebugData.NodePropertyIndexToNodeMap.Find(AnimNodeIndex);
-													if (GraphNode != nullptr && GraphNode->Get())
-													{
-														AnimBlueprintEditor->JumpToHyperlink(GraphNode->Get());
-													}
-												}
+												AnimBlueprintEditor->JumpToHyperlink(GraphNode->Get());
 											}
 										}
-									})
-								]
-							];
+									}
+								}
+							})
+						]
+					];
 #else
-						return 
-							SNew(SBorder)
-							.BorderImage(bIsRoot ? FGameplayInsightsStyle::Get().GetBrush("SchematicViewRootLeft") : FCoreStyle::Get().GetBrush("NoBorder"))
-							[
-								SNew(SHorizontalBox)
-								+SHorizontalBox::Slot()
-								.AutoWidth()
-								.VAlign(VAlign_Center)
-								.Padding(6, 0, 0, 0)
-								[
-									SNew(SExpanderArrow, SharedThis(this))
-									.IndentAmount(12)
-								]
-								+SHorizontalBox::Slot()
-								.AutoWidth()
-								.VAlign(VAlign_Center)
-								[
-									SNew(SHyperlink)
-									.Text(Node->Type)
-									.TextStyle(&FCoreStyle::Get().GetWidgetStyle<FTextBlockStyle>("SmallText"))
-									.HighlightText(FilterText)
-								]
-							];
+				return
+					SNew(SBorder)
+					.BorderImage(bIsRoot ? FGameplayInsightsStyle::Get().GetBrush("SchematicViewRootLeft") : FCoreStyle::Get().GetBrush("NoBorder"))
+					[
+						SNew(SHorizontalBox)
+						+SHorizontalBox::Slot()
+						.AutoWidth()
+						.VAlign(VAlign_Center)
+						.Padding(6, 0, 0, 0)
+						[
+							SNew(SExpanderArrow, SharedThis(this))
+							.IndentAmount(12)
+						]
+						+SHorizontalBox::Slot()
+						.AutoWidth()
+						.VAlign(VAlign_Center)
+						[
+							SNew(SHyperlink)
+							.Text(Node->Type)
+							.TextStyle(&FCoreStyle::Get().GetWidgetStyle<FTextBlockStyle>("SmallText"))
+							.HighlightText(FilterText)
+						]
+					];
 #endif
-					}
-					else
-					{
-						return 
-							SNew(SBorder)
-							.BorderImage(bIsRoot ? FGameplayInsightsStyle::Get().GetBrush("SchematicViewRootLeft") : FCoreStyle::Get().GetBrush("NoBorder"))
-							[
-								SNew(SHorizontalBox)
-								+SHorizontalBox::Slot()
-								.AutoWidth()
-								.VAlign(VAlign_Center)
-								.Padding(6, 0, 0, 0)
-								[
-									SNew(SExpanderArrow, SharedThis(this))
-									.IndentAmount(12)
-								]
-								+SHorizontalBox::Slot()
-								.AutoWidth()
-								.VAlign(VAlign_Center)
-								[
-									SNew(SHyperlink)
-									.Text(Node->Type)
-									.TextStyle(&FCoreStyle::Get().GetWidgetStyle<FTextBlockStyle>("SmallText"))
-									.HighlightText(FilterText)
-								]
-							];
-					}
-				}
+			}
+			else
+			{
+				return
+					SNew(SBorder)
+					.BorderImage(bIsRoot ? FGameplayInsightsStyle::Get().GetBrush("SchematicViewRootLeft") : FCoreStyle::Get().GetBrush("NoBorder"))
+					[
+						SNew(SHorizontalBox)
+						+SHorizontalBox::Slot()
+						.AutoWidth()
+						.VAlign(VAlign_Center)
+						.Padding(6, 0, 0, 0)
+						[
+							SNew(SExpanderArrow, SharedThis(this))
+							.IndentAmount(12)
+						]
+						+SHorizontalBox::Slot()
+						.AutoWidth()
+						.VAlign(VAlign_Center)
+						[
+							SNew(SHyperlink)
+							.Text(Node->Type)
+							.TextStyle(&FCoreStyle::Get().GetWidgetStyle<FTextBlockStyle>("SmallText"))
+							.HighlightText(FilterText)
+						]
+					];
 			}
 		}
 		else
@@ -540,7 +550,8 @@ void SAnimGraphSchematicView::RefreshNodes()
 
 	if(AnimationProvider && GameplayProvider)
 	{
-		TraceServices::FAnalysisSessionReadScope SessionReadScope(*AnalysisSession);
+		TraceServices::FProviderReadScopeLock AnimationProviderReadScope(*AnimationProvider);
+		TraceServices::FProviderReadScopeLock GameplayProviderReadScope(*GameplayProvider);
 
 		AnimationProvider->ReadAnimGraphTimeline(AnimInstanceId, [this, &NodeMap, AnimationProvider, GameplayProvider](const FAnimationProvider::AnimGraphTimeline& InGraphTimeline)
 		{
@@ -1023,10 +1034,10 @@ bool FAnimGraphSchematicTrackCreator::HasDebugInfoInternal(const RewindDebugger:
 {
 	const TraceServices::IAnalysisSession* AnalysisSession = IRewindDebugger::Instance()->GetAnalysisSession();
 	
-	TraceServices::FAnalysisSessionReadScope SessionReadScope(*AnalysisSession);
 	bool bHasData = false;
 	if (const FAnimationProvider* AnimationProvider = AnalysisSession->ReadProvider<FAnimationProvider>(FAnimationProvider::ProviderName))
 	{
+		TraceServices::FProviderReadScopeLock ProviderReadScope(*AnimationProvider);
 		AnimationProvider->ReadAnimGraphTimeline(InObjectId.GetMainId(), [&bHasData](const FAnimationProvider::AnimGraphTimeline& InGraphTimeline)
 		{
 			bHasData = true;
diff --git a/Engine/Plugins/Animation/GameplayInsights/Source/GameplayInsights/Private/SAnimationCurvesView.cpp b/Engine/Plugins/Animation/GameplayInsights/Source/GameplayInsights/Private/SAnimationCurvesView.cpp
--- a/Engine/Plugins/Animation/GameplayInsights/Source/GameplayInsights/Private/SAnimationCurvesView.cpp
+++ b/Engine/Plugins/Animation/GameplayInsights/Source/GameplayInsights/Private/SAnimationCurvesView.cpp
@@ -2,8 +2,8 @@
 
 #include "SAnimationCurvesView.h"
 #include "AnimationProvider.h"
+#include "Common/ProviderLock.h"
 #include "TraceServices/Model/AnalysisSession.h"
-#include "GameplayProvider.h"
 #include "TraceServices/Model/Frames.h"
 #include "VariantTreeNode.h"
 
@@ -11,12 +11,11 @@
 
 void SAnimationCurvesView::GetVariantsAtFrame(const TraceServices::FFrame& InFrame, TArray<TSharedRef<FVariantTreeNode>>& OutVariants) const
 {
-	const FGameplayProvider* GameplayProvider = AnalysisSession->ReadProvider<FGameplayProvider>(FGameplayProvider::ProviderName);
 	const FAnimationProvider* AnimationProvider = AnalysisSession->ReadProvider<FAnimationProvider>(FAnimationProvider::ProviderName);
 
-	if(GameplayProvider && AnimationProvider)
+	if (AnimationProvider)
 	{
-		TraceServices::FAnalysisSessionReadScope SessionReadScope(*AnalysisSession);
+		TraceServices::FProviderReadScopeLock AnimationProviderReadScope(*AnimationProvider);
 
 		TSharedRef<FVariantTreeNode> Header = OutVariants.Add_GetRef(FVariantTreeNode::MakeHeader(LOCTEXT("Animation Curves", "Animation Curves"), 0));
 
diff --git a/Engine/Plugins/Animation/GameplayInsights/Source/GameplayInsights/Private/SBlendWeightsView.cpp b/Engine/Plugins/Animation/GameplayInsights/Source/GameplayInsights/Private/SBlendWeightsView.cpp
--- a/Engine/Plugins/Animation/GameplayInsights/Source/GameplayInsights/Private/SBlendWeightsView.cpp
+++ b/Engine/Plugins/Animation/GameplayInsights/Source/GameplayInsights/Private/SBlendWeightsView.cpp
@@ -2,6 +2,7 @@
 
 #include "SBlendWeightsView.h"
 #include "AnimationProvider.h"
+#include "Common/ProviderLock.h"
 #include "TraceServices/Model/AnalysisSession.h"
 #include "GameplayProvider.h"
 #include "TraceServices/Model/Frames.h"
@@ -16,7 +17,8 @@ void SBlendWeightsView::GetVariantsAtFrame(const TraceServices::FFrame& InFrame,
 
 	if(GameplayProvider && AnimationProvider)
 	{
-		TraceServices::FAnalysisSessionReadScope SessionReadScope(*AnalysisSession);
+		TraceServices::FProviderReadScopeLock AnimationProviderReadScope(*AnimationProvider);
+		TraceServices::FProviderReadScopeLock GameplayProviderReadScope(*GameplayProvider);
 
 		AnimationProvider->ReadTickRecordTimeline(ObjectId, [this, &AnimationProvider, &GameplayProvider, &OutVariants, &InFrame](const FAnimationProvider::TickRecordTimeline& InTimeline)
 		{
diff --git a/Engine/Plugins/Animation/GameplayInsights/Source/GameplayInsights/Private/SInertializationDetailsView.cpp b/Engine/Plugins/Animation/GameplayInsights/Source/GameplayInsights/Private/SInertializationDetailsView.cpp
--- a/Engine/Plugins/Animation/GameplayInsights/Source/GameplayInsights/Private/SInertializationDetailsView.cpp
+++ b/Engine/Plugins/Animation/GameplayInsights/Source/GameplayInsights/Private/SInertializationDetailsView.cpp
@@ -2,6 +2,7 @@
 
 #include "SInertializationDetailsView.h"
 #include "AnimationProvider.h"
+#include "Common/ProviderLock.h"
 #include "TraceServices/Model/AnalysisSession.h"
 #include "GameplayProvider.h"
 #include "TraceServices/Model/Frames.h"
@@ -17,14 +18,13 @@ struct FInertializationDetailsNodeItem
 
 void SInertializationDetailsView::GetVariantsAtFrame(const TraceServices::FFrame& InFrame, TArray<TSharedRef<FVariantTreeNode>>& OutVariants) const
 {
-	const FGameplayProvider* GameplayProvider = AnalysisSession->ReadProvider<FGameplayProvider>(FGameplayProvider::ProviderName);
 	const FAnimationProvider* AnimationProvider = AnalysisSession->ReadProvider<FAnimationProvider>(FAnimationProvider::ProviderName);
 
 	TSortedMap<TTuple<int32, uint64>, FInertializationDetailsNodeItem, TInlineAllocator<8>> NodeMap;
 
-	if(GameplayProvider && AnimationProvider)
+	if (AnimationProvider)
 	{
-		TraceServices::FAnalysisSessionReadScope SessionReadScope(*AnalysisSession);
+		TraceServices::FProviderReadScopeLock AnimationProviderReadScope(*AnimationProvider);
 
 		// Find all inertialization nodes within window
 
@@ -54,7 +54,6 @@ void SInertializationDetailsView::GetVariantsAtFrame(const TraceServices::FFrame
 		});
 
 		// Fill in Node Trace Details
-
 		AnimationProvider->ReadAnimNodeValuesTimeline(ObjectId, [this, &NodeMap, &AnimationProvider, &InFrame](const FAnimationProvider::AnimNodeValuesTimeline& InTimeline)
 		{
 			InTimeline.EnumerateEvents(InFrame.StartTime, InFrame.EndTime, [this, &NodeMap, &AnimationProvider, &InFrame](double InStartTime, double InEndTime, uint32 InDepth, const FAnimNodeValueMessage& InMessage)
@@ -70,39 +69,39 @@ void SInertializationDetailsView::GetVariantsAtFrame(const TraceServices::FFrame
 				return TraceServices::EEventEnumerate::Continue;
 			});
 		});
+	}
 
-		// Make UI Variant Items
+	// Make UI Variant Items
 
-		for (TPair<TTuple<int32, uint64>, FInertializationDetailsNodeItem>& NodePair : NodeMap)
-		{
-			bool bInertializationActive = false;
+	for (TPair<TTuple<int32, uint64>, FInertializationDetailsNodeItem>& NodePair : NodeMap)
+	{
+		bool bInertializationActive = false;
 
-			const int32 PropertyNum = NodePair.Value.PropertyNames.Num();
+		const int32 PropertyNum = NodePair.Value.PropertyNames.Num();
 
-			// Check if inertialization node is Active
+		// Check if inertialization node is Active
 
-			for (int32 PropertyIdx = 0; PropertyIdx < PropertyNum; PropertyIdx++)
+		for (int32 PropertyIdx = 0; PropertyIdx < PropertyNum; PropertyIdx++)
+		{
+			if (FCString::Strcmp(NodePair.Value.PropertyNames[PropertyIdx], TEXT("State")) == 0 &&
+				FCString::Strcmp(NodePair.Value.PropertyValues[PropertyIdx].String.Value, TEXT("EInertializationState::Active")) == 0)
 			{
-				if (FCString::Strcmp(NodePair.Value.PropertyNames[PropertyIdx], TEXT("State")) == 0 &&
-					FCString::Strcmp(NodePair.Value.PropertyValues[PropertyIdx].String.Value, TEXT("EInertializationState::Active")) == 0)
-				{
-					bInertializationActive = true;
-				}
+				bInertializationActive = true;
 			}
+		}
 
-			// If node is active then add to details view
+		// If node is active then add to details view
 
-			if (bInertializationActive)
-			{
-				TSharedRef<FVariantTreeNode> Header = OutVariants.Add_GetRef(
-					FVariantTreeNode::MakeAnimNode(LOCTEXT("InertializationNode", "Inertialization Node"), NodePair.Key.Get<0>(), NodePair.Key.Get<1>()));
+		if (bInertializationActive)
+		{
+			TSharedRef<FVariantTreeNode> Header = OutVariants.Add_GetRef(
+				FVariantTreeNode::MakeAnimNode(LOCTEXT("InertializationNode", "Inertialization Node"), NodePair.Key.Get<0>(), NodePair.Key.Get<1>()));
 
-				for (int32 PropertyIdx = 0; PropertyIdx < PropertyNum; PropertyIdx++)
-				{
-					Header->AddChild(MakeShared<FVariantTreeNode>(
-						FText::FromString(NodePair.Value.PropertyNames[PropertyIdx]), 
-						NodePair.Value.PropertyValues[PropertyIdx], 0, INDEX_NONE));
-				}
+			for (int32 PropertyIdx = 0; PropertyIdx < PropertyNum; PropertyIdx++)
+			{
+				Header->AddChild(MakeShared<FVariantTreeNode>(
+					FText::FromString(NodePair.Value.PropertyNames[PropertyIdx]), 
+					NodePair.Value.PropertyValues[PropertyIdx], 0, INDEX_NONE));
 			}
 		}
 	}
diff --git a/Engine/Plugins/Animation/GameplayInsights/Source/GameplayInsights/Private/SMontageView.cpp b/Engine/Plugins/Animation/GameplayInsights/Source/GameplayInsights/Private/SMontageView.cpp
--- a/Engine/Plugins/Animation/GameplayInsights/Source/GameplayInsights/Private/SMontageView.cpp
+++ b/Engine/Plugins/Animation/GameplayInsights/Source/GameplayInsights/Private/SMontageView.cpp
@@ -2,6 +2,7 @@
 
 #include "SMontageView.h"
 #include "AnimationProvider.h"
+#include "Common/ProviderLock.h"
 #include "TraceServices/Model/AnalysisSession.h"
 #include "GameplayProvider.h"
 #include "IRewindDebugger.h"
@@ -22,7 +23,8 @@ void SMontageView::GetVariantsAtFrame(const TraceServices::FFrame& InFrame, TArr
 
 	if(AnimationProvider && GameplayProvider)
 	{
-		TraceServices::FAnalysisSessionReadScope SessionReadScope(*AnalysisSession);
+		TraceServices::FProviderReadScopeLock AnimationProviderReadScope(*AnimationProvider);
+		TraceServices::FProviderReadScopeLock GameplayProviderReadScope(*GameplayProvider);
 		
 		AnimationProvider->ReadMontageTimeline(ObjectId, [this, &InFrame, &AnimationProvider, &GameplayProvider, &OutVariants](const FAnimationProvider::AnimMontageTimeline& InTimeline)
 		{
diff --git a/Engine/Plugins/Animation/GameplayInsights/Source/GameplayInsights/Private/SPoseWatchCurvesView.cpp b/Engine/Plugins/Animation/GameplayInsights/Source/GameplayInsights/Private/SPoseWatchCurvesView.cpp
--- a/Engine/Plugins/Animation/GameplayInsights/Source/GameplayInsights/Private/SPoseWatchCurvesView.cpp
+++ b/Engine/Plugins/Animation/GameplayInsights/Source/GameplayInsights/Private/SPoseWatchCurvesView.cpp
@@ -2,6 +2,7 @@
 
 #include "SPoseWatchCurvesView.h"
 #include "AnimationProvider.h"
+#include "Common/ProviderLock.h"
 #include "TraceServices/Model/AnalysisSession.h"
 #include "GameplayProvider.h"
 #include "TraceServices/Model/Frames.h"
@@ -11,15 +12,13 @@
 
 void SPoseWatchCurvesView::GetVariantsAtFrame(const TraceServices::FFrame& InFrame, TArray<TSharedRef<FVariantTreeNode>>& OutVariants) const
 {
-	const FGameplayProvider* GameplayProvider = AnalysisSession->ReadProvider<FGameplayProvider>(FGameplayProvider::ProviderName);
 	const FAnimationProvider* AnimationProvider = AnalysisSession->ReadProvider<FAnimationProvider>(FAnimationProvider::ProviderName);
 
-	if(GameplayProvider && AnimationProvider)
+	if (AnimationProvider)
 	{
-		TraceServices::FAnalysisSessionReadScope SessionReadScope(*AnalysisSession);
-
 		TSharedRef<FVariantTreeNode> Header = OutVariants.Add_GetRef(FVariantTreeNode::MakeHeader(LOCTEXT("Animation Curves", "Animation Curves"), 0));
 
+		TraceServices::FProviderReadScopeLock AnimationProviderReadScope(*AnimationProvider);
 		AnimationProvider->ReadPoseWatchTimeline(ObjectId, [this, &Header, &AnimationProvider, &InFrame](const FAnimationProvider::PoseWatchTimeline& InPoseWatchTimeline)
 		{
 			InPoseWatchTimeline.EnumerateEvents(InFrame.StartTime, InFrame.EndTime, [this, &Header, &AnimationProvider, &InFrame](double InStartTime, double InEndTime, uint32 InDepth, const FPoseWatchMessage& InMessage)
diff --git a/Engine/Plugins/Animation/GameplayInsights/Source/GameplayInsights/Private/SPropertyTrackView.cpp b/Engine/Plugins/Animation/GameplayInsights/Source/GameplayInsights/Private/SPropertyTrackView.cpp
--- a/Engine/Plugins/Animation/GameplayInsights/Source/GameplayInsights/Private/SPropertyTrackView.cpp
+++ b/Engine/Plugins/Animation/GameplayInsights/Source/GameplayInsights/Private/SPropertyTrackView.cpp
@@ -2,6 +2,7 @@
 
 #include "SPropertyTrackView.h"
 
+#include "Common/ProviderLock.h"
 #include "GameplayProvider.h"
 #include "PropertyHelpers.h"
 #include "PropertyTrack.h"
@@ -17,11 +18,11 @@ void SPropertyTrackView::GetVariantsAtFrame(const TraceServices::FFrame& InFrame
 	
 	if (const TSharedPtr<RewindDebugger::FPropertyTrack> PropertyTrack = Track.Pin())
 	{
-		TraceServices::FAnalysisSessionReadScope SessionReadScope(*AnalysisSession);
 		TSharedPtr<FVariantTreeNode> Header = OutVariants.Add_GetRef(FVariantTreeNode::MakeHeader(FText::FromName(PropertyTrack->GetPropertyName()), INDEX_NONE));
-		
+
 		if (const FGameplayProvider* GameplayProvider = AnalysisSession->ReadProvider<FGameplayProvider>(FGameplayProvider::ProviderName))
 		{
+			TraceServices::FProviderReadScopeLock ProviderReadScope(*GameplayProvider);
 			TArray<TSharedPtr<FVariantTreeNode>, TMemStackAllocator<>> PropertyVariants;
 			
 			PropertyTrack->ReadObjectPropertyValueAtFrame(InFrame, *GameplayProvider, [GameplayProvider, &PropertyTrack, &Header, &PropertyVariants, &OutVariants](const FObjectPropertyValue& InValue, uint32 InValueIndex, const FObjectPropertiesMessage& InMessage)
diff --git a/Engine/Plugins/Animation/GameplayInsights/Source/GameplayInsights/Private/SVariantValueView.cpp b/Engine/Plugins/Animation/GameplayInsights/Source/GameplayInsights/Private/SVariantValueView.cpp
--- a/Engine/Plugins/Animation/GameplayInsights/Source/GameplayInsights/Private/SVariantValueView.cpp
+++ b/Engine/Plugins/Animation/GameplayInsights/Source/GameplayInsights/Private/SVariantValueView.cpp
@@ -1,6 +1,7 @@
 // Copyright Epic Games, Inc. All Rights Reserved.
 
 #include "SVariantValueView.h"
+#include "Common/ProviderLock.h"
 #include "Widgets/Input/SSearchBox.h"
 #include "TraceServices/Model/AnalysisSession.h"
 #include "GameplayProvider.h"
@@ -156,35 +157,40 @@ TSharedRef<SWidget> SVariantValueView::MakeVariantValueWidget(const TraceService
 		const FGameplayProvider* GameplayProvider = InAnalysisSession.ReadProvider<FGameplayProvider>(FGameplayProvider::ProviderName);
 		if(GameplayProvider)
 		{
-			TraceServices::FAnalysisSessionReadScope SessionReadScope(InAnalysisSession);
-
-			const FObjectInfo& ObjectInfo = GameplayProvider->GetObjectInfo(InValue.Object.Value);
+			FString ObjectName;
+			FString ObjectPathName;
+			{
+				TraceServices::FProviderReadScopeLock ProviderReadScope(*GameplayProvider);
+				const FObjectInfo& ObjectInfo = GameplayProvider->GetObjectInfo(InValue.Object.Value);
+				ObjectName = ObjectInfo.Name;
+				ObjectPathName = ObjectInfo.PathName;
+			}
 #if WITH_EDITOR
-			return 
+			return
 				SNew(SHyperlink)
-				.Text(FText::FromString(ObjectInfo.Name))
+				.Text(FText::FromString(ObjectName))
 				.TextStyle(&FCoreStyle::Get().GetWidgetStyle<FTextBlockStyle>("SmallText"))
-				.ToolTipText(FText::Format(LOCTEXT("AssetHyperlinkTooltipFormat", "Open asset '{0}'"), FText::FromString(ObjectInfo.PathName)))
-				.OnNavigate_Lambda([ObjectInfo, InValue]()
+				.ToolTipText(FText::Format(LOCTEXT("AssetHyperlinkTooltipFormat", "Open asset '{0}'"), FText::FromString(ObjectPathName)))
+				.OnNavigate_Lambda([ObjectPathName, InValue]()
 				{
 					UObject* Asset = nullptr;
-					
-					FString PackagePathString = FPackageName::ObjectPathToPackageName(FString(ObjectInfo.PathName));
+
+					FString PackagePathString = FPackageName::ObjectPathToPackageName(ObjectPathName);
 
 					UPackage* Package = LoadPackage(NULL, ToCStr(PackagePathString), LOAD_NoRedirects);
 					if (Package)
 					{
 						Package->FullyLoad();
-                
-						FString AssetName = FPaths::GetBaseFilename(ObjectInfo.PathName);
+
+						FString AssetName = FPaths::GetBaseFilename(ObjectPathName);
 						Asset = FindObject<UObject>(Package, *AssetName);
 					}
 					else
 					{
 						// fallback for unsaved assets
-						Asset = FindObject<UObject>(nullptr, ObjectInfo.PathName);
+						Asset = FindObject<UObject>(nullptr, *ObjectPathName);
 					}
-                    	
+
 					if (Asset != nullptr)
 					{
 						if (UAssetEditorSubsystem* AssetEditorSS = GEditor->GetEditorSubsystem<UAssetEditorSubsystem>())
@@ -205,11 +211,11 @@ TSharedRef<SWidget> SVariantValueView::MakeVariantValueWidget(const TraceService
 					}
 				});
 #else
-			return 
+			return
 				SNew(STextBlock)
 				.Font(FCoreStyle::Get().GetFontStyle("SmallFont"))
-				.Text(FText::FromString(ObjectInfo.Name))
-				.ToolTipText(FText::FromString(ObjectInfo.PathName))
+				.Text(FText::FromString(ObjectName))
+				.ToolTipText(FText::FromString(ObjectPathName))
 				.HighlightText(InHighlightText);
 #endif
 		}
@@ -221,18 +227,23 @@ TSharedRef<SWidget> SVariantValueView::MakeVariantValueWidget(const TraceService
 		const FGameplayProvider* GameplayProvider = InAnalysisSession.ReadProvider<FGameplayProvider>(FGameplayProvider::ProviderName);
 		if(GameplayProvider)
 		{
-			TraceServices::FAnalysisSessionReadScope SessionReadScope(InAnalysisSession);
-
-			const FClassInfo& ClassInfo = GameplayProvider->GetClassInfo(InValue.Class.Value);
+			FString ClassName;
+			FString ClassPathName;
+			{
+				TraceServices::FProviderReadScopeLock ProviderReadScope(*GameplayProvider);
+				const FClassInfo& ClassInfo = GameplayProvider->GetClassInfo(InValue.Class.Value);
+				ClassName = ClassInfo.Name;
+				ClassPathName = ClassInfo.PathName;
+			}
 #if WITH_EDITOR
-			return 
+			return
 				SNew(SHyperlink)
-				.Text(FText::FromString(ClassInfo.Name))
+				.Text(FText::FromString(ClassName))
 				.TextStyle(&FCoreStyle::Get().GetWidgetStyle<FTextBlockStyle>("SmallText"))
-				.ToolTipText(FText::Format(LOCTEXT("ClassHyperlinkTooltipFormat", "Open class '{0}'"), FText::FromString(ClassInfo.PathName)))
-				.OnNavigate_Lambda([ClassInfo]()
+				.ToolTipText(FText::Format(LOCTEXT("ClassHyperlinkTooltipFormat", "Open class '{0}'"), FText::FromString(ClassPathName)))
+				.OnNavigate_Lambda([ClassName, ClassPathName]()
 				{
-					if (UClass* FoundClass = UClass::TryFindTypeSlow<UClass>(ClassInfo.Name))
+					if (UClass* FoundClass = UClass::TryFindTypeSlow<UClass>(ClassName))
 					{
 						if (FoundClass->IsNative())
 						{
@@ -242,16 +253,16 @@ TSharedRef<SWidget> SVariantValueView::MakeVariantValueWidget(const TraceService
 					}
 
 					// for non-native classes, try opening the asset
-					GEditor->GetEditorSubsystem<UAssetEditorSubsystem>()->OpenEditorForAsset(ClassInfo.PathName);
+					GEditor->GetEditorSubsystem<UAssetEditorSubsystem>()->OpenEditorForAsset(ClassPathName);
 				});
 #else
-			return 
+			return
 				SNew(STextBlock)
 				.Font(FCoreStyle::Get().GetFontStyle("SmallFont"))
-				.Text(FText::FromString(ClassInfo.Name))
-				.ToolTipText(FText::FromString(ClassInfo.PathName))
+				.Text(FText::FromString(ClassName))
+				.ToolTipText(FText::FromString(ClassPathName))
 				.HighlightText(InHighlightText);
-#endif			
+#endif
 		}
 	}
 
@@ -262,24 +273,37 @@ TSharedRef<SWidget> SVariantValueView::MakeVariantValueWidget(const TraceService
 
 		if (GameplayProvider && AnimationProvider)
 		{
-			TraceServices::FAnalysisSessionReadScope SessionReadScope(InAnalysisSession);
+			FString AnimNodeName;
+			int32 AnimNodeId = INDEX_NONE;
+			FString AnimInstanceClassPathName;
+			{
+				TraceServices::FProviderReadScopeLock AnimationProviderReadScope(*AnimationProvider);
+				TraceServices::FProviderReadScopeLock GameplayProviderReadScope(*GameplayProvider);
+
+				const FAnimNodeInfo* AnimNodeInfo = AnimationProvider->FindAnimNodeInfo(InValue.AnimNode.Value, InValue.AnimNode.AnimInstanceId);
+				const FObjectInfo* ObjectInfo = GameplayProvider->FindObjectInfo(InValue.AnimNode.AnimInstanceId);
+				const FClassInfo* AnimInstanceClassInfo = ObjectInfo ? GameplayProvider->FindClassInfo(ObjectInfo->ClassId) : nullptr;
 
-			const FAnimNodeInfo* AnimNodeInfo = AnimationProvider->FindAnimNodeInfo(InValue.AnimNode.Value, InValue.AnimNode.AnimInstanceId);
-			const FObjectInfo* ObjectInfo = GameplayProvider->FindObjectInfo(InValue.AnimNode.AnimInstanceId);
-			const FClassInfo* AnimInstanceClassInfo = ObjectInfo ? GameplayProvider->FindClassInfo(ObjectInfo->ClassId) : nullptr;
+				if (AnimNodeInfo && ObjectInfo && AnimInstanceClassInfo)
+				{
+					AnimNodeName = AnimNodeInfo->Name;
+					AnimNodeId = AnimNodeInfo->Id;
+					AnimInstanceClassPathName = AnimInstanceClassInfo->PathName;
+				}
+			}
 
-			if (AnimNodeInfo && ObjectInfo && AnimInstanceClassInfo)
+			if (AnimNodeId != INDEX_NONE)
 			{
 #if WITH_EDITOR
 				return
 					SNew(SHyperlink)
-					.Text(FText::FromString(AnimNodeInfo->Name))
+					.Text(FText::FromString(AnimNodeName))
 					.TextStyle(&FCoreStyle::Get().GetWidgetStyle<FTextBlockStyle>("SmallText"))
-					.ToolTipText(FText::Format(LOCTEXT("NodeHyperlinkTooltipFormat", "Open node '{0}'"), FText::FromString(AnimNodeInfo->Name)))
-					.OnNavigate_Lambda([ObjectInfo, AnimNodeInfo, AnimInstanceClassInfo, InValue]()
+					.ToolTipText(FText::Format(LOCTEXT("NodeHyperlinkTooltipFormat", "Open node '{0}'"), FText::FromString(AnimNodeName)))
+					.OnNavigate_Lambda([AnimNodeId, AnimInstanceClassPathName, InValue]()
 					{
 						TSoftObjectPtr<UAnimBlueprintGeneratedClass> InstanceClass;
-						InstanceClass = FSoftObjectPath(AnimInstanceClassInfo->PathName);
+						InstanceClass = FSoftObjectPath(AnimInstanceClassPathName);
 
 						if (InstanceClass.LoadSynchronous())
 						{
@@ -289,7 +313,7 @@ TSharedRef<SWidget> SVariantValueView::MakeVariantValueWidget(const TraceService
 
 								if (IAnimationBlueprintEditor* AnimBlueprintEditor = static_cast<IAnimationBlueprintEditor*>(GEditor->GetEditorSubsystem<UAssetEditorSubsystem>()->FindEditorForAsset(AnimBlueprint, true)))
 								{
-									int32 AnimNodeIndex = InstanceClass.Get()->GetAnimNodeProperties().Num() - AnimNodeInfo->Id - 1;
+									int32 AnimNodeIndex = InstanceClass.Get()->GetAnimNodeProperties().Num() - AnimNodeId - 1;
 									TWeakObjectPtr<const UEdGraphNode>* GraphNode = InstanceClass.Get()->AnimBlueprintDebugData.NodePropertyIndexToNodeMap.Find(AnimNodeIndex);
 									if (GraphNode != nullptr && GraphNode->Get())
 									{
@@ -303,8 +327,8 @@ TSharedRef<SWidget> SVariantValueView::MakeVariantValueWidget(const TraceService
 				return
 					SNew(STextBlock)
 					.Font(FCoreStyle::Get().GetFontStyle("SmallFont"))
-					.Text(FText::FromString(AnimNodeInfo->Name))
-					.ToolTipText(FText::FromString(AnimNodeInfo->Name))
+					.Text(FText::FromString(AnimNodeName))
+					.ToolTipText(FText::FromString(AnimNodeName))
 					.HighlightText(InHighlightText);
 #endif
 			}
diff --git a/Engine/Plugins/Animation/GameplayInsights/Source/GameplayInsights/Private/SkeletalMeshCurvesTrack.cpp b/Engine/Plugins/Animation/GameplayInsights/Source/GameplayInsights/Private/SkeletalMeshCurvesTrack.cpp
--- a/Engine/Plugins/Animation/GameplayInsights/Source/GameplayInsights/Private/SkeletalMeshCurvesTrack.cpp
+++ b/Engine/Plugins/Animation/GameplayInsights/Source/GameplayInsights/Private/SkeletalMeshCurvesTrack.cpp
@@ -1,6 +1,7 @@
 // Copyright Epic Games, Inc. All Rights Reserved.
 
 #include "SkeletalMeshCurvesTrack.h"
+#include "Common/ProviderLock.h"
 #include "GameplayProvider.h"
 #include "AnimationProvider.h"
 #include "Insights/ViewModels/TimingTrackViewport.h"
@@ -31,7 +32,7 @@ void FSkeletalMeshCurvesTrack::AddAllSeries()
 	bool bFirstSeries = AllSeries.Num() == 0;
 	if(AnimationProvider)
 	{
-		TraceServices::FAnalysisSessionReadScope SessionReadScope(SharedData.GetAnalysisSession());
+		TraceServices::FProviderReadScopeLock ProviderReadScope(*AnimationProvider);
 
 		AnimationProvider->EnumerateSkeletalMeshCurveIds(GetGameplayTrack().GetObjectId(), [this, &bFirstSeries, &AnimationProvider](uint32 InCurveId)
 		{
@@ -77,13 +78,12 @@ bool FSkeletalMeshCurvesTrack::UpdateSeriesBounds(FGameplayGraphSeries& InSeries
 
 	if(AnimationProvider)
 	{
-		TraceServices::FAnalysisSessionReadScope SessionReadScope(SharedData.GetAnalysisSession());
-
 		FSkeletalMeshCurveSeries& CurveSeries = *static_cast<FSkeletalMeshCurveSeries*>(&InSeries);
 
 		CurveSeries.CurrentMin = 0.0f;
 		CurveSeries.CurrentMax = 0.0f;
 
+		TraceServices::FProviderReadScopeLock ProviderReadScope(*AnimationProvider);
 		AnimationProvider->ReadSkeletalMeshPoseTimeline(GetGameplayTrack().GetObjectId(), [&bFoundEvents, &InViewport, &CurveSeries, &AnimationProvider](const FAnimationProvider::SkeletalMeshPoseTimeline& InTimeline, bool bInHasCurves)
 		{
 			InTimeline.EnumerateEvents(InViewport.GetStartTime(), InViewport.GetEndTime(), [&bFoundEvents, &CurveSeries, &AnimationProvider](double InStartTime, double InEndTime, uint32 InDepth, const FSkeletalMeshPoseMessage& InMessage)
@@ -111,12 +111,11 @@ void FSkeletalMeshCurvesTrack::UpdateSeries(FGameplayGraphSeries& InSeries, cons
 
 	if(AnimationProvider)
 	{
-		TraceServices::FAnalysisSessionReadScope SessionReadScope(SharedData.GetAnalysisSession());
-
 		FSkeletalMeshCurveSeries& CurveSeries = *static_cast<FSkeletalMeshCurveSeries*>(&InSeries);
 
 		FGraphTrackBuilder Builder(*this, CurveSeries, InViewport);
 
+		TraceServices::FProviderReadScopeLock ProviderReadScope(*AnimationProvider);
 		AnimationProvider->ReadSkeletalMeshPoseTimeline(GetGameplayTrack().GetObjectId(), [this, &Builder, &InViewport, &CurveSeries, &AnimationProvider](const FAnimationProvider::SkeletalMeshPoseTimeline& InTimeline, bool bInHasCurves)
 		{
 			int32 FrameCounter = 0;
@@ -153,16 +152,13 @@ void FSkeletalMeshCurvesTrack::InitTooltip(FTooltipDrawState& Tooltip, const ITi
 
 		Tooltip.AddTitle(GetName());
 
+		const FAnimationProvider* AnimationProvider = SharedData.GetAnalysisSession().ReadProvider<FAnimationProvider>(FAnimationProvider::ProviderName);
+		if(AnimationProvider)
 		{
-			TraceServices::FAnalysisSessionReadScope SessionReadScope(SharedData.GetAnalysisSession());
+			TraceServices::FProviderReadScopeLock ProviderReadScope(*AnimationProvider);
+			const TCHAR* CurveName = AnimationProvider->GetName(StaticCastSharedRef<const FSkeletalMeshCurveSeries>(GraphTrackEvent.GetSeries())->CurveId);
 
-			const FAnimationProvider* AnimationProvider = SharedData.GetAnalysisSession().ReadProvider<FAnimationProvider>(FAnimationProvider::ProviderName);
-			if(AnimationProvider)
-			{
-				const TCHAR* CurveName = AnimationProvider->GetName(StaticCastSharedRef<const FSkeletalMeshCurveSeries>(GraphTrackEvent.GetSeries())->CurveId);
-
-				Tooltip.AddNameValueTextLine(LOCTEXT("CurveName", "Curve").ToString(), CurveName);
-			}
+			Tooltip.AddNameValueTextLine(LOCTEXT("CurveName", "Curve").ToString(), CurveName);
 		}
 
 		Tooltip.AddNameValueTextLine(LOCTEXT("EventTime", "Time").ToString(), FText::AsNumber(InFoundStartTime).ToString());
@@ -195,7 +191,7 @@ void FSkeletalMeshCurvesTrack::FindSkeletalMeshPoseMessage(const FTimingEventSea
 
 			if(AnimationProvider)
 			{
-				TraceServices::FAnalysisSessionReadScope SessionReadScope(SharedData.GetAnalysisSession());
+				TraceServices::FProviderReadScopeLock ProviderReadScope(*AnimationProvider);
 
 				AnimationProvider->ReadSkeletalMeshPoseTimeline(GetGameplayTrack().GetObjectId(), [&InContext](const FAnimationProvider::SkeletalMeshPoseTimeline& InTimeline, bool bInHasCurves)
 				{
@@ -216,11 +212,10 @@ void FSkeletalMeshCurvesTrack::FindSkeletalMeshPoseMessage(const FTimingEventSea
 
 void FSkeletalMeshCurvesTrack::GetVariantsAtFrame(const TraceServices::FFrame& InFrame, TArray<TSharedRef<FVariantTreeNode>>& OutVariants) const
 {
-	const FGameplayProvider* GameplayProvider = SharedData.GetAnalysisSession().ReadProvider<FGameplayProvider>(FGameplayProvider::ProviderName);
 	const FAnimationProvider* AnimationProvider = SharedData.GetAnalysisSession().ReadProvider<FAnimationProvider>(FAnimationProvider::ProviderName);
-	if(GameplayProvider && AnimationProvider)
+	if (AnimationProvider)
 	{
-		TraceServices::FAnalysisSessionReadScope SessionReadScope(SharedData.GetAnalysisSession());
+		TraceServices::FProviderReadScopeLock AnimationProviderReadScope(*AnimationProvider);
 
 		TSharedRef<FVariantTreeNode> Header = OutVariants.Add_GetRef(FVariantTreeNode::MakeHeader(FText::FromString(GetName()), 0));
 
