[RewindDebugger]

diff --git a/Engine/Plugins/Animation/GameplayInsights/Source/GameplayInsights/Private/SkeletalMeshPoseTrack.cpp b/Engine/Plugins/Animation/GameplayInsights/Source/GameplayInsights/Private/SkeletalMeshPoseTrack.cpp
--- a/Engine/Plugins/Animation/GameplayInsights/Source/GameplayInsights/Private/SkeletalMeshPoseTrack.cpp
+++ b/Engine/Plugins/Animation/GameplayInsights/Source/GameplayInsights/Private/SkeletalMeshPoseTrack.cpp
@@ -2,6 +2,7 @@
 
 #include "SkeletalMeshPoseTrack.h"
 #include "AssetRegistry/AssetData.h"
+#include "Common/ProviderLock.h"
 #include "GameplayProvider.h"
 #include "AnimationProvider.h"
 #include "Insights/ViewModels/TimingTrackViewport.h"
@@ -72,7 +73,7 @@ void FSkeletalMeshPoseTrack::BuildDrawState(ITimingEventsTrackDrawStateBuilder&
 
 	if(AnimationProvider)
 	{
-		TraceServices::FAnalysisSessionReadScope SessionReadScope(SharedData.GetAnalysisSession());
+		TraceServices::FProviderReadScopeLock ProviderReadScope(*AnimationProvider);
 
 		AnimationProvider->ReadSkeletalMeshPoseTimeline(GetGameplayTrack().GetObjectId(), [&Context, &Builder](const FAnimationProvider::SkeletalMeshPoseTimeline& InTimeline, bool bInHasCurves)
 		{
@@ -103,20 +104,25 @@ void FSkeletalMeshPoseTrack::InitTooltip(FTooltipDrawState& Tooltip, const ITimi
 
 		Tooltip.AddNameValueTextLine(LOCTEXT("EventTime", "Time").ToString(), FText::AsNumber(InFoundStartTime).ToString());
 
+		FString PathName;
 		{
-			TraceServices::FAnalysisSessionReadScope SessionReadScope(SharedData.GetAnalysisSession());
-
 			const FGameplayProvider* GameplayProvider = SharedData.GetAnalysisSession().ReadProvider<FGameplayProvider>(FGameplayProvider::ProviderName);
+			TraceServices::FProviderReadScopeLock ProviderReadScope(*GameplayProvider);
 			const FObjectInfo* SkeletalMeshObjectInfo = GameplayProvider->FindObjectInfo(InMessage.MeshId);
-			if(SkeletalMeshObjectInfo != nullptr)
+			if (SkeletalMeshObjectInfo)
 			{
-				Tooltip.AddNameValueTextLine(LOCTEXT("Mesh", "Mesh").ToString(), SkeletalMeshObjectInfo->PathName);
+				PathName = SkeletalMeshObjectInfo->PathName;
+			}
+		}
 
-				FAssetRegistryModule& AssetRegistryModule = FModuleManager::LoadModuleChecked<FAssetRegistryModule>(TEXT("AssetRegistry"));
-				if(!AssetRegistryModule.Get().GetAssetByObjectPath(FSoftObjectPath(SkeletalMeshObjectInfo->PathName)).IsValid())
-				{
-					Tooltip.AddTextLine(LOCTEXT("MeshNotFound", "Mesh not found").ToString(), FLinearColor::Red);
-				}
+		if (!PathName.IsEmpty())
+		{
+			Tooltip.AddNameValueTextLine(LOCTEXT("Mesh", "Mesh").ToString(), PathName);
+
+			FAssetRegistryModule& AssetRegistryModule = FModuleManager::LoadModuleChecked<FAssetRegistryModule>(TEXT("AssetRegistry"));
+			if(!AssetRegistryModule.Get().GetAssetByObjectPath(FSoftObjectPath(PathName)).IsValid())
+			{
+				Tooltip.AddTextLine(LOCTEXT("MeshNotFound", "Mesh not found").ToString(), FLinearColor::Red);
 			}
 		}
 
@@ -152,7 +158,7 @@ void FSkeletalMeshPoseTrack::FindSkeletalMeshPoseMessage(const FTimingEventSearc
 
 			if(AnimationProvider)
 			{
-				TraceServices::FAnalysisSessionReadScope SessionReadScope(SharedData.GetAnalysisSession());
+				TraceServices::FProviderReadScopeLock ProviderReadScope(*AnimationProvider);
 
 				AnimationProvider->ReadSkeletalMeshPoseTimeline(GetGameplayTrack().GetObjectId(), [&InContext](const FAnimationProvider::SkeletalMeshPoseTimeline& InTimeline, bool bInHasCurves)
 				{
@@ -215,40 +221,52 @@ void FSkeletalMeshPoseTrack::BuildContextMenu(FMenuBuilder& MenuBuilder)
 
 	if(GameplayProvider)
 	{
-		TraceServices::FAnalysisSessionReadScope SessionReadScope(SharedData.GetAnalysisSession());
-	
-		const FObjectInfo* ComponentObjectInfo = GameplayProvider->FindObjectInfo(GetGameplayTrack().GetAssociatedObjectId());
-		if(ComponentObjectInfo != nullptr)
+		FString ActorName;
+		RewindDebugger::FObjectId ActorObjectId;
 		{
-			// @FIXME: Outer does always equal owning actor, although does in nearly all cases with skeletal mesh components
-			const FObjectInfo* ActorObjectInfo = GameplayProvider->FindObjectInfo(ComponentObjectInfo->GetOuterId());
-			if(ActorObjectInfo != nullptr)
+			TraceServices::FProviderReadScopeLock ProviderReadScope(*GameplayProvider);
+
+			const FObjectInfo* ComponentObjectInfo = GameplayProvider->FindObjectInfo(GetGameplayTrack().GetAssociatedObjectId());
+			if(ComponentObjectInfo != nullptr)
 			{
-				MenuBuilder.BeginSection(TEXT("DrawingSection"), FText::Format(LOCTEXT("DrawingActor", "Drawing ({0})"), FText::FromString(ActorObjectInfo->Name)));
+				// @FIXME: Outer does always equal owning actor, although does in nearly all cases with skeletal mesh components
+				const FObjectInfo* ActorObjectInfo = GameplayProvider->FindObjectInfo(ComponentObjectInfo->GetOuterId());
+				if(ActorObjectInfo != nullptr)
+				{
+					ActorName = ActorObjectInfo->Name;
+					ActorObjectId = ActorObjectInfo->GetId();
+				}
+			}
+		}
+
+		if(ActorObjectId.IsSet())
+		{
+			{
+				MenuBuilder.BeginSection(TEXT("DrawingSection"), FText::Format(LOCTEXT("DrawingActor", "Drawing ({0})"), FText::FromString(ActorName)));
 				{
 					MenuBuilder.AddMenuEntry(
 						LOCTEXT("ToggleDrawPoseActor", "Draw Pose for Actor"),
 						LOCTEXT("ToggleDrawPoseActor_Tooltip", "Draw the poses in this track and all other tracks for the current actor"),
 						FSlateIcon(),
 						FUIAction(
-							FExecuteAction::CreateLambda([this, ActorObjectInfo, GameplayProvider]()
+							FExecuteAction::CreateLambda([this, ActorObjectId, GameplayProvider]()
 							{
-								TraceServices::FAnalysisSessionReadScope SessionReadScope(SharedData.GetAnalysisSession());
+								TraceServices::FProviderReadScopeLock ProviderReadScope(*GameplayProvider);
 
 								bool bSetDrawPose = true;
-								SharedData.EnumerateSkeletalMeshPoseTracks([GameplayProvider, ActorObjectInfo, &bSetDrawPose](const TSharedRef<FSkeletalMeshPoseTrack>& InTrack)
+								SharedData.EnumerateSkeletalMeshPoseTracks([GameplayProvider, ActorObjectId, &bSetDrawPose](const TSharedRef<FSkeletalMeshPoseTrack>& InTrack)
 								{
 									const FObjectInfo* OtherComponentObjectInfo = GameplayProvider->FindObjectInfo(InTrack->GetGameplayTrack().GetAssociatedObjectId());
-									if(OtherComponentObjectInfo->GetOuterId() == ActorObjectInfo->GetId())
+									if(OtherComponentObjectInfo && OtherComponentObjectInfo->GetOuterId() == ActorObjectId)
 									{
 										bSetDrawPose &= InTrack->bDrawPose;
 									}
 								});
 
-								SharedData.EnumerateSkeletalMeshPoseTracks([GameplayProvider, ActorObjectInfo, &bSetDrawPose](const TSharedRef<FSkeletalMeshPoseTrack>& InTrack)
+								SharedData.EnumerateSkeletalMeshPoseTracks([GameplayProvider, ActorObjectId, &bSetDrawPose](const TSharedRef<FSkeletalMeshPoseTrack>& InTrack)
 								{
 									const FObjectInfo* OtherComponentObjectInfo = GameplayProvider->FindObjectInfo(InTrack->GetGameplayTrack().GetAssociatedObjectId());
-									if(OtherComponentObjectInfo->GetOuterId() == ActorObjectInfo->GetId())
+									if(OtherComponentObjectInfo && OtherComponentObjectInfo->GetOuterId() == ActorObjectId)
 									{
 										InTrack->bDrawPose = !bSetDrawPose;
 										InTrack->UpdateComponentVisibility();
@@ -258,15 +276,15 @@ void FSkeletalMeshPoseTrack::BuildContextMenu(FMenuBuilder& MenuBuilder)
 								SharedData.InvalidateViewports();
 							}),
 							FCanExecuteAction(),
-							FIsActionChecked::CreateLambda([this, ActorObjectInfo, GameplayProvider]()
+							FIsActionChecked::CreateLambda([this, ActorObjectId, GameplayProvider]()
 							{
-								TraceServices::FAnalysisSessionReadScope SessionReadScope(SharedData.GetAnalysisSession());
+								TraceServices::FProviderReadScopeLock ProviderReadScope(*GameplayProvider);
 
 								bool bDrawPoseSet = true;
-								SharedData.EnumerateSkeletalMeshPoseTracks([GameplayProvider, ActorObjectInfo, &bDrawPoseSet](const TSharedRef<FSkeletalMeshPoseTrack>& InTrack)
+								SharedData.EnumerateSkeletalMeshPoseTracks([GameplayProvider, ActorObjectId, &bDrawPoseSet](const TSharedRef<FSkeletalMeshPoseTrack>& InTrack)
 								{
 									const FObjectInfo* OtherComponentObjectInfo = GameplayProvider->FindObjectInfo(InTrack->GetGameplayTrack().GetAssociatedObjectId());
-									if(OtherComponentObjectInfo->GetOuterId() == ActorObjectInfo->GetId())
+									if(OtherComponentObjectInfo && OtherComponentObjectInfo->GetOuterId() == ActorObjectId)
 									{
 										bDrawPoseSet &= InTrack->bDrawPose;
 									}
@@ -283,24 +301,24 @@ void FSkeletalMeshPoseTrack::BuildContextMenu(FMenuBuilder& MenuBuilder)
 						LOCTEXT("ToggleDrawSkeletonActor_Tooltip", "Draw the skeleton for poses in this track and all other tracks for the current actor (when pose drawing is also enabled)"),
 						FSlateIcon(),
 						FUIAction(
-							FExecuteAction::CreateLambda([this, ActorObjectInfo, GameplayProvider]()
+							FExecuteAction::CreateLambda([this, ActorObjectId, GameplayProvider]()
 							{
-								TraceServices::FAnalysisSessionReadScope SessionReadScope(SharedData.GetAnalysisSession());
+								TraceServices::FProviderReadScopeLock ProviderReadScope(*GameplayProvider);
 
 								bool bSetDrawSkeleton = true;
-								SharedData.EnumerateSkeletalMeshPoseTracks([GameplayProvider, ActorObjectInfo, &bSetDrawSkeleton](const TSharedRef<FSkeletalMeshPoseTrack>& InTrack)
+								SharedData.EnumerateSkeletalMeshPoseTracks([GameplayProvider, ActorObjectId, &bSetDrawSkeleton](const TSharedRef<FSkeletalMeshPoseTrack>& InTrack)
 								{
 									const FObjectInfo* OtherComponentObjectInfo = GameplayProvider->FindObjectInfo(InTrack->GetGameplayTrack().GetAssociatedObjectId());
-									if(OtherComponentObjectInfo->GetOuterId() == ActorObjectInfo->GetId())
+									if(OtherComponentObjectInfo && OtherComponentObjectInfo->GetOuterId() == ActorObjectId)
 									{
 										bSetDrawSkeleton &= InTrack->bDrawSkeleton;
 									}
 								});
 
-								SharedData.EnumerateSkeletalMeshPoseTracks([GameplayProvider, ActorObjectInfo, &bSetDrawSkeleton](const TSharedRef<FSkeletalMeshPoseTrack>& InTrack)
+								SharedData.EnumerateSkeletalMeshPoseTracks([GameplayProvider, ActorObjectId, &bSetDrawSkeleton](const TSharedRef<FSkeletalMeshPoseTrack>& InTrack)
 								{
 									const FObjectInfo* OtherComponentObjectInfo = GameplayProvider->FindObjectInfo(InTrack->GetGameplayTrack().GetAssociatedObjectId());
-									if(OtherComponentObjectInfo->GetOuterId() == ActorObjectInfo->GetId())
+									if(OtherComponentObjectInfo && OtherComponentObjectInfo->GetOuterId() == ActorObjectId)
 									{
 										InTrack->bDrawSkeleton = !bSetDrawSkeleton;
 										InTrack->UpdateComponentVisibility();
@@ -310,15 +328,15 @@ void FSkeletalMeshPoseTrack::BuildContextMenu(FMenuBuilder& MenuBuilder)
 								SharedData.InvalidateViewports();
 							}),
 							FCanExecuteAction(),
-							FIsActionChecked::CreateLambda([this, ActorObjectInfo, GameplayProvider]()
+							FIsActionChecked::CreateLambda([this, ActorObjectId, GameplayProvider]()
 							{
-								TraceServices::FAnalysisSessionReadScope SessionReadScope(SharedData.GetAnalysisSession());
+								TraceServices::FProviderReadScopeLock ProviderReadScope(*GameplayProvider);
 
 								bool bDrawSkeletonSet = true;
-								SharedData.EnumerateSkeletalMeshPoseTracks([GameplayProvider, ActorObjectInfo, &bDrawSkeletonSet](const TSharedRef<FSkeletalMeshPoseTrack>& InTrack)
+								SharedData.EnumerateSkeletalMeshPoseTracks([GameplayProvider, ActorObjectId, &bDrawSkeletonSet](const TSharedRef<FSkeletalMeshPoseTrack>& InTrack)
 								{
 									const FObjectInfo* OtherComponentObjectInfo = GameplayProvider->FindObjectInfo(InTrack->GetGameplayTrack().GetAssociatedObjectId());
-									if(OtherComponentObjectInfo->GetOuterId() == ActorObjectInfo->GetId())
+									if(OtherComponentObjectInfo && OtherComponentObjectInfo->GetOuterId() == ActorObjectId)
 									{
 										bDrawSkeletonSet &= InTrack->bDrawSkeleton;
 									}
@@ -377,14 +395,14 @@ void FSkeletalMeshPoseTrack::DrawPoses(UWorld* InWorld, double InTime, double In
 
 		if(AnimationProvider && GameplayProvider)
 		{
-			TraceServices::FAnalysisSessionReadScope SessionReadScope(SharedData.GetAnalysisSession());
-
 			FWorldComponentCache& CacheForWorld = GetWorldCache(InWorld);
 			if(CacheForWorld.Component)
 			{
 				CacheForWorld.Component->SetVisibility(false);
 			}
 
+			TraceServices::FProviderReadScopeLock AnimationProviderReadScope(*AnimationProvider);
+			TraceServices::FProviderReadScopeLock GameplayProviderReadScope(*GameplayProvider);
 			AnimationProvider->ReadSkeletalMeshPoseTimeline(GetGameplayTrack().GetObjectId(), [this, &CacheForWorld, &AnimationProvider, &GameplayProvider, &InFrameStartTime, &InFrameEndTime](const FAnimationProvider::SkeletalMeshPoseTimeline& InTimeline, bool bInHasCurves)
 			{
 				InTimeline.EnumerateEvents(InFrameStartTime, InFrameEndTime, [this, &CacheForWorld, &AnimationProvider, &GameplayProvider, &InFrameStartTime, &InFrameEndTime](double InStartTime, double InEndTime, uint32 InDepth, const FSkeletalMeshPoseMessage& InMessage)
diff --git a/Engine/Plugins/Animation/GameplayInsights/Source/GameplayInsights/Private/VariantTreeNode.cpp b/Engine/Plugins/Animation/GameplayInsights/Source/GameplayInsights/Private/VariantTreeNode.cpp
--- a/Engine/Plugins/Animation/GameplayInsights/Source/GameplayInsights/Private/VariantTreeNode.cpp
+++ b/Engine/Plugins/Animation/GameplayInsights/Source/GameplayInsights/Private/VariantTreeNode.cpp
@@ -1,6 +1,7 @@
 // Copyright Epic Games, Inc. All Rights Reserved.
 
 #include "VariantTreeNode.h"
+#include "Common/ProviderLock.h"
 #include "TraceServices/Model/AnalysisSession.h"
 #include "GameplayProvider.h"
 #include "AnimationProvider.h"
@@ -28,7 +29,7 @@ FString FVariantTreeNode::GetValueAsString(const TraceServices::IAnalysisSession
 			const FGameplayProvider* GameplayProvider = InAnalysisSession.ReadProvider<FGameplayProvider>(FGameplayProvider::ProviderName);
 			if(GameplayProvider)
 			{
-				TraceServices::FAnalysisSessionReadScope SessionReadScope(InAnalysisSession);
+				TraceServices::FProviderReadScopeLock ProviderReadScope(*GameplayProvider);
 
 				const FObjectInfo& ObjectInfo = GameplayProvider->GetObjectInfo(Value.Object.Value);
 				return ObjectInfo.PathName;
@@ -41,7 +42,7 @@ FString FVariantTreeNode::GetValueAsString(const TraceServices::IAnalysisSession
 			const FGameplayProvider* GameplayProvider = InAnalysisSession.ReadProvider<FGameplayProvider>(FGameplayProvider::ProviderName);
 			if(GameplayProvider)
 			{
-				TraceServices::FAnalysisSessionReadScope SessionReadScope(InAnalysisSession);
+				TraceServices::FProviderReadScopeLock ProviderReadScope(*GameplayProvider);
 
 				const FClassInfo& ClassInfo = GameplayProvider->GetClassInfo(Value.Class.Value);
 				return ClassInfo.PathName;
@@ -54,7 +55,7 @@ FString FVariantTreeNode::GetValueAsString(const TraceServices::IAnalysisSession
 			const FAnimationProvider* AnimationProvider = InAnalysisSession.ReadProvider<FAnimationProvider>(FAnimationProvider::ProviderName);
 			if (AnimationProvider)
 			{
-				TraceServices::FAnalysisSessionReadScope SessionReadScope(InAnalysisSession);
+				TraceServices::FProviderReadScopeLock ProviderReadScope(*AnimationProvider);
 
 				const FAnimNodeInfo* AnimNodeInfo = AnimationProvider->FindAnimNodeInfo(Value.AnimNode.Value, Value.AnimNode.AnimInstanceId);
 				return AnimNodeInfo ? AnimNodeInfo->Name : TEXT("");
diff --git a/Engine/Plugins/Animation/GameplayInsights/Source/GameplayInsights/Public/IGameplayProvider.h b/Engine/Plugins/Animation/GameplayInsights/Source/GameplayInsights/Public/IGameplayProvider.h
--- a/Engine/Plugins/Animation/GameplayInsights/Source/GameplayInsights/Public/IGameplayProvider.h
+++ b/Engine/Plugins/Animation/GameplayInsights/Source/GameplayInsights/Public/IGameplayProvider.h
@@ -183,7 +183,7 @@ struct FSlateIcon;
 // Delegate fired when an object receives an end play event
 DECLARE_MULTICAST_DELEGATE_ThreeParams(FOnObjectEndPlay, uint64 /*ObjectId*/, double /*Time*/, const FObjectInfo& /*ObjectInfo*/);
 
-class IGameplayProvider : public TraceServices::IProvider
+class IGameplayProvider : public TraceServices::IProvider, public TraceServices::IEditableProvider
 {
 public:
 	typedef TraceServices::ITimeline<FObjectEventMessage> ObjectEventsTimeline;
diff --git a/Engine/Plugins/Animation/GameplayInsights/Source/RewindDebugger/Private/AnimInstanceHelpers.cpp b/Engine/Plugins/Animation/GameplayInsights/Source/RewindDebugger/Private/AnimInstanceHelpers.cpp
--- a/Engine/Plugins/Animation/GameplayInsights/Source/RewindDebugger/Private/AnimInstanceHelpers.cpp
+++ b/Engine/Plugins/Animation/GameplayInsights/Source/RewindDebugger/Private/AnimInstanceHelpers.cpp
@@ -1,6 +1,7 @@
 ﻿// Copyright Epic Games, Inc. All Rights Reserved.
 
 #include "AnimInstanceHelpers.h"
+#include "Common/ProviderLock.h"
 #include "Modules/ModuleManager.h"
 #include "ObjectTrace.h"
 #include "IGameplayProvider.h"
@@ -19,48 +20,57 @@
 
 static bool OpenAnimBlueprintAndAttachDebugger(const TraceServices::IAnalysisSession* Session, uint64 ObjectId)
 {
-	TraceServices::FAnalysisSessionReadScope SessionReadScope(*Session);
-
 	if (const IGameplayProvider* GameplayProvider = Session->ReadProvider<IGameplayProvider>("GameplayProvider"))
 	{
-		if (const FObjectInfo* ObjectInfo = GameplayProvider->FindObjectInfo(ObjectId))
+		// Cache the class path name under a tight lock scope
+		FString ClassPathName;
 		{
-			if (const FClassInfo* ClassInfo = GameplayProvider->FindClassInfo(ObjectInfo->ClassId))
+			TraceServices::FProviderReadScopeLock ProviderReadScope(*GameplayProvider);
+			if (const FObjectInfo* ObjectInfo = GameplayProvider->FindObjectInfo(ObjectId))
 			{
-				if (UAnimBlueprintGeneratedClass* InstanceClass = TSoftObjectPtr<UAnimBlueprintGeneratedClass>(FSoftObjectPath(ClassInfo->PathName)).LoadSynchronous())
+				if (const FClassInfo* ClassInfo = GameplayProvider->FindClassInfo(ObjectInfo->ClassId))
 				{
-					if(UAnimBlueprint* AnimBlueprint = Cast<UAnimBlueprint>(InstanceClass->ClassGeneratedBy))
-					{
-						GEditor->GetEditorSubsystem<UAssetEditorSubsystem>()->OpenEditorForAsset(AnimBlueprint);
-						
-						UObject* SelectedInstance = nullptr;
+					ClassPathName = ClassInfo->PathName;
+				}
+			}
+		}
+
+		// Perform asset loading and editor operations outside the lock
+		if (!ClassPathName.IsEmpty())
+		{
+			if (UAnimBlueprintGeneratedClass* InstanceClass = TSoftObjectPtr<UAnimBlueprintGeneratedClass>(FSoftObjectPath(ClassPathName)).LoadSynchronous())
+			{
+				if(UAnimBlueprint* AnimBlueprint = Cast<UAnimBlueprint>(InstanceClass->ClassGeneratedBy))
+				{
+					GEditor->GetEditorSubsystem<UAssetEditorSubsystem>()->OpenEditorForAsset(AnimBlueprint);
+
+					UObject* SelectedInstance = nullptr;
 #if OBJECT_TRACE_ENABLED
-						SelectedInstance = FObjectTrace::GetObjectFromId(ObjectId);
+					SelectedInstance = FObjectTrace::GetObjectFromId(ObjectId);
 #endif
-						if (SelectedInstance == nullptr)
+					if (SelectedInstance == nullptr)
+					{
+						if (FRewindDebuggerAnimation* RewindDebuggerAnimation = FRewindDebuggerAnimation::GetInstance())
 						{
-							if (FRewindDebuggerAnimation* RewindDebuggerAnimation = FRewindDebuggerAnimation::GetInstance())
-							{
-								SelectedInstance = RewindDebuggerAnimation->GetDebugAnimInstance(ObjectId);
-							}
+							SelectedInstance = RewindDebuggerAnimation->GetDebugAnimInstance(ObjectId);
 						}
+					}
 
-						if(SelectedInstance)
-						{
-							AnimBlueprint->SetObjectBeingDebugged(SelectedInstance);
-						}
+					if(SelectedInstance)
+					{
+						AnimBlueprint->SetObjectBeingDebugged(SelectedInstance);
+					}
 
-						if(IAnimationBlueprintEditor* AnimBlueprintEditor = static_cast<IAnimationBlueprintEditor*>(GEditor->GetEditorSubsystem<UAssetEditorSubsystem>()->FindEditorForAsset(AnimBlueprint, true)))
+					if(IAnimationBlueprintEditor* AnimBlueprintEditor = static_cast<IAnimationBlueprintEditor*>(GEditor->GetEditorSubsystem<UAssetEditorSubsystem>()->FindEditorForAsset(AnimBlueprint, true)))
+					{
+						// navigate the opened editor to the AnimGraph
+						if (AnimBlueprint->FunctionGraphs.Num() > 0)
 						{
-							// navigate the opened editor to the AnimGraph
-							if (AnimBlueprint->FunctionGraphs.Num() > 0)
-							{
-								AnimBlueprintEditor->JumpToHyperlink(AnimBlueprint->FunctionGraphs[0]);
-							}
+							AnimBlueprintEditor->JumpToHyperlink(AnimBlueprint->FunctionGraphs[0]);
 						}
-
-						return true;
 					}
+
+					return true;
 				}
 			}
 		}
diff --git a/Engine/Plugins/Animation/GameplayInsights/Source/RewindDebugger/Private/BlueprintDoubleClickHandler.cpp b/Engine/Plugins/Animation/GameplayInsights/Source/RewindDebugger/Private/BlueprintDoubleClickHandler.cpp
--- a/Engine/Plugins/Animation/GameplayInsights/Source/RewindDebugger/Private/BlueprintDoubleClickHandler.cpp
+++ b/Engine/Plugins/Animation/GameplayInsights/Source/RewindDebugger/Private/BlueprintDoubleClickHandler.cpp
@@ -3,6 +3,7 @@
 #include "BlueprintDoubleClickHandler.h"
 
 #include "IGameplayProvider.h"
+#include "Common/ProviderLock.h"
 #include "ObjectTrace.h"
 #include "Engine/Blueprint.h"
 #include "Engine/BlueprintGeneratedClass.h"
@@ -14,27 +15,36 @@
 /** Open blueprint editor given an object id. If no editor found, it will use the properties editor. */
 static bool OpenBlueprintAndAttachDebugger(const TraceServices::IAnalysisSession* Session, uint64 ObjectId)
 {
-	TraceServices::FAnalysisSessionReadScope SessionReadScope(*Session);
-
 	if (const IGameplayProvider* GameplayProvider = Session->ReadProvider<IGameplayProvider>("GameplayProvider"))
 	{
-		if (const FObjectInfo* ObjectInfo = GameplayProvider->FindObjectInfo(ObjectId))
+		// Cache the class path name under a tight lock scope
+		FString ClassPathName;
 		{
-			if (const FClassInfo* ClassInfo = GameplayProvider->FindClassInfo(ObjectInfo->ClassId))
+			TraceServices::FProviderReadScopeLock ProviderReadScope(*GameplayProvider);
+			if (const FObjectInfo* ObjectInfo = GameplayProvider->FindObjectInfo(ObjectId))
 			{
-				if (const UBlueprintGeneratedClass* InstanceClass = TSoftObjectPtr<UBlueprintGeneratedClass>(FSoftObjectPath(ClassInfo->PathName)).LoadSynchronous())
+				if (const FClassInfo* ClassInfo = GameplayProvider->FindClassInfo(ObjectInfo->ClassId))
 				{
-					if (UBlueprint* Blueprint = Cast<UBlueprint>(InstanceClass->ClassGeneratedBy))
-					{
-						GEditor->GetEditorSubsystem<UAssetEditorSubsystem>()->OpenEditorForAsset(Blueprint);
-						
-						if (UObject* SelectedInstance = FObjectTrace::GetObjectFromId(ObjectId))
-						{
-							Blueprint->SetObjectBeingDebugged(SelectedInstance);
-						}
+					ClassPathName = ClassInfo->PathName;
+				}
+			}
+		}
 
-						return true;
+		// Perform asset loading and editor operations outside the lock
+		if (!ClassPathName.IsEmpty())
+		{
+			if (const UBlueprintGeneratedClass* InstanceClass = TSoftObjectPtr<UBlueprintGeneratedClass>(FSoftObjectPath(ClassPathName)).LoadSynchronous())
+			{
+				if (UBlueprint* Blueprint = Cast<UBlueprint>(InstanceClass->ClassGeneratedBy))
+				{
+					GEditor->GetEditorSubsystem<UAssetEditorSubsystem>()->OpenEditorForAsset(Blueprint);
+
+					if (UObject* SelectedInstance = FObjectTrace::GetObjectFromId(ObjectId))
+					{
+						Blueprint->SetObjectBeingDebugged(SelectedInstance);
 					}
+
+					return true;
 				}
 			}
 		}
diff --git a/Engine/Plugins/Animation/GameplayInsights/Source/RewindDebugger/Private/RewindDebugger.cpp b/Engine/Plugins/Animation/GameplayInsights/Source/RewindDebugger/Private/RewindDebugger.cpp
--- a/Engine/Plugins/Animation/GameplayInsights/Source/RewindDebugger/Private/RewindDebugger.cpp
+++ b/Engine/Plugins/Animation/GameplayInsights/Source/RewindDebugger/Private/RewindDebugger.cpp
@@ -2,6 +2,7 @@
 
 #include "RewindDebugger.h"
 
+#include "Common/ProviderLock.h"
 #include "DesktopPlatformModule.h"
 #include "Editor.h"
 #include "Editor/UnrealEdEngine.h"
@@ -310,10 +311,9 @@ void FRewindDebugger::GetTargetObjectIds(TArray<RewindDebugger::FObjectId>& OutT
 
 	if (const TraceServices::IAnalysisSession* Session = GetAnalysisSession())
 	{
-		TraceServices::FAnalysisSessionReadScope SessionReadScope(*Session);
-
 		if (const IGameplayProvider* GameplayProvider = Session->ReadProvider<IGameplayProvider>("GameplayProvider"))
 		{
+			TraceServices::FProviderReadScopeLock ProviderReadScope(*GameplayProvider);
 			GameplayProvider->EnumerateObjects(CurrentTraceRange.GetLowerBoundValue(), CurrentTraceRange.GetUpperBoundValue(), [this, &OutTargetObjectIds](const FObjectInfo& InObjectInfo)
 				{
 					if (RootObjectName.Get() == InObjectInfo.Name)
@@ -813,22 +813,31 @@ void FRewindDebugger::Step(const int32 InNumberOfFrames)
 
 		if (const TraceServices::IAnalysisSession* Session = GetAnalysisSession())
 		{
-			TraceServices::FAnalysisSessionReadScope SessionReadScope(*Session);
 			if (const IGameplayProvider* GameplayProvider = Session->ReadProvider<IGameplayProvider>("GameplayProvider"))
 			{
-				if (const IGameplayProvider::RecordingInfoTimeline* Recording = GameplayProvider->GetRecordingInfo(RecordingIndex))
+				bool bEventFound = false;
+				double ElapsedTime = 0;
 				{
-					const uint64 EventCount = Recording->GetEventCount();
-
-					if (EventCount > 0)
+					TraceServices::FProviderReadScopeLock ProviderReadScope(*GameplayProvider);
+					if (const IGameplayProvider::RecordingInfoTimeline* Recording = GameplayProvider->GetRecordingInfo(RecordingIndex))
 					{
-						ScrubTimeInformation.FrameIndex = FMath::Clamp<int64>(ScrubTimeInformation.FrameIndex + InNumberOfFrames, 0, (int64)EventCount - 1);
-						const FRecordingInfoMessage& Event = Recording->GetEvent(ScrubTimeInformation.FrameIndex);
+						const uint64 EventCount = Recording->GetEventCount();
 
-						SetCurrentScrubTime(Event.ElapsedTime, EAutoScrollChange::Disable);
-						TrackCursorDelegate.ExecuteIfBound(/*bReverse*/false);
+						if (EventCount > 0)
+						{
+							ScrubTimeInformation.FrameIndex = FMath::Clamp<int64>(ScrubTimeInformation.FrameIndex + InNumberOfFrames, 0, (int64)EventCount - 1);
+							const FRecordingInfoMessage& Event = Recording->GetEvent(ScrubTimeInformation.FrameIndex);
+							bEventFound = true;
+							ElapsedTime = Event.ElapsedTime;
+						}
 					}
 				}
+
+				if (bEventFound)
+				{
+					SetCurrentScrubTime(ElapsedTime, EAutoScrollChange::Disable);
+					TrackCursorDelegate.ExecuteIfBound(/*bReverse*/false);
+				}
 			}
 		}
 	}
@@ -933,12 +942,10 @@ void FRewindDebugger::SetCurrentScrubTime(const double Time, const EAutoScrollCh
 void FRewindDebugger::GetScrubTimeInformation(double InDebugTime, FScrubTimeInformation& InOutTimeInformation, uint16 InRecordingIndex, const TraceServices::IAnalysisSession* AnalysisSession)
 {
 	const IGameplayProvider* GameplayProvider = AnalysisSession->ReadProvider<IGameplayProvider>("GameplayProvider");
-	const IAnimationProvider* AnimationProvider = AnalysisSession->ReadProvider<IAnimationProvider>("AnimationProvider");
 
-	if (GameplayProvider && AnimationProvider)
+	if (GameplayProvider)
 	{
-		TraceServices::FAnalysisSessionReadScope SessionReadScope(*AnalysisSession);
-
+		TraceServices::FProviderReadScopeLock ProviderReadScope(*GameplayProvider);
 		if (const IGameplayProvider::RecordingInfoTimeline* Recording = GameplayProvider->GetRecordingInfo(InRecordingIndex))
 		{
 			const uint64 EventCount = Recording->GetEventCount();
@@ -1309,7 +1316,7 @@ void FRewindDebugger::Tick(float DeltaTime)
 
 		if (const IGameplayProvider* GameplayProvider = Session->ReadProvider<IGameplayProvider>("GameplayProvider"))
 		{
-			TraceServices::FAnalysisSessionReadScope SessionReadScope(*Session);
+			TraceServices::FProviderReadScopeLock ProviderReadScope(*GameplayProvider);
 
 			RecordingDurationValue = GameplayProvider->GetRecordingDuration();
 
@@ -1521,9 +1528,9 @@ void FRewindDebugger::MakeOtherWorldsMenu(UToolMenu* Menu)
 
 	if (const TraceServices::IAnalysisSession* Session = RewindDebugger->GetAnalysisSession())
 	{
-		TraceServices::FAnalysisSessionReadScope SessionReadScope(*Session);
 		const IGameplayProvider* GameplayProvider = Session->ReadProvider<IGameplayProvider>("GameplayProvider");
 
+		TraceServices::FProviderReadScopeLock ProviderReadScope(*GameplayProvider);
 		GameplayProvider->EnumerateWorlds([GameplayProvider, &Section](const FWorldInfo& WorldInfo)
 			{
 				const FObjectInfo* ObjectInfo = GameplayProvider->FindObjectInfo(WorldInfo.Id);
@@ -1609,9 +1616,8 @@ void FRewindDebugger::MakeWorldsMenu(UToolMenu* Menu)
 
 	if (const TraceServices::IAnalysisSession* Session = RewindDebugger->GetAnalysisSession())
 	{
-		TraceServices::FAnalysisSessionReadScope SessionReadScope(*Session);
 		const IGameplayProvider* GameplayProvider = Session->ReadProvider<IGameplayProvider>("GameplayProvider");
-
+		TraceServices::FProviderReadScopeLock ProviderReadScope(*GameplayProvider);
 		GameplayProvider->EnumerateWorlds([GameplayProvider, &GameWorldsSection, &OtherWorldsSection, &ServerWorldsSection](const FWorldInfo& WorldInfo)
 			{
 				const FObjectInfo* ObjectInfo = GameplayProvider->FindObjectInfo(WorldInfo.Id);
@@ -1947,9 +1953,9 @@ TSharedPtr<SWidget> FRewindDebugger::BuildTrackContextMenu() const
 		// build a list of class hierarchy names to make it easier for extensions to enable menu entries by type
 		if (const TraceServices::IAnalysisSession* Session = GetAnalysisSession())
 		{
-			TraceServices::FAnalysisSessionReadScope SessionReadScope(*Session);
-
 			const IGameplayProvider* GameplayProvider = Session->ReadProvider<IGameplayProvider>("GameplayProvider");
+			TraceServices::FProviderReadScopeLock ProviderReadScope(*GameplayProvider);
+
 			const FObjectInfo& ObjectInfo = GameplayProvider->GetObjectInfo(SelectedTrack->GetAssociatedObjectId());
 			uint64 ClassId = ObjectInfo.ClassId;
 			while (ClassId != 0)
diff --git a/Engine/Plugins/Animation/GameplayInsights/Source/RewindDebugger/Private/RewindDebuggerAnimation.cpp b/Engine/Plugins/Animation/GameplayInsights/Source/RewindDebugger/Private/RewindDebuggerAnimation.cpp
--- a/Engine/Plugins/Animation/GameplayInsights/Source/RewindDebugger/Private/RewindDebuggerAnimation.cpp
+++ b/Engine/Plugins/Animation/GameplayInsights/Source/RewindDebugger/Private/RewindDebuggerAnimation.cpp
@@ -3,6 +3,7 @@
 #include "RewindDebuggerAnimation.h"
 #include "Animation/AnimBlueprint.h"
 #include "Animation/AnimBlueprintGeneratedClass.h"
+#include "Common/ProviderLock.h"
 #include "Editor.h"
 #include "Engine/PoseWatch.h"
 #include "IAnimationProvider.h"
@@ -283,14 +284,18 @@ void FRewindDebuggerAnimation::Update(float DeltaTime, IRewindDebugger* RewindDe
 
 	if (const TraceServices::IAnalysisSession* Session = RewindDebugger->GetAnalysisSession())
 	{
-		TraceServices::FAnalysisSessionReadScope SessionReadScope(*Session);
-
 		const double CurrentTraceTime = RewindDebugger->CurrentTraceTime();
 		if (CurrentTraceTime != LastScrubTime)
 		{
-			const TraceServices::IFrameProvider& FrameProvider = TraceServices::ReadFrameProvider(*Session);
+			bool bFrameFound;
 			TraceServices::FFrame Frame;
-			if (FrameProvider.GetFrameFromTime(TraceFrameType_Game, CurrentTraceTime, Frame))
+			{
+				TraceServices::FAnalysisSessionReadScope SessionReadScope(*Session);
+				const TraceServices::IFrameProvider& FrameProvider = TraceServices::ReadFrameProvider(*Session);
+				bFrameFound = FrameProvider.GetFrameFromTime(TraceFrameType_Game, CurrentTraceTime, Frame);
+			}
+
+			if (bFrameFound)
 			{
 				const IAnimationProvider* AnimationProvider = Session->ReadProvider<IAnimationProvider>("AnimationProvider");
 				const IGameplayProvider* GameplayProvider = Session->ReadProvider<IGameplayProvider>("GameplayProvider");
@@ -315,6 +320,8 @@ void FRewindDebuggerAnimation::Update(float DeltaTime, IRewindDebugger* RewindDe
 
 						const bool ResetOrSpawnMeshes = !RewindDebugger->IsPIESimulating() && RewindDebugger->GetAnalysisSession() != nullptr;
 
+						TraceServices::FProviderReadScopeLock AnimProviderReadScope(*AnimationProvider);
+						TraceServices::FProviderReadScopeLock GameProviderReadScope(*GameplayProvider);
 						AnimationProvider->EnumerateSkeletalMeshPoseTimelines(
 						[this, ResetOrSpawnMeshes, RewindDebugger, &Frame, AnimationProvider, GameplayProvider, &bSearchForActor, &DebuggedObjects, &bNewActor](uint64 ObjectId, const IAnimationProvider::SkeletalMeshPoseTimeline& TimelineData)
 						{
@@ -408,6 +415,7 @@ void FRewindDebuggerAnimation::Update(float DeltaTime, IRewindDebugger* RewindDe
 						// mark the target position as invalid for a frame when the actor changes, so it will be treated as a teleport by the camera system
 						if (!bNewActor)
 						{
+							TraceServices::FProviderReadScopeLock AnimProviderReadScope(*AnimationProvider);
 							AnimationProvider->ReadSkeletalMeshPoseTimeline(TargetActorMeshId.GetMainId(), [this, &Frame, bNewActor](const IAnimationProvider::SkeletalMeshPoseTimeline& TimelineData, bool bHasCurves)
 							{
 								const FSkeletalMeshPoseMessage* PoseMessage = nullptr;
@@ -437,6 +445,9 @@ void FRewindDebuggerAnimation::Update(float DeltaTime, IRewindDebugger* RewindDe
 						// - enumerate over all anim graph timelines
 						// - check if their instance class still exists and is the debugging target for the Animation Blueprint Editor
 						// - if it is copy that debug data into the class debug data for the blueprint debugger
+
+						TraceServices::FProviderReadScopeLock AnimProviderReadScope(*AnimationProvider);
+						TraceServices::FProviderReadScopeLock GameProviderReadScope(*GameplayProvider);
 						AnimationProvider->EnumerateAnimGraphTimelines([this, &Frame, AnimationProvider, GameplayProvider](uint64 ObjectId, const IAnimationProvider::AnimGraphTimeline& AnimGraphTimeline)
 						{
 							UAnimInstance* AnimInstance = nullptr;
diff --git a/Engine/Plugins/Animation/GameplayInsights/Source/RewindDebugger/Private/RewindDebuggerCamera.cpp b/Engine/Plugins/Animation/GameplayInsights/Source/RewindDebugger/Private/RewindDebuggerCamera.cpp
--- a/Engine/Plugins/Animation/GameplayInsights/Source/RewindDebugger/Private/RewindDebuggerCamera.cpp
+++ b/Engine/Plugins/Animation/GameplayInsights/Source/RewindDebugger/Private/RewindDebuggerCamera.cpp
@@ -3,6 +3,7 @@
 #include "RewindDebuggerCamera.h"
 #include "Camera/CameraActor.h"
 #include "Camera/CameraComponent.h"
+#include "Common/ProviderLock.h"
 #include "Editor.h"
 #include "IGameplayProvider.h"
 #include "Insights/IUnrealInsightsModule.h"
@@ -106,14 +107,13 @@ void FRewindDebuggerCamera::Update(float DeltaTime, IRewindDebugger* RewindDebug
 	
 	if (const TraceServices::IAnalysisSession* Session = RewindDebugger->GetAnalysisSession())
 	{
-		TraceServices::FAnalysisSessionReadScope SessionReadScope(*Session);
 		double CurrentTraceTime = RewindDebugger->CurrentTraceTime();
-		
+
 		static double LastCameraScrubTime = 0.0f;
 		if (CurrentTraceTime != LastCameraScrubTime)
 		{
 			bool bCameraTraceDataFound = false;
-			
+
 			FLevelEditorModule& LevelEditor = FModuleManager::GetModuleChecked<FLevelEditorModule>("LevelEditor");
 			if (SLevelViewport* LevelViewport = LevelEditor.GetFirstActiveLevelViewport().Get())
 			{
@@ -137,11 +137,18 @@ void FRewindDebuggerCamera::Update(float DeltaTime, IRewindDebugger* RewindDebug
 				// always update the camera actor to the replay values even if it isn't locked
 				if (const IGameplayProvider* GameplayProvider = Session->ReadProvider<IGameplayProvider>("GameplayProvider"))
 				{
+					TraceServices::FProviderReadScopeLock ProviderReadScope(*GameplayProvider);
 					GameplayProvider->ReadViewTimeline([&bCameraTraceDataFound, this, RewindDebugger, CurrentTraceTime, Session](const IGameplayProvider::ViewTimeline& TimelineData)
 					{
-						const TraceServices::IFrameProvider& FrameProvider = TraceServices::ReadFrameProvider(*Session);
+						bool bFrameFound;
 						TraceServices::FFrame Frame;
-						if(FrameProvider.GetFrameFromTime(ETraceFrameType::TraceFrameType_Game, CurrentTraceTime, Frame))
+						{
+							TraceServices::FAnalysisSessionReadScope SessionReadScope(*Session);
+							const TraceServices::IFrameProvider& FrameProvider = TraceServices::ReadFrameProvider(*Session);
+							bFrameFound = FrameProvider.GetFrameFromTime(TraceFrameType_Game, CurrentTraceTime, Frame);
+						}
+
+						if (bFrameFound)
 						{
 							TimelineData.EnumerateEvents(Frame.StartTime, Frame.EndTime,
 								[&bCameraTraceDataFound, this, RewindDebugger](double InStartTime, double InEndTime, uint32 InDepth, const FViewMessage& ViewMessage)
diff --git a/Engine/Plugins/Animation/GameplayInsights/Source/RewindDebugger/Private/RewindDebuggerObjectTrack.cpp b/Engine/Plugins/Animation/GameplayInsights/Source/RewindDebugger/Private/RewindDebuggerObjectTrack.cpp
--- a/Engine/Plugins/Animation/GameplayInsights/Source/RewindDebugger/Private/RewindDebuggerObjectTrack.cpp
+++ b/Engine/Plugins/Animation/GameplayInsights/Source/RewindDebugger/Private/RewindDebuggerObjectTrack.cpp
@@ -1,6 +1,7 @@
 // Copyright Epic Games, Inc. All Rights Reserved.
 
 #include "RewindDebuggerObjectTrack.h"
+#include "Common/ProviderLock.h"
 #include "IAnimationProvider.h"
 #include "IGameplayProvider.h"
 #include "IRewindDebugger.h"
@@ -23,12 +24,12 @@ namespace RewindDebugger
 // check if an object is or is a subclass of a type by name, based on Insights traced type info
 static bool GetTypeHierarchyNames(const RewindDebugger::FObjectId InObjectId, const TraceServices::IAnalysisSession& InSession, TArray<const TCHAR*>& OutTypeNames)
 {
-	TraceServices::FAnalysisSessionReadScope SessionReadScope(InSession);
-
 	const IGameplayProvider* GameplayProvider = InSession.ReadProvider<IGameplayProvider>("GameplayProvider");
+
+	TraceServices::FProviderReadScopeLock ProviderReadScope(*GameplayProvider);
 	const FObjectInfo& ObjectInfo = GameplayProvider->GetObjectInfo(InObjectId);
 	uint64 ClassId = ObjectInfo.ClassId;
-	
+
 	while (ClassId != FObjectId::InvalidId)
 	{
 		const FClassInfo& ClassInfo = GameplayProvider->GetClassInfo(ClassId);
@@ -123,24 +124,33 @@ bool FRewindDebuggerObjectTrack::HandleDoubleClickInternal()
 
 	if (const TraceServices::IAnalysisSession* Session = RewindDebugger->GetAnalysisSession())
 	{
-		TraceServices::FAnalysisSessionReadScope SessionReadScope(*Session);
-
 		const IGameplayProvider* GameplayProvider = Session->ReadProvider<IGameplayProvider>("GameplayProvider");
-		const FObjectInfo& ObjectInfo = GameplayProvider->GetObjectInfo(ObjectId);
-		uint64 ClassId = ObjectInfo.ClassId;
-		bool bHandled = false;
 
+		// Cache the class hierarchy names under the lock
+		TArray<FName> ClassHierarchyNames;
+		{
+			TraceServices::FProviderReadScopeLock ProviderReadScope(*GameplayProvider);
+			const FObjectInfo& ObjectInfo = GameplayProvider->GetObjectInfo(ObjectId);
+			uint64 ClassId = ObjectInfo.ClassId;
+
+			while (ClassId != 0)
+			{
+				const FClassInfo& ClassInfo = GameplayProvider->GetClassInfo(ClassId);
+				ClassHierarchyNames.Add(ClassInfo.Name);
+				ClassId = ClassInfo.SuperId;
+			}
+		}
+
+		// Iterate handlers outside the lock since HandleDoubleClick may perform heavy operations
+		bool bHandled = false;
 		const int32 NumExtensions = ModularFeatures.GetModularFeatureImplementationCount(HandlerFeatureName);
 
-		// iterate up the class hierarchy, looking for a registered double click handler, until we find the one that succeeds that is most specific to the type of this object
-		while (ClassId != 0 && !bHandled)
+		for (int32 HierarchyIndex = 0; HierarchyIndex < ClassHierarchyNames.Num() && !bHandled; ++HierarchyIndex)
 		{
-			const FClassInfo& ClassInfo = GameplayProvider->GetClassInfo(ClassId);
-
 			for (int32 ExtensionIndex = 0; ExtensionIndex < NumExtensions; ++ExtensionIndex)
 			{
 				IRewindDebuggerDoubleClickHandler* Handler = static_cast<IRewindDebuggerDoubleClickHandler*>(ModularFeatures.GetModularFeatureImplementation(HandlerFeatureName, ExtensionIndex));
-				if (Handler->GetTargetTypeName() == ClassInfo.Name)
+				if (Handler->GetTargetTypeName() == ClassHierarchyNames[HierarchyIndex])
 				{
 					if (Handler->HandleDoubleClick(RewindDebugger))
 					{
@@ -149,8 +159,6 @@ bool FRewindDebuggerObjectTrack::HandleDoubleClickInternal()
 					}
 				}
 			}
-
-			ClassId = ClassInfo.SuperId;
 		}
 	}
 
@@ -237,9 +245,8 @@ FText FRewindDebuggerObjectTrack::GetDisplayNameInternal() const
 		{
 			IRewindDebugger* RewindDebugger = IRewindDebugger::Instance();
 			const TraceServices::IAnalysisSession* Session = RewindDebugger->GetAnalysisSession();
-			TraceServices::FAnalysisSessionReadScope SessionReadScope(*Session);
-
 			const IGameplayProvider* GameplayProvider = Session->ReadProvider<IGameplayProvider>("GameplayProvider");
+			TraceServices::FProviderReadScopeLock ProviderReadScope(*GameplayProvider);
 
 			const FObjectInfo& ObjectInfo = GameplayProvider->GetObjectInfo(ObjectId);
 
@@ -281,41 +288,43 @@ bool FRewindDebuggerObjectTrack::HasDebugDataInternal() const
 bool FRewindDebuggerObjectTrack::UpdateInternal()
 {
 	TRACE_CPUPROFILER_EVENT_SCOPE(FRewindDebuggerObjectTrack::UpdateInternal);
-	IRewindDebugger* RewindDebugger = IRewindDebugger::Instance();
-	const TraceServices::IAnalysisSession* Session = RewindDebugger->GetAnalysisSession();
-	TraceServices::FAnalysisSessionReadScope SessionReadScope(*Session);
-
-	const IGameplayProvider* GameplayProvider = Session->ReadProvider<IGameplayProvider>("GameplayProvider");
+	const IRewindDebugger* RewindDebugger = IRewindDebugger::Instance();
 
 	bool bChanged = false;
 
-	TRange<double> Existence = GameplayProvider->GetObjectRecordingLifetime(ObjectId);
-
-	ExistenceRange->Windows.SetNum(0, EAllowShrinking::No);
-	if (Existence.HasLowerBound() && Existence.HasUpperBound())
+	const TraceServices::IAnalysisSession* Session = RewindDebugger->GetAnalysisSession();
+	const IGameplayProvider* GameplayProvider = Session->ReadProvider<IGameplayProvider>("GameplayProvider");
 	{
-		double EndTime = Existence.GetUpperBoundValue();
-		if (EndTime == std::numeric_limits<double>::infinity())
-		{
-			EndTime = GameplayProvider->GetRecordingDuration();
-		}
-		ExistenceRange->Windows.Add({ Existence.GetLowerBoundValue(), EndTime, LOCTEXT("Object Existence","Object Existence"), LOCTEXT("Object Existence","Object Existence"), FLinearColor(0.1f,0.11f,0.1f) });
-	}
+		TraceServices::FProviderReadScopeLock ProviderReadScope(*GameplayProvider);
 
-	if (!bIconSearched)
-	{
-		TRACE_CPUPROFILER_EVENT_SCOPE(FRewindDebuggerObjectTrack::FindIcon);
-		if (const FObjectInfo* ObjectInfo = GameplayProvider->FindObjectInfo(ObjectId))
+		const TRange<double> Existence = GameplayProvider->GetObjectRecordingLifetime(ObjectId);
+
+		ExistenceRange->Windows.SetNum(0, EAllowShrinking::No);
+		if (Existence.HasLowerBound() && Existence.HasUpperBound())
 		{
-			Icon = GameplayProvider->FindIconForClass(ObjectInfo->ClassId);
+			double EndTime = Existence.GetUpperBoundValue();
+			if (EndTime == std::numeric_limits<double>::infinity())
+			{
+				EndTime = GameplayProvider->GetRecordingDuration();
+			}
+			ExistenceRange->Windows.Add({ Existence.GetLowerBoundValue(), EndTime, LOCTEXT("Object Existence","Object Existence"), LOCTEXT("Object Existence","Object Existence"), FLinearColor(0.1f,0.11f,0.1f) });
 		}
 
-		if (!Icon.IsSet())
+		if (!bIconSearched)
 		{
-			Icon = FSlateIconFinder::FindIconForClass(UObject::StaticClass());
+			TRACE_CPUPROFILER_EVENT_SCOPE(FRewindDebuggerObjectTrack::FindIcon);
+			if (const FObjectInfo* ObjectInfo = GameplayProvider->FindObjectInfo(ObjectId))
+			{
+				Icon = GameplayProvider->FindIconForClass(ObjectInfo->ClassId);
+			}
+
+			if (!Icon.IsSet())
+			{
+				Icon = FSlateIconFinder::FindIconForClass(UObject::StaticClass());
+			}
+			bIconSearched = true;
+			bChanged = true;
 		}
-		bIconSearched = true;
-		bChanged = true;
 	}
 
 	TArray<FObjectId, TInlineAllocator<32>> FoundObjects;
@@ -419,6 +428,7 @@ bool FRewindDebuggerObjectTrack::UpdateInternal()
 			}
 		};
 
+		TraceServices::FProviderReadScopeLock ProviderReadScope(*GameplayProvider);
 		GameplayProvider->EnumerateSubobjects(ObjectId, [this, AddChildObject](const FObjectId& SubObjectId)
 		{
 			AddChildObject(SubObjectId);
@@ -466,6 +476,7 @@ bool FRewindDebuggerObjectTrack::UpdateInternal()
 	if (bAddController)
 	{
 		TRACE_CPUPROFILER_EVENT_SCOPE(FRewindDebuggerObjectTrack::FindController);
+		TraceServices::FProviderReadScopeLock ProviderReadScope(*GameplayProvider);
 		// Should probably update this to use a time range and return all possessing controllers from the visible time range.  For now just returns the one at the current time.
 		if (uint64 ControllerId = GameplayProvider->FindPossessingController(ObjectId.GetMainId(), RewindDebugger->CurrentTraceTime()))
 		{
diff --git a/Engine/Plugins/Animation/GameplayInsights/Source/RewindDebugger/Private/SRewindDebugger.cpp b/Engine/Plugins/Animation/GameplayInsights/Source/RewindDebugger/Private/SRewindDebugger.cpp
--- a/Engine/Plugins/Animation/GameplayInsights/Source/RewindDebugger/Private/SRewindDebugger.cpp
+++ b/Engine/Plugins/Animation/GameplayInsights/Source/RewindDebugger/Private/SRewindDebugger.cpp
@@ -2,6 +2,7 @@
 
 #include "SRewindDebugger.h"
 #include "ActorPickerMode.h"
+#include "Common/ProviderLock.h"
 #include "Editor.h"
 #include "Editor/EditorEngine.h"
 #include "Framework/Application/SlateApplication.h"
@@ -232,66 +233,82 @@ TSharedRef<SWidget> SRewindDebugger::MakeObjectSelectionMenu()
 	{
 		if (const TraceServices::IAnalysisSession* Session = RewindDebugger->GetAnalysisSession())
 		{
-			TraceServices::FAnalysisSessionReadScope SessionReadScope(*Session);
 			const IGameplayProvider* GameplayProvider = Session->ReadProvider<IGameplayProvider>("GameplayProvider");
-
-			TArray<FString> AllowedTypeNames;
-
-			// Gather types allowed by the project
-			const URewindDebuggerProjectSettings* ProjectSettings = GetDefault<URewindDebuggerProjectSettings>();
-			Algo::Transform(ProjectSettings->SelectorAllowedTypes, AllowedTypeNames, [](const FSoftClassPath& Path) { return Path.ToString(); });
-			Algo::Transform(ProjectSettings->SelectorAllowedStructTypes, AllowedTypeNames, [](const FInstancedStruct& Type) { return GetPathNameSafe(Type.GetScriptStruct()); });
-
-			// Gather types allowed by the user
-			const URewindDebuggerSettings* UserSettings = GetDefault<URewindDebuggerSettings>();
-			Algo::Transform(UserSettings->SelectorAllowedTypes, AllowedTypeNames, [](const FSoftClassPath& Path) { return Path.ToString(); });
-			Algo::Transform(UserSettings->SelectorAllowedStructTypes, AllowedTypeNames, [](const FInstancedStruct& Type) { return GetPathNameSafe(Type.GetScriptStruct()); });
-
-			TArray<const uint64> AllowedTypeIds;
-			AllowedTypeIds.Reserve(AllowedTypeNames.Num());
-			for (const FString& AllowedTypeName : AllowedTypeNames)
+			if (GameplayProvider)
 			{
-				if (const FClassInfo* ClassInfo = GameplayProvider->FindClassInfo(*AllowedTypeName))
+				// Gather allowed type names from settings outside the provider lock
+				TArray<FString> AllowedTypeNames;
+	
+				const URewindDebuggerProjectSettings* ProjectSettings = GetDefault<URewindDebuggerProjectSettings>();
+				Algo::Transform(ProjectSettings->SelectorAllowedTypes, AllowedTypeNames, [](const FSoftClassPath& Path) { return Path.ToString(); });
+				Algo::Transform(ProjectSettings->SelectorAllowedStructTypes, AllowedTypeNames, [](const FInstancedStruct& Type) { return GetPathNameSafe(Type.GetScriptStruct()); });
+	
+				const URewindDebuggerSettings* UserSettings = GetDefault<URewindDebuggerSettings>();
+				Algo::Transform(UserSettings->SelectorAllowedTypes, AllowedTypeNames, [](const FSoftClassPath& Path) { return Path.ToString(); });
+				Algo::Transform(UserSettings->SelectorAllowedStructTypes, AllowedTypeNames, [](const FInstancedStruct& Type) { return GetPathNameSafe(Type.GetScriptStruct()); });
+	
+				// Cache object data from provider under a tight lock scope
+				struct FCachedObjectEntry
 				{
-					AllowedTypeIds.AddUnique(ClassInfo->Id);
-				}
-			}
-
-			auto IsClassOrChildOfAllowed = [&AllowedTypeIds, GameplayProvider](const uint64 ClassId)
+					FString Name;
+					FSlateIcon Icon;
+				};
+				TArray<FCachedObjectEntry> CachedObjects;
+	
 				{
-					if (AllowedTypeIds.Num() > 0)
+					TraceServices::FProviderReadScopeLock ProviderReadScope(*GameplayProvider);
+	
+					TArray<const uint64> AllowedTypeIds;
+					AllowedTypeIds.Reserve(AllowedTypeNames.Num());
+					for (const FString& AllowedTypeName : AllowedTypeNames)
 					{
-						for (const uint64 AllowedTypeId : AllowedTypeIds)
+						if (const FClassInfo* ClassInfo = GameplayProvider->FindClassInfo(*AllowedTypeName))
 						{
-							if (GameplayProvider->IsSubClassOf(ClassId, AllowedTypeId))
-							{
-								return true;
-							}
+							AllowedTypeIds.AddUnique(ClassInfo->Id);
 						}
-						return false;
 					}
-
-					// Allow everything if allowed class list is not used
-					return true;
-				};
-
-			GameplayProvider->EnumerateObjects([this, &MenuBuilder, GameplayProvider, &IsClassOrChildOfAllowed](const FObjectInfo& ObjectInfo)
+	
+					auto IsClassOrChildOfAllowed = [&AllowedTypeIds, GameplayProvider](const uint64 ClassId)
+						{
+							if (AllowedTypeIds.Num() > 0)
+							{
+								for (const uint64 AllowedTypeId : AllowedTypeIds)
+								{
+									if (GameplayProvider->IsSubClassOf(ClassId, AllowedTypeId))
+									{
+										return true;
+									}
+								}
+								return false;
+							}
+	
+							// Allow everything if allowed class list is not used
+							return true;
+						};
+	
+					GameplayProvider->EnumerateObjects([&CachedObjects, GameplayProvider, &IsClassOrChildOfAllowed](const FObjectInfo& ObjectInfo)
+						{
+							if (!IsClassOrChildOfAllowed(ObjectInfo.ClassId))
+							{
+								return;
+							}
+	
+							CachedObjects.Add({ ObjectInfo.Name, GameplayProvider->FindIconForClass(ObjectInfo.ClassId) });
+						});
+				}
+	
+				// Build menu entries outside the lock
+				for (const FCachedObjectEntry& Entry : CachedObjects)
 				{
-					if (!IsClassOrChildOfAllowed(ObjectInfo.ClassId))
-					{
-						return;
-					}
-
-					const FString Name = ObjectInfo.Name;
-					const FText SelectedLabel = FText::FromString(Name);
-					const FSlateIcon Icon = GameplayProvider->FindIconForClass(ObjectInfo.ClassId);
-
-					MenuBuilder.AddMenuEntry(SelectedLabel, FText(), Icon, FExecuteAction::CreateLambda([this, Name]()
+					const FText SelectedLabel = FText::FromString(Entry.Name);
+	
+					MenuBuilder.AddMenuEntry(SelectedLabel, FText(), Entry.Icon, FExecuteAction::CreateLambda([this, Name = Entry.Name]()
 						{
 							FSlateApplication::Get().DismissAllMenus();
 							DebuggedObjectName.Set(Name);
 						}));
-				});
+				}
+			}
 		}
 		else
 		{
diff --git a/Engine/Plugins/Animation/GameplayInsights/Source/RewindDebuggerVLog/Private/RewindDebuggerVLog.cpp b/Engine/Plugins/Animation/GameplayInsights/Source/RewindDebuggerVLog/Private/RewindDebuggerVLog.cpp
--- a/Engine/Plugins/Animation/GameplayInsights/Source/RewindDebuggerVLog/Private/RewindDebuggerVLog.cpp
+++ b/Engine/Plugins/Animation/GameplayInsights/Source/RewindDebuggerVLog/Private/RewindDebuggerVLog.cpp
@@ -2,6 +2,7 @@
 
 #include "RewindDebuggerVLog.h"
 
+#include "Common/ProviderLock.h"
 #include "FindInBlueprintManager.h"
 #include "IRewindDebugger.h"
 #include "IRewindDebuggerModule.h"
@@ -54,16 +55,21 @@ void FRewindDebuggerVLog::OnShowDebugInfo(UCanvas* Canvas, APlayerController* Pl
 			ObjectsVisited.Empty();
 			if (const TraceServices::IAnalysisSession* Session = RewindDebugger->GetAnalysisSession())
 			{
-				TraceServices::FAnalysisSessionReadScope SessionReadScope(*Session);
-				double CurrentTraceTime = RewindDebugger->CurrentTraceTime();
+				const double CurrentTraceTime = RewindDebugger->CurrentTraceTime();
 
-				const TraceServices::IFrameProvider& FrameProvider = TraceServices::ReadFrameProvider(*Session);
+				bool bFrameFound;
 				TraceServices::FFrame CurrentFrame;
+				{
+					TraceServices::FAnalysisSessionReadScope SessionReadScope(*Session);
+					const TraceServices::IFrameProvider& FrameProvider = TraceServices::ReadFrameProvider(*Session);
+					bFrameFound = FrameProvider.GetFrameFromTime(TraceFrameType_Game, CurrentTraceTime, CurrentFrame);
+				}
 
-				if (FrameProvider.GetFrameFromTime(TraceFrameType_Game, CurrentTraceTime, CurrentFrame))
+				if (bFrameFound)
 				{
 					if (const IVisualLoggerProvider* VisualLoggerProvider = Session->ReadProvider<IVisualLoggerProvider>("VisualLoggerProvider"))
 					{
+						TraceServices::FProviderReadScopeLock ProviderReadScope(*VisualLoggerProvider);
 						AddLogEntries(RewindDebugger->GetDebuggedObjects(), CurrentFrame.StartTime, CurrentFrame.EndTime, VisualLoggerProvider, Canvas);
 					}
 				}
@@ -228,24 +234,34 @@ void FRewindDebuggerVLog::MakeCategoriesMenu(UToolMenu* Menu) const
 
 	if (const TraceServices::IAnalysisSession* Session = RewindDebugger ? RewindDebugger->GetAnalysisSession() : nullptr)
 	{
-		TraceServices::FAnalysisSessionReadScope SessionReadScope(*Session);
 		if (const IVisualLoggerProvider* VisualLoggerProvider = Session->ReadProvider<IVisualLoggerProvider>("VisualLoggerProvider"))
 		{
-			VisualLoggerProvider->EnumerateCategories([&Section, this](const FName& Category)
-				{
-					const FText CategoryAsText = FText::FromName(Category);
-					Section.AddMenuEntry(Category,
-						CategoryAsText,
-						FText::Format(LOCTEXT("CategoryDesc", "Enable/Disable {0} category"), CategoryAsText),
-						FSlateIcon(),
-						FUIAction(
-							FExecuteAction::CreateLambda([Category, this]() { ToggleCategory(Category); }),
-							FCanExecuteAction(),
-							FGetActionCheckState::CreateLambda([Category, this]() { return IsCategoryActive(Category) ? ECheckBoxState::Checked : ECheckBoxState::Unchecked; })
-						),
-						EUserInterfaceActionType::ToggleButton
-					);
-				});
+			// Cache category names under a tight lock scope
+			TArray<FName> Categories;
+			{
+				TraceServices::FProviderReadScopeLock ProviderReadScope(*VisualLoggerProvider);
+				VisualLoggerProvider->EnumerateCategories([&Categories](const FName& Category)
+					{
+						Categories.Add(Category);
+					});
+			}
+
+			// Build menu entries outside the lock
+			for (const FName& Category : Categories)
+			{
+				const FText CategoryAsText = FText::FromName(Category);
+				Section.AddMenuEntry(Category,
+					CategoryAsText,
+					FText::Format(LOCTEXT("CategoryDesc", "Enable/Disable {0} category"), CategoryAsText),
+					FSlateIcon(),
+					FUIAction(
+						FExecuteAction::CreateLambda([Category, this]() { ToggleCategory(Category); }),
+						FCanExecuteAction(),
+						FGetActionCheckState::CreateLambda([Category, this]() { return IsCategoryActive(Category) ? ECheckBoxState::Checked : ECheckBoxState::Unchecked; })
+					),
+					EUserInterfaceActionType::ToggleButton
+				);
+			}
 		}
 	}
 }
diff --git a/Engine/Plugins/Animation/GameplayInsights/Source/RewindDebuggerVLog/Private/VLogTraceModule.cpp b/Engine/Plugins/Animation/GameplayInsights/Source/RewindDebuggerVLog/Private/VLogTraceModule.cpp
--- a/Engine/Plugins/Animation/GameplayInsights/Source/RewindDebuggerVLog/Private/VLogTraceModule.cpp
+++ b/Engine/Plugins/Animation/GameplayInsights/Source/RewindDebuggerVLog/Private/VLogTraceModule.cpp
@@ -15,7 +15,7 @@ void FVLogTraceModule::GetModuleInfo(TraceServices::FModuleInfo& OutModuleInfo)
 void FVLogTraceModule::OnAnalysisBegin(TraceServices::IAnalysisSession& InSession)
 {
 	TSharedPtr<FVisualLoggerProvider> VisualLoggerProvider = MakeShared<FVisualLoggerProvider>(InSession);
-	InSession.AddProvider(FVisualLoggerProvider::ProviderName, VisualLoggerProvider);
+	InSession.AddProvider(FVisualLoggerProvider::ProviderName, VisualLoggerProvider, VisualLoggerProvider);
 
 	InSession.AddAnalyzer(new FVisualLoggerAnalyzer(InSession, *VisualLoggerProvider));
 }
diff --git a/Engine/Plugins/Animation/GameplayInsights/Source/RewindDebuggerVLog/Private/VisualLogTrack.cpp b/Engine/Plugins/Animation/GameplayInsights/Source/RewindDebuggerVLog/Private/VisualLogTrack.cpp
--- a/Engine/Plugins/Animation/GameplayInsights/Source/RewindDebuggerVLog/Private/VisualLogTrack.cpp
+++ b/Engine/Plugins/Animation/GameplayInsights/Source/RewindDebuggerVLog/Private/VisualLogTrack.cpp
@@ -2,6 +2,7 @@
 
 #include "VisualLogTrack.h"
 
+#include "Common/ProviderLock.h"
 #include "IGameplayProvider.h"
 #include "IRewindDebugger.h"
 #include "Modules/ModuleManager.h"
@@ -102,7 +103,7 @@ bool FVisualLogCategoryTrack::UpdateInternal()
 			EventData->Points.SetNum(0, EAllowShrinking::No);
 			EventData->Windows.SetNum(0);
 
-			TraceServices::FAnalysisSessionReadScope SessionReadScope(*AnalysisSession);
+			TraceServices::FProviderReadScopeLock ProviderReadScope(*VisLogProvider);
 
 			VisLogProvider->ReadVisualLogEntryTimeline(ObjectId, [this, StartTime, EndTime, VisLogProvider, AnalysisSession](const FVisualLoggerProvider::VisualLogEntryTimeline& InTimeline)
 				{
@@ -135,12 +136,17 @@ bool FVisualLogCategoryTrack::UpdateInternal()
 			PreviousScrubTime = CurrentScrubTime;
 			LogEntries.Reset();
 
-			const TraceServices::IFrameProvider& FramesProvider = TraceServices::ReadFrameProvider(*AnalysisSession);
-			TraceServices::FAnalysisSessionReadScope SessionReadScope(*AnalysisSession);
-
+			bool bFrameFound;
 			TraceServices::FFrame MarkerFrame;
-			if (FramesProvider.GetFrameFromTime(TraceFrameType_Game, CurrentScrubTime, MarkerFrame))
 			{
+				TraceServices::FAnalysisSessionReadScope SessionReadScope(*AnalysisSession);
+				const TraceServices::IFrameProvider& FrameProvider = TraceServices::ReadFrameProvider(*AnalysisSession);
+				bFrameFound = FrameProvider.GetFrameFromTime(TraceFrameType_Game, CurrentScrubTime, MarkerFrame);
+			}
+
+			if (bFrameFound)
+			{
+				TraceServices::FProviderReadScopeLock ProviderReadScope(*VisLogProvider);
 				VisLogProvider->ReadVisualLogEntryTimeline(ObjectId,
 					[this, &Entries = LogEntries, &MarkerFrame, EndTime, VisLogProvider, AnalysisSession](const FVisualLoggerProvider::VisualLogEntryTimeline& InTimeline)
 					{
@@ -237,7 +243,7 @@ TOptional<double> FVisualLogCategoryTrack::GetStepFrameTimeInternal(const EStepM
 	}
 
 	const TraceServices::IFrameProvider& FramesProvider = TraceServices::ReadFrameProvider(*AnalysisSession);
-	TraceServices::FAnalysisSessionReadScope SessionReadScope(*AnalysisSession);
+	TraceServices::FProviderReadScopeLock GameProviderReadScope(*GameplayProvider);
 
 	// Legacy Index that is always 0 now
 	constexpr int32 RecordingIndex = 0;
@@ -247,9 +253,12 @@ TOptional<double> FVisualLogCategoryTrack::GetStepFrameTimeInternal(const EStepM
 		return {};
 	}
 
+	TraceServices::FAnalysisSessionReadScope SessionReadScope(*AnalysisSession);
 	TraceServices::FFrame Frame;
 	if (FramesProvider.GetFrameFromTime(TraceFrameType_Game, CurrentScrubTime.ProfileTime, Frame))
 	{
+		TraceServices::FProviderReadScopeLock VisLogProviderReadScope(*VisLogProvider);
+
 		bool bEventFound = false;
 
 		// Set the frame index increment based on the step mode
@@ -291,7 +300,7 @@ TOptional<double> FVisualLogCategoryTrack::GetStepFrameTimeInternal(const EStepM
 			if (bEventFound)
 			{
 				const FRecordingInfoMessage* MatchingRecordingInfo = nullptr;
-				Recording->EnumerateEvents(CandidateFrame->StartTime, CandidateFrame->EndTime, 
+				Recording->EnumerateEvents(CandidateFrame->StartTime, CandidateFrame->EndTime,
 					[&MatchingRecordingInfo](double InStartTime, double InEndTime, uint32 InDepth, const FRecordingInfoMessage& InMessage)
 					{
 						MatchingRecordingInfo = &InMessage;
@@ -344,26 +353,28 @@ bool FVisualLogTrack::UpdateInternal()
 	{
 		TArray<FName> UniqueTrackIds;
 
-		TraceServices::FAnalysisSessionReadScope SessionReadScope(*AnalysisSession);
+		{
+			TraceServices::FProviderReadScopeLock ProviderReadScope(*VisLogProvider);
 
-		VisLogProvider->ReadVisualLogEntryTimeline(ObjectId, [this, StartTime, EndTime, VisLogProvider, AnalysisSession, &UniqueTrackIds](const FVisualLoggerProvider::VisualLogEntryTimeline& InTimeline)
-			{
-				InTimeline.EnumerateEvents(StartTime, EndTime, [this, StartTime, EndTime, VisLogProvider, AnalysisSession, &UniqueTrackIds](double InStartTime, double InEndTime, uint32 InDepth, const FVisualLogEntry& InMessage)
-					{
-						if (InMessage.LogLines.Num() > 0)
+			VisLogProvider->ReadVisualLogEntryTimeline(ObjectId, [this, StartTime, EndTime, VisLogProvider, AnalysisSession, &UniqueTrackIds](const FVisualLoggerProvider::VisualLogEntryTimeline& InTimeline)
+				{
+					InTimeline.EnumerateEvents(StartTime, EndTime, [this, StartTime, EndTime, VisLogProvider, AnalysisSession, &UniqueTrackIds](double InStartTime, double InEndTime, uint32 InDepth, const FVisualLogEntry& InMessage)
 						{
-							for (const FVisualLogLine& Line : InMessage.LogLines)
+							if (InMessage.LogLines.Num() > 0)
+							{
+								for (const FVisualLogLine& Line : InMessage.LogLines)
+								{
+									UniqueTrackIds.AddUnique(Line.Category);
+								}
+							}
+							for (const FVisualLogShapeElement& Element : InMessage.ElementsToDraw)
 							{
-								UniqueTrackIds.AddUnique(Line.Category);
+								UniqueTrackIds.AddUnique(Element.Category);
 							}
-						}
-						for (const FVisualLogShapeElement& Element : InMessage.ElementsToDraw)
-						{
-							UniqueTrackIds.AddUnique(Element.Category);
-						}
-						return TraceServices::EEventEnumerate::Continue;
-					});
-			});
+							return TraceServices::EEventEnumerate::Continue;
+						});
+				});
+		}
 
 		UniqueTrackIds.StableSort([](const FName& A, const FName& B) { return A.ToString() < B.ToString(); });
 		const int32 TrackCount = UniqueTrackIds.Num();
@@ -431,10 +442,10 @@ bool FVisualLogTrackCreator::HasDebugInfoInternal(const FObjectId& InObjectId) c
 	TRACE_CPUPROFILER_EVENT_SCOPE(FVisualLogTrack::HasDebugInfoInternal);
 	const TraceServices::IAnalysisSession* AnalysisSession = IRewindDebugger::Instance()->GetAnalysisSession();
 
-	TraceServices::FAnalysisSessionReadScope SessionReadScope(*AnalysisSession);
 	bool bHasData = false;
 	if (const FVisualLoggerProvider* VLogProvider = AnalysisSession->ReadProvider<FVisualLoggerProvider>(FVisualLoggerProvider::ProviderName))
 	{
+		TraceServices::FProviderReadScopeLock ProviderReadScope(*VLogProvider);
 		VLogProvider->ReadVisualLogEntryTimeline(InObjectId.GetMainId(), [&bHasData](const FVisualLoggerProvider::VisualLogEntryTimeline& InTimeline)
 			{
 				bHasData = true;
diff --git a/Engine/Plugins/Animation/GameplayInsights/Source/RewindDebuggerVLog/Private/VisualLoggerAnalyzer.cpp b/Engine/Plugins/Animation/GameplayInsights/Source/RewindDebuggerVLog/Private/VisualLoggerAnalyzer.cpp
--- a/Engine/Plugins/Animation/GameplayInsights/Source/RewindDebuggerVLog/Private/VisualLoggerAnalyzer.cpp
+++ b/Engine/Plugins/Animation/GameplayInsights/Source/RewindDebuggerVLog/Private/VisualLoggerAnalyzer.cpp
@@ -2,9 +2,9 @@
 
 #include "VisualLoggerAnalyzer.h"
 
+#include "Common/ProviderLock.h"
 #include "HAL/LowLevelMemTracker.h"
 #include "Serialization/MemoryReader.h"
-#include "TraceServices/Model/AnalysisSession.h"
 #include "VisualLoggerProvider.h"
 
 FVisualLoggerAnalyzer::FVisualLoggerAnalyzer(TraceServices::IAnalysisSession& InSession, FVisualLoggerProvider& InProvider)
@@ -29,7 +29,7 @@ bool FVisualLoggerAnalyzer::OnEvent(uint16 RouteId, EStyle Style, const FOnEvent
 {
 	LLM_SCOPE_BYNAME(TEXT("Insights/FVisualLoggerAnalyzer"));
 
-	TraceServices::FAnalysisSessionEditScope _(Session);
+	TraceServices::FProviderEditScopeLock ProviderEditScope(Provider);
 
 	const auto& EventData = Context.EventData;
 	switch (RouteId)
diff --git a/Engine/Plugins/Animation/GameplayInsights/Source/RewindDebuggerVLog/Private/VisualLoggerProvider.cpp b/Engine/Plugins/Animation/GameplayInsights/Source/RewindDebuggerVLog/Private/VisualLoggerProvider.cpp
--- a/Engine/Plugins/Animation/GameplayInsights/Source/RewindDebuggerVLog/Private/VisualLoggerProvider.cpp
+++ b/Engine/Plugins/Animation/GameplayInsights/Source/RewindDebuggerVLog/Private/VisualLoggerProvider.cpp
@@ -3,6 +3,11 @@
 #include "VisualLoggerProvider.h"
 #include "Serialization/MemoryReader.h"
 
+namespace TraceServices
+{
+thread_local FProviderLock::FThreadLocalState GVisualLoggerProviderLockState;
+}
+
 FName FVisualLoggerProvider::ProviderName("VisualLoggerProvider");
 
 #define LOCTEXT_NAMESPACE "VisualLoggerProvider"
@@ -14,7 +19,7 @@ FVisualLoggerProvider::FVisualLoggerProvider(TraceServices::IAnalysisSession& In
 
 bool FVisualLoggerProvider::ReadVisualLogEntryTimeline(uint64 InObjectId, TFunctionRef<void(const VisualLogEntryTimeline&)> Callback) const
 {
-	Session.ReadAccessCheck();
+	ReadAccessCheck();
 
 	const uint32* IndexPtr = ObjectIdToLogEntryTimelines.Find(InObjectId);
 	if(IndexPtr != nullptr)
@@ -31,7 +36,7 @@ bool FVisualLoggerProvider::ReadVisualLogEntryTimeline(uint64 InObjectId, TFunct
 
 void FVisualLoggerProvider::EnumerateCategories(TFunctionRef<void(const FName&)> Callback) const
 {
-	Session.ReadAccessCheck();
+	ReadAccessCheck();
 
 	for(const FName& Category : Categories)
 	{
@@ -41,7 +46,7 @@ void FVisualLoggerProvider::EnumerateCategories(TFunctionRef<void(const FName&)>
 
 void  FVisualLoggerProvider::AppendVisualLogEntry(uint64 InObjectId, double InTime, const FVisualLogEntry& Entry)
 {
-	Session.WriteAccessCheck();
+	EditAccessCheck();
 
 	TSharedPtr<TraceServices::TPointTimeline<FVisualLogEntry, FVLogTimelineSettings>> Timeline;
 	uint32* IndexPtr = ObjectIdToLogEntryTimelines.Find(InObjectId);
@@ -71,11 +76,15 @@ void  FVisualLoggerProvider::AppendVisualLogEntry(uint64 InObjectId, double InTi
 
 FVLogEntry& FVisualLoggerProvider::AddLogEntry()
 {
+	EditAccessCheck();
+
 	return LogEntries.AddDefaulted_GetRef();
 }
 
 void FVisualLoggerProvider::AddLogEntryChunk(const uint32 ID, const uint32 ChunkNum, const uint16 Size, const TArrayView<const uint8>& ChunkData)
 {
+	EditAccessCheck();
+
 	const int32 Index = LogEntries.IndexOfByPredicate([ID](const FVLogEntry& Entry) { return Entry.EntryID == ID; });
 	if (Index == INDEX_NONE)
 	{
diff --git a/Engine/Plugins/Animation/GameplayInsights/Source/RewindDebuggerVLog/Private/VisualLoggerProvider.h b/Engine/Plugins/Animation/GameplayInsights/Source/RewindDebuggerVLog/Private/VisualLoggerProvider.h
--- a/Engine/Plugins/Animation/GameplayInsights/Source/RewindDebuggerVLog/Private/VisualLoggerProvider.h
+++ b/Engine/Plugins/Animation/GameplayInsights/Source/RewindDebuggerVLog/Private/VisualLoggerProvider.h
@@ -3,12 +3,17 @@
 #pragma once
 
 #include "IVisualLoggerProvider.h"
+#include "Common/ProviderLock.h"
 #include "Model/PointTimeline.h"
 #include "Containers/Map.h"
 #include "Templates/SharedPointer.h"
 #include "VisualLogger/VisualLoggerTypes.h"
 
-namespace TraceServices { class IAnalysisSession; }
+namespace TraceServices
+{
+class IAnalysisSession;
+extern thread_local FProviderLock::FThreadLocalState GVisualLoggerProviderLockState;
+}
 
 struct FVLogTimelineSettings
 {
@@ -36,6 +41,34 @@ public:
 
 	explicit FVisualLoggerProvider(TraceServices::IAnalysisSession& InSession);
 
+	//~ IProvider (read lock)
+	virtual void BeginRead() const override
+	{
+		Lock.BeginRead(TraceServices::GVisualLoggerProviderLockState);
+	}
+	virtual void EndRead() const override
+	{
+		Lock.EndRead(TraceServices::GVisualLoggerProviderLockState);
+	}
+	virtual void ReadAccessCheck() const override
+	{
+		Lock.ReadAccessCheck(TraceServices::GVisualLoggerProviderLockState);
+	}
+
+	//~ IEditableProvider (write lock)
+	virtual void BeginEdit() const override
+	{
+		Lock.BeginWrite(TraceServices::GVisualLoggerProviderLockState);
+	}
+	virtual void EndEdit() const override
+	{
+		Lock.EndWrite(TraceServices::GVisualLoggerProviderLockState);
+	}
+	virtual void EditAccessCheck() const override
+	{
+		Lock.WriteAccessCheck(TraceServices::GVisualLoggerProviderLockState);
+	}
+
 	//~ IVisualLoggerProvider interface
 	virtual bool ReadVisualLogEntryTimeline(uint64 InObjectId, TFunctionRef<void(const VisualLogEntryTimeline&)> Callback) const override;
 	virtual void EnumerateCategories(TFunctionRef<void(const FName&)> Callback) const override;
@@ -47,6 +80,7 @@ public:
 	void AddLogEntryChunk(uint32 ID, uint32 ChunkNum, uint16 Size, const TArrayView<const uint8>& ChunkData);
 
 private:
+	mutable TraceServices::FProviderLock Lock;
 	TraceServices::IAnalysisSession& Session;
 
 	TArray<FVLogEntry> LogEntries;
diff --git a/Engine/Plugins/Animation/GameplayInsights/Source/RewindDebuggerVLog/Private/Widgets/LogCategoryVerbositySelector.cpp b/Engine/Plugins/Animation/GameplayInsights/Source/RewindDebuggerVLog/Private/Widgets/LogCategoryVerbositySelector.cpp
--- a/Engine/Plugins/Animation/GameplayInsights/Source/RewindDebuggerVLog/Private/Widgets/LogCategoryVerbositySelector.cpp
+++ b/Engine/Plugins/Animation/GameplayInsights/Source/RewindDebuggerVLog/Private/Widgets/LogCategoryVerbositySelector.cpp
@@ -2,6 +2,7 @@
 
 #include "LogCategoryVerbositySelector.h"
 
+#include "Common/ProviderLock.h"
 #include "Components/VerticalBox.h"
 #include "IRewindDebugger.h"
 #include "IVisualLoggerProvider.h"
@@ -343,9 +344,9 @@ void SLogCategoryFilter::RefreshCategoryList()
 	const IRewindDebugger* RewindDebugger = IRewindDebugger::Instance();
 	if (const TraceServices::IAnalysisSession* Session = RewindDebugger ? RewindDebugger->GetAnalysisSession() : nullptr)
 	{
-		TraceServices::FAnalysisSessionReadScope SessionReadScope(*Session);
 		if (const IVisualLoggerProvider* VisualLoggerProvider = Session->ReadProvider<IVisualLoggerProvider>("VisualLoggerProvider"))
 		{
+			TraceServices::FProviderReadScopeLock ProviderReadScope(*VisualLoggerProvider);
 			VisualLoggerProvider->EnumerateCategories([InsertFunc](const FName& Category)
 				{
 					InsertFunc(Category);
diff --git a/Engine/Plugins/Animation/GameplayInsights/Source/RewindDebuggerVLog/Public/IVisualLoggerProvider.h b/Engine/Plugins/Animation/GameplayInsights/Source/RewindDebuggerVLog/Public/IVisualLoggerProvider.h
--- a/Engine/Plugins/Animation/GameplayInsights/Source/RewindDebuggerVLog/Public/IVisualLoggerProvider.h
+++ b/Engine/Plugins/Animation/GameplayInsights/Source/RewindDebuggerVLog/Public/IVisualLoggerProvider.h
@@ -15,7 +15,9 @@
 
 struct FVisualLogEntry;
 
-class IVisualLoggerProvider : public TraceServices::IProvider
+class IVisualLoggerProvider
+	: public TraceServices::IProvider
+	, public TraceServices::IEditableProvider
 {
 public:
 	typedef TraceServices::ITimeline<FVisualLogEntry> VisualLogEntryTimeline;
diff --git a/Engine/Plugins/Animation/PoseSearch/Source/Editor/Private/PoseSearchDebugger.cpp b/Engine/Plugins/Animation/PoseSearch/Source/Editor/Private/PoseSearchDebugger.cpp
--- a/Engine/Plugins/Animation/PoseSearch/Source/Editor/Private/PoseSearchDebugger.cpp
+++ b/Engine/Plugins/Animation/PoseSearch/Source/Editor/Private/PoseSearchDebugger.cpp
@@ -119,8 +119,6 @@ void SCostTimelineView::UpdateInternal(uint64 ObjectId)
 	check(AnalysisSession);
 	if (const FTraceProvider* PoseSearchProvider = AnalysisSession->ReadProvider<FTraceProvider>(FTraceProvider::ProviderName))
 	{
-		TraceServices::FAnalysisSessionReadScope SessionReadScope(*AnalysisSession);
-
 		BestCostView->CurveData->Points.Reset();
 		BruteForceCostView->CurveData->Points.Reset();
 		BestPosePosView->CurveData->Points.Reset();
@@ -130,6 +128,7 @@ void SCostTimelineView::UpdateInternal(uint64 ObjectId)
 		double StartTime = TraceTimeRange.GetLowerBoundValue();
 		double EndTime = TraceTimeRange.GetUpperBoundValue();
 
+		TraceServices::FProviderReadScopeLock ProviderReadScope(*PoseSearchProvider);
 		PoseSearchProvider->EnumerateMotionMatchingStateTimelines(ObjectId, [StartTime, EndTime, this](const FTraceProvider::FMotionMatchingStateTimeline& InTimeline)
 		{
 			// this isn't very efficient, and it gets called every frame.  will need optimizing
@@ -479,8 +478,6 @@ bool FDebuggerTrack::UpdateInternal()
 	check(AnalysisSession);
 	if (const FTraceProvider* PoseSearchProvider = AnalysisSession->ReadProvider<FTraceProvider>(FTraceProvider::ProviderName))
 	{
-		TraceServices::FAnalysisSessionReadScope SessionReadScope(*AnalysisSession);
-
 		// convert time range to from rewind debugger times to profiler times
 		TRange<double> TraceTimeRange = RewindDebugger->GetCurrentTraceRange();
 		double StartTime = TraceTimeRange.GetLowerBoundValue();
@@ -495,6 +492,8 @@ bool FDebuggerTrack::UpdateInternal()
 		TMap<int32, FText, TInlineSetAllocator<64>> SearchIdNames;
 		const IGameplayProvider* GameplayProvider = AnalysisSession->ReadProvider<IGameplayProvider>("GameplayProvider");
 
+		TraceServices::FProviderReadScopeLock GameplayProviderReadScope(*GameplayProvider);
+		TraceServices::FProviderReadScopeLock PoseSearchProviderReadScope(*PoseSearchProvider);
 		PoseSearchProvider->EnumerateMotionMatchingStateTimelines(ObjectId, [StartTime, EndTime, &SearchIdNames, GameplayProvider](const FTraceProvider::FMotionMatchingStateTimeline& InTimeline)
 			{
 				// this isn't very efficient, and it gets called every frame.  will need optimizing
diff --git a/Engine/Plugins/Animation/PoseSearch/Source/Editor/Private/PoseSearchDebuggerDatabaseView.cpp b/Engine/Plugins/Animation/PoseSearch/Source/Editor/Private/PoseSearchDebuggerDatabaseView.cpp
--- a/Engine/Plugins/Animation/PoseSearch/Source/Editor/Private/PoseSearchDebuggerDatabaseView.cpp
+++ b/Engine/Plugins/Animation/PoseSearch/Source/Editor/Private/PoseSearchDebuggerDatabaseView.cpp
@@ -6,6 +6,7 @@
 #include "Animation/AnimComposite.h"
 #include "Animation/AnimMontage.h"
 #include "Animation/AnimSequence.h"
+#include "Common/ProviderLock.h"
 #include "IGameplayProvider.h"
 #include "Internationalization/Regex.h"
 #include "IRewindDebugger.h"
@@ -35,12 +36,20 @@ UPoseSearchDatabase* ResolveDatabaseFromId(uint64 DatabaseId)
 	UPoseSearchDatabase* Database = nullptr;
 	if (const TraceServices::IAnalysisSession* AnalysisSession = IRewindDebugger::Instance()->GetAnalysisSession())
 	{
-		TraceServices::FAnalysisSessionReadScope SessionReadScope(*AnalysisSession);
 		if (const IGameplayProvider* GameplayProvider = AnalysisSession->ReadProvider<IGameplayProvider>("GameplayProvider"))
 		{
-			if (const FObjectInfo* DatabaseInfo = GameplayProvider->FindObjectInfo(DatabaseId))
+			FString PathName;
 			{
-				Database = TSoftObjectPtr<UPoseSearchDatabase>(FSoftObjectPath(DatabaseInfo->PathName)).LoadSynchronous();
+				TraceServices::FProviderReadScopeLock ProviderReadScope(*GameplayProvider);
+				if (const FObjectInfo* DatabaseInfo = GameplayProvider->FindObjectInfo(DatabaseId))
+				{
+					PathName = DatabaseInfo->PathName;
+				}
+			}
+
+			if (!PathName.IsEmpty())
+			{
+				Database = TSoftObjectPtr<UPoseSearchDatabase>(FSoftObjectPath(PathName)).LoadSynchronous();
 			}
 		}
 	}
diff --git a/Engine/Plugins/Animation/PoseSearch/Source/Editor/Private/PoseSearchDebuggerTrackCreator.cpp b/Engine/Plugins/Animation/PoseSearch/Source/Editor/Private/PoseSearchDebuggerTrackCreator.cpp
--- a/Engine/Plugins/Animation/PoseSearch/Source/Editor/Private/PoseSearchDebuggerTrackCreator.cpp
+++ b/Engine/Plugins/Animation/PoseSearch/Source/Editor/Private/PoseSearchDebuggerTrackCreator.cpp
@@ -35,18 +35,15 @@ bool FDebuggerTrackCreator::HasDebugInfoInternal(const RewindDebugger::FObjectId
 	TRACE_CPUPROFILER_EVENT_SCOPE(PoseSearchDebugger::HasDebugInfoInternal);
 	// Get provider and validate
 	const TraceServices::IAnalysisSession* Session = IRewindDebugger::Instance()->GetAnalysisSession();
-	TraceServices::FAnalysisSessionReadScope SessionReadScope(*Session);
-
-	const UE::PoseSearch::FTraceProvider* PoseSearchProvider = Session->ReadProvider<FTraceProvider>(FTraceProvider::ProviderName);
-	const IAnimationProvider* AnimationProvider = Session->ReadProvider<IAnimationProvider>("AnimationProvider");
-	const IGameplayProvider* GameplayProvider = Session->ReadProvider<IGameplayProvider>("GameplayProvider");
-	if (!(PoseSearchProvider && AnimationProvider && GameplayProvider))
+	const FTraceProvider* PoseSearchProvider = Session->ReadProvider<FTraceProvider>(FTraceProvider::ProviderName);
+	if (!PoseSearchProvider)
 	{
 		return false;
 	}
 
 	bool bHasData = false;
 
+	TraceServices::FProviderReadScopeLock ProviderReadScope(*PoseSearchProvider);
 	PoseSearchProvider->EnumerateMotionMatchingStateTimelines(InObjectId.GetMainId(), [&bHasData](const FTraceProvider::FMotionMatchingStateTimeline& InTimeline)
 	{
 		bHasData = true;
diff --git a/Engine/Plugins/Animation/PoseSearch/Source/Editor/Private/PoseSearchDebuggerView.cpp b/Engine/Plugins/Animation/PoseSearch/Source/Editor/Private/PoseSearchDebuggerView.cpp
--- a/Engine/Plugins/Animation/PoseSearch/Source/Editor/Private/PoseSearchDebuggerView.cpp
+++ b/Engine/Plugins/Animation/PoseSearch/Source/Editor/Private/PoseSearchDebuggerView.cpp
@@ -2,6 +2,7 @@
 
 #include "PoseSearchDebuggerView.h"
 #include "Animation/TrajectoryTypes.h"
+#include "Common/ProviderLock.h"
 #include "Editor.h"
 #include "IGameplayProvider.h"
 #include "Modules/ModuleManager.h"
@@ -523,23 +524,39 @@ bool SDebuggerView::UpdateNodeSelection()
 		IRewindDebugger* RewindDebugger = IRewindDebugger::Instance();
 		const TraceServices::IAnalysisSession* AnalysisSession = RewindDebugger->GetAnalysisSession();
 		check(AnalysisSession);
-		const IGameplayProvider* GameplayProvider = AnalysisSession->ReadProvider<IGameplayProvider>("GameplayProvider");
-		TraceServices::FAnalysisSessionReadScope SessionReadScope(*AnalysisSession);
 
-		for (const FTraceMotionMatchingStateMessage& MotionMatchingState : MotionMatchingStates)
+		// Cache display text under a tight lock scope
+		struct FCachedSearchEntry
 		{
-			Model->OnUpdateSearchSelection(MotionMatchingState.GetSearchId());
+			FText DisplayText;
+			int32 SearchId;
+		};
+
+		TArray<FCachedSearchEntry> CachedEntries;
+		{
+			const IGameplayProvider* GameplayProvider = AnalysisSession->ReadProvider<IGameplayProvider>("GameplayProvider");
+			TraceServices::FProviderReadScopeLock ProviderReadScope(*GameplayProvider);
+			for (const FTraceMotionMatchingStateMessage& MotionMatchingState : MotionMatchingStates)
+			{
+				CachedEntries.Add({ GenerateSearchNameWithBestScore(MotionMatchingState, GameplayProvider), MotionMatchingState.GetSearchId() });
+			}
+		}
+
+		// Build widgets outside the lock
+		for (const FCachedSearchEntry& Entry : CachedEntries)
+		{
+			Model->OnUpdateSearchSelection(Entry.SearchId);
 			SelectionView->AddSlot()
 				.HAlign(HAlign_Fill)
 				.VAlign(VAlign_Center)
 				.Padding(10.0f)
 				[
 					SNew(SButton)
-						.Text(GenerateSearchNameWithBestScore(MotionMatchingState, GameplayProvider))
+						.Text(Entry.DisplayText)
 						.HAlign(HAlign_Center)
 						.VAlign(VAlign_Center)
 						.ContentPadding(10.0f)
-						.OnClicked(this, &SDebuggerView::OnUpdateSearchSelection, MotionMatchingState.GetSearchId())
+						.OnClicked(this, &SDebuggerView::OnUpdateSearchSelection, Entry.SearchId)
 				];
 		}
 	}
diff --git a/Engine/Plugins/Animation/PoseSearch/Source/Editor/Private/PoseSearchDebuggerViewModel.cpp b/Engine/Plugins/Animation/PoseSearch/Source/Editor/Private/PoseSearchDebuggerViewModel.cpp
--- a/Engine/Plugins/Animation/PoseSearch/Source/Editor/Private/PoseSearchDebuggerViewModel.cpp
+++ b/Engine/Plugins/Animation/PoseSearch/Source/Editor/Private/PoseSearchDebuggerViewModel.cpp
@@ -4,6 +4,7 @@
 #include "Animation/AnimComposite.h"
 #include "Animation/AnimSequence.h"
 #include "Animation/MirrorDataTable.h"
+#include "Common/ProviderLock.h"
 #include "Engine/SkeletalMesh.h"
 #include "IAnimationProvider.h"
 #include "IGameplayProvider.h"
@@ -55,7 +56,6 @@ void FDebuggerViewModel::OnUpdate()
 
 	// Get provider and validate
 	const TraceServices::IAnalysisSession* Session = RewindDebugger.Get()->GetAnalysisSession();
-	TraceServices::FAnalysisSessionReadScope SessionReadScope(*Session);
 
 	const FTraceProvider* PoseSearchProvider = Session->ReadProvider<FTraceProvider>(FTraceProvider::ProviderName);
 	const IAnimationProvider* AnimationProvider = Session->ReadProvider<IAnimationProvider>("AnimationProvider");
@@ -65,25 +65,37 @@ void FDebuggerViewModel::OnUpdate()
 		return;
 	}
 	const double TraceTime = RewindDebugger.Get()->CurrentTraceTime();
+
+	bool bFrameFound;
 	TraceServices::FFrame Frame;
-	ReadFrameProvider(*Session).GetFrameFromTime(TraceFrameType_Game, TraceTime, Frame);
-	PoseSearchProvider->EnumerateMotionMatchingStateTimelines(AnimInstanceId, [&Frame, this](const FTraceProvider::FMotionMatchingStateTimeline& InTimeline)
 	{
-		const FTraceMotionMatchingStateMessage* Message = nullptr;
+		TraceServices::FAnalysisSessionReadScope SessionReadScope(*Session);
+		const TraceServices::IFrameProvider& FrameProvider = TraceServices::ReadFrameProvider(*Session);
+		bFrameFound = ReadFrameProvider(*Session).GetFrameFromTime(TraceFrameType_Game, TraceTime, Frame);
+	}
 
-		InTimeline.EnumerateEvents(Frame.StartTime, Frame.EndTime, [&Message](double InStartTime, double InEndTime, const FTraceMotionMatchingStateMessage& InMessage)
+	// Phase 1: Read motion matching states under lock
+	if (bFrameFound)
+	{
+		TraceServices::FProviderReadScopeLock PoseSearchProviderReadScope(*PoseSearchProvider);
+		PoseSearchProvider->EnumerateMotionMatchingStateTimelines(AnimInstanceId, [&Frame, this](const FTraceProvider::FMotionMatchingStateTimeline& InTimeline)
 		{
-			Message = &InMessage;
-			return TraceServices::EEventEnumerate::Stop;
+			const FTraceMotionMatchingStateMessage* Message = nullptr;
+
+			InTimeline.EnumerateEvents(Frame.StartTime, Frame.EndTime, [&Message](double InStartTime, double InEndTime, const FTraceMotionMatchingStateMessage& InMessage)
+			{
+				Message = &InMessage;
+				return TraceServices::EEventEnumerate::Stop;
+			});
+			if (Message)
+			{
+				check(Message->Roles.Num() == Message->SkeletalMeshComponentIds.Num());
+				MotionMatchingStates.Add(*Message);
+			}
 		});
-		if (Message)
-		{
-			check(Message->Roles.Num() == Message->SkeletalMeshComponentIds.Num());
-			MotionMatchingStates.Add(*Message);
-		}
-	});
-	
-	// adding all the missing DebugDrawActors, and tracking them via UsedDebugDrawActorIds
+	}
+
+	// Phase 2: Actor spawning/cleanup outside any lock
 	TSet<uint64> UsedDebugDrawActorIds;
 	for (const FTraceMotionMatchingStateMessage& MotionMatchingState : MotionMatchingStates)
 	{
@@ -136,48 +148,69 @@ void FDebuggerViewModel::OnUpdate()
 		return;
 	}
 
-	for (TPair<uint64, TWeakObjectPtr<AActor>>& DebugDrawActorPair : DebugDrawActors)
+	// Phase 3: Read pose data under lock, cache mesh path names, then apply outside lock
+	struct FCachedPoseData
 	{
-		const uint64 SkeletalMeshComponentId = DebugDrawActorPair.Key;
-		if (AActor* DebugDrawActor = DebugDrawActorPair.Value.Get())
+		UPoseSearchMeshComponent* MeshComponent;
+		FString MeshPathName;
+		FTransform ComponentWorldTransform;
+		TArray<FTransform> ComponentSpaceTransforms;
+	};
+	TArray<FCachedPoseData> CachedPoseEntries;
+
+	{
+		TraceServices::FProviderReadScopeLock AnimationProviderReadScope(*AnimationProvider);
+		TraceServices::FProviderReadScopeLock GameplayProviderReadScope(*GameplayProvider);
+
+		for (TPair<uint64, TWeakObjectPtr<AActor>>& DebugDrawActorPair : DebugDrawActors)
 		{
-			for (UActorComponent* ActorComponent : DebugDrawActor->GetInstanceComponents())
+			const uint64 SkeletalMeshComponentId = DebugDrawActorPair.Key;
+			if (AActor* DebugDrawActor = DebugDrawActorPair.Value.Get())
 			{
-				if (UPoseSearchMeshComponent* PoseSearchMeshComponent = Cast<UPoseSearchMeshComponent>(ActorComponent))
+				for (UActorComponent* ActorComponent : DebugDrawActor->GetInstanceComponents())
 				{
-					AnimationProvider->ReadSkeletalMeshPoseTimeline(SkeletalMeshComponentId, [&Frame, AnimationProvider, GameplayProvider, PoseSearchMeshComponent](const IAnimationProvider::SkeletalMeshPoseTimeline& TimelineData, bool bHasCurves)
-						{
-							TimelineData.EnumerateEvents(Frame.StartTime, Frame.EndTime, [AnimationProvider, GameplayProvider, PoseSearchMeshComponent](double InStartTime, double InEndTime, uint32 InDepth, const FSkeletalMeshPoseMessage& PoseMessage) -> TraceServices::EEventEnumerate
-								{
-									const FSkeletalMeshInfo* SkeletalMeshInfo = AnimationProvider->FindSkeletalMeshInfo(PoseMessage.MeshId);
-									const FObjectInfo* SkeletalMeshObjectInfo = GameplayProvider->FindObjectInfo(PoseMessage.MeshId);
-									if (!SkeletalMeshInfo || !SkeletalMeshObjectInfo)
-									{
-										return TraceServices::EEventEnumerate::Stop;
-									}
-
-									USkeletalMesh* SkeletalMesh = TSoftObjectPtr<USkeletalMesh>(FSoftObjectPath(SkeletalMeshObjectInfo->PathName)).LoadSynchronous();
-									if (SkeletalMesh)
+					if (UPoseSearchMeshComponent* PoseSearchMeshComponent = Cast<UPoseSearchMeshComponent>(ActorComponent))
+					{
+						AnimationProvider->ReadSkeletalMeshPoseTimeline(SkeletalMeshComponentId, [&Frame, AnimationProvider, GameplayProvider, PoseSearchMeshComponent, &CachedPoseEntries](const IAnimationProvider::SkeletalMeshPoseTimeline& TimelineData, bool bHasCurves)
+							{
+								TimelineData.EnumerateEvents(Frame.StartTime, Frame.EndTime, [AnimationProvider, GameplayProvider, PoseSearchMeshComponent, &CachedPoseEntries](double InStartTime, double InEndTime, uint32 InDepth, const FSkeletalMeshPoseMessage& PoseMessage) -> TraceServices::EEventEnumerate
 									{
-										PoseSearchMeshComponent->SetSkinnedAssetAndUpdate(SkeletalMesh, true);
-									}
-									FTransform ComponentWorldTransform;
-									// Active skeleton is simply the traced bone transforms
-									TArray<FTransform>& ComponentSpaceTransforms = PoseSearchMeshComponent->GetEditableComponentSpaceTransforms();
-									AnimationProvider->GetSkeletalMeshComponentSpacePose(PoseMessage, *SkeletalMeshInfo, ComponentWorldTransform, ComponentSpaceTransforms);
+										const FSkeletalMeshInfo* SkeletalMeshInfo = AnimationProvider->FindSkeletalMeshInfo(PoseMessage.MeshId);
+										const FObjectInfo* SkeletalMeshObjectInfo = GameplayProvider->FindObjectInfo(PoseMessage.MeshId);
+										if (!SkeletalMeshInfo || !SkeletalMeshObjectInfo)
+										{
+											return TraceServices::EEventEnumerate::Stop;
+										}
 
-									check(ComponentWorldTransform.Equals(PoseMessage.ComponentToWorld));
+										FCachedPoseData& CachedEntry = CachedPoseEntries.AddDefaulted_GetRef();
+										CachedEntry.MeshComponent = PoseSearchMeshComponent;
+										CachedEntry.MeshPathName = SkeletalMeshObjectInfo->PathName;
+										AnimationProvider->GetSkeletalMeshComponentSpacePose(PoseMessage, *SkeletalMeshInfo, CachedEntry.ComponentWorldTransform, CachedEntry.ComponentSpaceTransforms);
 
-									PoseSearchMeshComponent->Initialize(ComponentWorldTransform);
+										check(CachedEntry.ComponentWorldTransform.Equals(PoseMessage.ComponentToWorld));
 
-									return TraceServices::EEventEnumerate::Stop;
-								});
-						});
-					break;
+										return TraceServices::EEventEnumerate::Stop;
+									});
+							});
+						break;
+					}
 				}
 			}
 		}
 	}
+
+	// Apply cached pose data outside the lock (includes LoadSynchronous)
+	for (FCachedPoseData& CachedEntry : CachedPoseEntries)
+	{
+		USkeletalMesh* SkeletalMesh = TSoftObjectPtr<USkeletalMesh>(FSoftObjectPath(CachedEntry.MeshPathName)).LoadSynchronous();
+		if (SkeletalMesh)
+		{
+			CachedEntry.MeshComponent->SetSkinnedAssetAndUpdate(SkeletalMesh, true);
+		}
+		TArray<FTransform>& ComponentSpaceTransforms = CachedEntry.MeshComponent->GetEditableComponentSpaceTransforms();
+		ComponentSpaceTransforms = MoveTemp(CachedEntry.ComponentSpaceTransforms);
+		CachedEntry.MeshComponent->Initialize(CachedEntry.ComponentWorldTransform);
+	}
 }
 
 void FDebuggerViewModel::OnUpdateSearchSelection(int32 InSearchId)
diff --git a/Engine/Plugins/Animation/PoseSearch/Source/Editor/Private/Trace/PoseSearchTraceAnalyzer.cpp b/Engine/Plugins/Animation/PoseSearch/Source/Editor/Private/Trace/PoseSearchTraceAnalyzer.cpp
--- a/Engine/Plugins/Animation/PoseSearch/Source/Editor/Private/Trace/PoseSearchTraceAnalyzer.cpp
+++ b/Engine/Plugins/Animation/PoseSearch/Source/Editor/Private/Trace/PoseSearchTraceAnalyzer.cpp
@@ -2,6 +2,7 @@
 
 #include "PoseSearchTraceAnalyzer.h"
 
+#include "Common/ProviderLock.h"
 #include "HAL/LowLevelMemTracker.h"
 #include "PoseSearchTraceProvider.h"
 #include "PoseSearch/Trace/PoseSearchTraceLogger.h"
@@ -29,7 +30,7 @@ bool FTraceAnalyzer::OnEvent(uint16 RouteId, EStyle Style, const FOnEventContext
 {
 	LLM_SCOPE_BYNAME(TEXT("Insights/PoseSearch::FTraceAnalyzer"));
 
-	TraceServices::FAnalysisSessionEditScope Scope(Session);
+	TraceServices::FProviderEditScopeLock ProviderEditScope(TraceProvider);
 
 	FPoseSearchCustomVersion::Type CustomVersionType = FPoseSearchCustomVersion::BeforeCustomVersionWasAdded;
 
diff --git a/Engine/Plugins/Animation/PoseSearch/Source/Editor/Private/Trace/PoseSearchTraceProvider.cpp b/Engine/Plugins/Animation/PoseSearch/Source/Editor/Private/Trace/PoseSearchTraceProvider.cpp
--- a/Engine/Plugins/Animation/PoseSearch/Source/Editor/Private/Trace/PoseSearchTraceProvider.cpp
+++ b/Engine/Plugins/Animation/PoseSearch/Source/Editor/Private/Trace/PoseSearchTraceProvider.cpp
@@ -5,6 +5,8 @@
 namespace UE::PoseSearch
 {
 
+thread_local TraceServices::FProviderLock::FThreadLocalState GPoseSearchProviderLockState;
+
 const FName FTraceProvider::ProviderName("PoseSearchTraceProvider");
 
 FTraceProvider::FTraceProvider(TraceServices::IAnalysisSession& InSession) : Session(InSession)
@@ -13,24 +15,27 @@ FTraceProvider::FTraceProvider(TraceServices::IAnalysisSession& InSession) : Ses
 
 bool FTraceProvider::ReadMotionMatchingStateTimeline(uint64 InAnimInstanceId, int32 InSearchId, TFunctionRef<void(const FMotionMatchingStateTimeline&)> Callback) const
 {
-	Session.ReadAccessCheck();
+	ReadAccessCheck();
 	return MotionMatchingStateTimelineStorage.ReadTimeline(InAnimInstanceId, InSearchId, Callback);
 }
 
 bool FTraceProvider::EnumerateMotionMatchingStateTimelines(uint64 InAnimInstanceId, TFunctionRef<void(const FMotionMatchingStateTimeline&)> Callback) const
 {
-	Session.ReadAccessCheck();
+	ReadAccessCheck();
 	return MotionMatchingStateTimelineStorage.EnumerateSearchTimelines(InAnimInstanceId, Callback);
 }
 
 void FTraceProvider::AppendMotionMatchingState(const FTraceMotionMatchingStateMessage& InMessage, double InTime)
 {
-	Session.WriteAccessCheck();
+	EditAccessCheck();
 
 	TSharedRef<TraceServices::TPointTimeline<FTraceMotionMatchingStateMessage>> Timeline = MotionMatchingStateTimelineStorage.GetTimeline(Session, InMessage.AnimInstanceId, InMessage.GetSearchId());
 	Timeline->AppendEvent(InTime, InMessage);
-	
-	Session.UpdateDurationSeconds(InTime);
+
+	{
+		TraceServices::FAnalysisSessionEditScope SessionEditScope(Session);
+		Session.UpdateDurationSeconds(InTime);
+	}
 }
 
 } // namespace UE::PoseSearch
diff --git a/Engine/Plugins/Animation/PoseSearch/Source/Editor/Private/Trace/PoseSearchTraceProvider.h b/Engine/Plugins/Animation/PoseSearch/Source/Editor/Private/Trace/PoseSearchTraceProvider.h
--- a/Engine/Plugins/Animation/PoseSearch/Source/Editor/Private/Trace/PoseSearchTraceProvider.h
+++ b/Engine/Plugins/Animation/PoseSearch/Source/Editor/Private/Trace/PoseSearchTraceProvider.h
@@ -2,7 +2,7 @@
 
 #pragma once
 
-#include "CoreMinimal.h"
+#include "Common/ProviderLock.h"
 #include "Model/PointTimeline.h"
 #include "PoseSearch/Trace/PoseSearchTraceLogger.h"
 #include "TraceServices/Model/AnalysisSession.h"
@@ -10,14 +10,44 @@
 namespace UE::PoseSearch
 {
 
+extern thread_local TraceServices::FProviderLock::FThreadLocalState GPoseSearchProviderLockState;
+
 /**
  * Provider to the widgets for pose search functionality, largely mimicking FAnimationProvider
  */
-class FTraceProvider : public TraceServices::IProvider
+class FTraceProvider : public TraceServices::IProvider, public TraceServices::IEditableProvider
 {
 public:
 	explicit FTraceProvider(TraceServices::IAnalysisSession& InSession);
 
+	//~ IProvider (read lock)
+	virtual void BeginRead() const override
+	{
+		Lock.BeginRead(GPoseSearchProviderLockState);
+	}
+	virtual void EndRead() const override
+	{
+		Lock.EndRead(GPoseSearchProviderLockState);
+	}
+	virtual void ReadAccessCheck() const override
+	{
+		Lock.ReadAccessCheck(GPoseSearchProviderLockState);
+	}
+
+	//~ IEditableProvider (edit lock)
+	virtual void BeginEdit() const override
+	{
+		Lock.BeginWrite(GPoseSearchProviderLockState);
+	}
+	virtual void EndEdit() const override
+	{
+		Lock.EndWrite(GPoseSearchProviderLockState);
+	}
+	virtual void EditAccessCheck() const override
+	{
+		Lock.WriteAccessCheck(GPoseSearchProviderLockState);
+	}
+
 	using FMotionMatchingStateTimeline = TraceServices::ITimeline<FTraceMotionMatchingStateMessage>;
 	
 	/** Read the node-relative timeline info of the AnimInstance and provide it via callback */
@@ -151,5 +181,6 @@ private:
 
 
 	TraceServices::IAnalysisSession& Session;
+	mutable TraceServices::FProviderLock Lock;
 };
 } // namespace UE::PoseSearch
diff --git a/Engine/Plugins/Chooser/Source/ChooserEditor/Private/ChooserAnalyzer.cpp b/Engine/Plugins/Chooser/Source/ChooserEditor/Private/ChooserAnalyzer.cpp
--- a/Engine/Plugins/Chooser/Source/ChooserEditor/Private/ChooserAnalyzer.cpp
+++ b/Engine/Plugins/Chooser/Source/ChooserEditor/Private/ChooserAnalyzer.cpp
@@ -5,6 +5,7 @@
 #include "HAL/LowLevelMemTracker.h"
 #include "Serialization/MemoryReader.h"
 #include "TraceServices/Model/AnalysisSession.h"
+#include "Common/ProviderLock.h"
 #include "ChooserProvider.h"
 #include "TraceServices/Utils.h"
 
@@ -27,7 +28,7 @@ bool FChooserAnalyzer::OnEvent(uint16 RouteId, EStyle Style, const FOnEventConte
 {
 	LLM_SCOPE_BYNAME(TEXT("Insights/FChooserAnalyzer"));
 
-	TraceServices::FAnalysisSessionEditScope _(Session);
+	TraceServices::FProviderEditScopeLock ProviderEditScope(Provider);
 
 	const auto& EventData = Context.EventData;
 	switch (RouteId)
diff --git a/Engine/Plugins/Chooser/Source/ChooserEditor/Private/ChooserProvider.cpp b/Engine/Plugins/Chooser/Source/ChooserEditor/Private/ChooserProvider.cpp
--- a/Engine/Plugins/Chooser/Source/ChooserEditor/Private/ChooserProvider.cpp
+++ b/Engine/Plugins/Chooser/Source/ChooserEditor/Private/ChooserProvider.cpp
@@ -3,6 +3,11 @@
 #include "ChooserProvider.h"
 #include "ObjectTrace.h"
 
+namespace TraceServices
+{
+thread_local FProviderLock::FThreadLocalState GChooserProviderLockState;
+}
+
 FName FChooserProvider::ProviderName("ChooserProvider");
 
 #define LOCTEXT_NAMESPACE "ChooserProvider"
@@ -14,7 +19,7 @@ FChooserProvider::FChooserProvider(TraceServices::IAnalysisSession& InSession)
 
 bool FChooserProvider::ReadChooserEvaluationTimeline(uint64 InObjectId, TFunctionRef<void(const ChooserEvaluationTimeline&)> Callback) const
 {
-	Session.ReadAccessCheck();
+	ReadAccessCheck();
 
 	const uint32* IndexPtr = ObjectIdToChooserEvaluationTimelines.Find(InObjectId);
 	if(IndexPtr != nullptr)
@@ -31,7 +36,7 @@ bool FChooserProvider::ReadChooserEvaluationTimeline(uint64 InObjectId, TFunctio
 
 bool FChooserProvider::ReadChooserValueTimeline(uint64 InObjectId, TFunctionRef<void(const ChooserValueTimeline&)> Callback) const
 {
-	Session.ReadAccessCheck();
+	ReadAccessCheck();
 
 	const uint32* IndexPtr = ObjectIdToChooserEvaluationTimelines.Find(InObjectId);
 	if(IndexPtr != nullptr)
@@ -48,7 +53,7 @@ bool FChooserProvider::ReadChooserValueTimeline(uint64 InObjectId, TFunctionRef<
 
 void FChooserProvider::EnumerateChooserEvaluationTimelines(TFunctionRef<void(uint64 OwnerId, const ChooserEvaluationTimeline&)> Callback) const
 {
-	Session.ReadAccessCheck();
+	ReadAccessCheck();
 
 	for (TTuple<uint64, uint32> Entry : ObjectIdToChooserEvaluationTimelines)
 	{
@@ -79,7 +84,7 @@ uint32 FChooserProvider::GetTimelineIndex(uint64 InObjectId)
 
 void  FChooserProvider::AppendChooserEvaluation(uint64 InChooserId, uint64 InObjectId, TArrayView<const int32> InSelectedIndices, double InRecordingTime)
 {
-	Session.WriteAccessCheck();
+	EditAccessCheck();
 
 	TSharedPtr<TraceServices::TPointTimeline<FChooserEvaluationData>> Timeline = ChooserEvaluationTimelines[GetTimelineIndex(InObjectId)];
 	Timeline->AppendEvent(InRecordingTime, { InChooserId, TArray<int32>(InSelectedIndices) });
diff --git a/Engine/Plugins/Chooser/Source/ChooserEditor/Private/ChooserProvider.h b/Engine/Plugins/Chooser/Source/ChooserEditor/Private/ChooserProvider.h
--- a/Engine/Plugins/Chooser/Source/ChooserEditor/Private/ChooserProvider.h
+++ b/Engine/Plugins/Chooser/Source/ChooserEditor/Private/ChooserProvider.h
@@ -8,8 +8,13 @@
 #include "Templates/SharedPointer.h"
 #include "Model/IntervalTimeline.h"
 #include "TraceServices/Model/AnalysisSession.h"
+#include "Common/ProviderLock.h"
 
-namespace TraceServices { class IAnalysisSession; }
+namespace TraceServices
+{
+	class IAnalysisSession;
+	extern thread_local FProviderLock::FThreadLocalState GChooserProviderLockState;
+}
 
 struct FChooserEvaluationData
 {
@@ -25,16 +30,48 @@ struct FChooserValueData
 	TArray<uint8> Value;
 };
 
-class FChooserProvider : public TraceServices::IProvider
+class FChooserProvider : public TraceServices::IProvider, public TraceServices::IEditableProvider
 {
 public:
 	typedef TraceServices::ITimeline<FChooserEvaluationData> ChooserEvaluationTimeline;
 	typedef TraceServices::ITimeline<FChooserValueData> ChooserValueTimeline;
-	
+
 	static FName ProviderName;
 
 	FChooserProvider(TraceServices::IAnalysisSession& InSession);
 
+	//~ IProvider (read lock)
+	virtual void BeginRead() const override
+	{
+		Lock.BeginRead(TraceServices::GChooserProviderLockState);
+	}
+
+	virtual void EndRead() const override
+	{
+		Lock.EndRead(TraceServices::GChooserProviderLockState);
+	}
+
+	virtual void ReadAccessCheck() const override
+	{
+		Lock.ReadAccessCheck(TraceServices::GChooserProviderLockState);
+	}
+
+	//~ IEditableProvider (edit lock)
+	virtual void BeginEdit() const override
+	{
+		Lock.BeginWrite(TraceServices::GChooserProviderLockState);
+	}
+
+	virtual void EndEdit() const override
+	{
+		Lock.EndWrite(TraceServices::GChooserProviderLockState);
+	}
+
+	virtual void EditAccessCheck() const override
+	{
+		Lock.WriteAccessCheck(TraceServices::GChooserProviderLockState);
+	}
+
 	bool ReadChooserEvaluationTimeline(uint64 InObjectId, TFunctionRef<void(const ChooserEvaluationTimeline&)> Callback) const;
 	bool ReadChooserValueTimeline(uint64 InObjectId, TFunctionRef<void(const ChooserValueTimeline&)> Callback) const;
 	void EnumerateChooserEvaluationTimelines(TFunctionRef<void(uint64 OwnerId, const ChooserEvaluationTimeline&)> Callback) const;
@@ -45,7 +82,7 @@ public:
 
 	void AppendChooserValue(uint64 InChooserId, uint64 InObjectId, double InRecordingTime, const FName& Key, TArrayView<const uint8> SerializedValue)
     {
-    	Session.WriteAccessCheck();
+    	EditAccessCheck();
     
     	TSharedPtr<TraceServices::TPointTimeline<FChooserValueData>> Timeline = ChooserValueTimelines[GetTimelineIndex(InObjectId)];
 		FChooserValueData EventData;
@@ -57,8 +94,9 @@ public:
 
 private:
 	uint32 GetTimelineIndex(uint64 ObjectId);
-	
+
 	TraceServices::IAnalysisSession& Session;
+	mutable TraceServices::FProviderLock Lock;
 
 	TMap<uint64, uint32> ObjectIdToChooserEvaluationTimelines;
 	TArray<TSharedRef<TraceServices::TPointTimeline<FChooserEvaluationData>>> ChooserEvaluationTimelines;
