Fix crash with nullptr trace session when attempting to debug anim graph via anim graph rewind debugger track.
#rnx

#Synced-CL 52374856

diff --git a/Engine/Plugins/Animation/GameplayInsights/Source/RewindDebugger/Private/AnimInstanceHelpers.cpp b/Engine/Plugins/Animation/GameplayInsights/Source/RewindDebugger/Private/AnimInstanceHelpers.cpp
--- a/Engine/Plugins/Animation/GameplayInsights/Source/RewindDebugger/Private/AnimInstanceHelpers.cpp
+++ b/Engine/Plugins/Animation/GameplayInsights/Source/RewindDebugger/Private/AnimInstanceHelpers.cpp
@@ -12,6 +12,7 @@
 #include "RewindDebuggerAnimation.h"
 #include "Subsystems/AssetEditorSubsystem.h"
 #include "IAnimationBlueprintEditor.h"
+#include "RewindDebugger.h"
 #include "RewindDebuggerModule.h"
 #include "Insights/IUnrealInsightsModule.h"
 #include "ToolMenus.h"
@@ -110,9 +111,11 @@ void FAnimInstanceMenu::Register()
 							FSlateIcon(),
 							FExecuteAction::CreateLambda([ObjectId = Context->SelectedObject->GetUObjectId()]()
 							{
-								const IUnrealInsightsModule& UnrealInsightsModule = FModuleManager::LoadModuleChecked<IUnrealInsightsModule>("TraceInsights");
-								const TSharedPtr<const TraceServices::IAnalysisSession> Session = UnrealInsightsModule.GetAnalysisSession();
-								OpenAnimBlueprintAndAttachDebugger(Session.Get(), ObjectId);
+								const TraceServices::IAnalysisSession* Session = FRewindDebugger::Instance() ? FRewindDebugger::Instance()->GetAnalysisSession() : nullptr;
+								if (Session)
+								{
+									OpenAnimBlueprintAndAttachDebugger(Session, ObjectId);
+								}
 							}));
 		}
 	}));
