Fix to allow the DisplayWorldId to be set while analyzing local sessions.  Without this, the debugged mesh is lost after the PIE session ends

#Synced-CL 52591684

diff --git a/Engine/Plugins/Animation/GameplayInsights/Source/RewindDebugger/Private/RewindDebugger.cpp b/Engine/Plugins/Animation/GameplayInsights/Source/RewindDebugger/Private/RewindDebugger.cpp
--- a/Engine/Plugins/Animation/GameplayInsights/Source/RewindDebugger/Private/RewindDebugger.cpp
+++ b/Engine/Plugins/Animation/GameplayInsights/Source/RewindDebugger/Private/RewindDebugger.cpp
@@ -1332,9 +1332,8 @@ void FRewindDebugger::Tick(float DeltaTime)
 
 			RecordingDurationValue = GameplayProvider->GetRecordingDuration();
 
-			// set a default display world when loading a trace or analyzing a remote session (first client/standalone world)
-			if (!bDisplayWorldIdValid
-				&& (IsAnalyzingTraceFile() || IsAnalyzingRemoteSession()))
+			// set a default display world when analyzing a session (first client/standalone world)
+			if (!bDisplayWorldIdValid && (GetAnalysisSession() != nullptr))
 			{
 				GameplayProvider->EnumerateWorlds([this](const FWorldInfo& WorldInfo)
 					{
