[RewindDebugger]
Add a class icon for MassEntityHandle and register it in the Mass editor style
Broaden the GameplayInsights type lookup from UClass to UStruct to support struct-based types, deprecating FindClass in favor of FindType
Add a configurable verbosity setting for StateTree Rewind Debugger track names, defaulting to compact display showing only the asset name

#Synced-CL 52727620

diff --git a/Engine/Plugins/AI/MassAI/Source/MassAIBehavior/Private/MassStateTreeProcessors.cpp b/Engine/Plugins/AI/MassAI/Source/MassAIBehavior/Private/MassStateTreeProcessors.cpp
--- a/Engine/Plugins/AI/MassAI/Source/MassAIBehavior/Private/MassStateTreeProcessors.cpp
+++ b/Engine/Plugins/AI/MassAI/Source/MassAIBehavior/Private/MassStateTreeProcessors.cpp
@@ -52,7 +52,7 @@ void ForEachEntityInChunk(
 		{
 			const UObject* TraceOwner = Context.GetEntityManagerChecked().GetOwner();
 			TRACE_OBJECT(TraceOwner);
-			TRACE_INSTANCE(TraceOwner, Entity.AsNumber(), FObjectTrace::GetObjectId(TraceOwner), FMassEntityHandle::StaticStruct(), LexToString(Entity));
+			TRACE_INSTANCE(TraceOwner, Entity.AsNumber(), FObjectTrace::GetObjectId(TraceOwner), FMassEntityHandle::StaticStruct(), FString::Printf(TEXT("Entity[%d:%d]"), Entity.Index, Entity.SerialNumber));
 
 			FMassStateTreeExecutionContext StateTreeContext(MassStateTreeSubsystem, *StateTree, *InstanceData, Context);
 			StateTreeContext.SetEntity(Entity);
diff --git a/Engine/Plugins/Animation/GameplayInsights/Source/GameplayInsights/Private/BlendWeightsTrack.cpp b/Engine/Plugins/Animation/GameplayInsights/Source/GameplayInsights/Private/BlendWeightsTrack.cpp
--- a/Engine/Plugins/Animation/GameplayInsights/Source/GameplayInsights/Private/BlendWeightsTrack.cpp
+++ b/Engine/Plugins/Animation/GameplayInsights/Source/GameplayInsights/Private/BlendWeightsTrack.cpp
@@ -223,11 +223,11 @@ bool FBlendWeightTrack::UpdateInternal()
 				bChanged = true;
 				bool bIsBlendSpace = false;
 
-				if (const UClass* FoundClass = GameplayProvider->FindClass(ObjectInfo->ClassId))
+				if (const UStruct* FoundType = GameplayProvider->FindType(ObjectInfo->ClassId))
 				{
-					Icon = FSlateIconFinder::FindIconForClass(FoundClass);
+					Icon = FSlateIconFinder::FindIconForClass(FoundType);
 #if WITH_EDITOR
-					if (FoundClass->IsChildOf(UBlendSpace::StaticClass()))
+					if (FoundType->IsChildOf(UBlendSpace::StaticClass()))
 					{
 						bIsBlendSpace = true;	
 					}
diff --git a/Engine/Plugins/Animation/GameplayInsights/Source/GameplayInsights/Private/GameplayProvider.cpp b/Engine/Plugins/Animation/GameplayInsights/Source/GameplayInsights/Private/GameplayProvider.cpp
--- a/Engine/Plugins/Animation/GameplayInsights/Source/GameplayInsights/Private/GameplayProvider.cpp
+++ b/Engine/Plugins/Animation/GameplayInsights/Source/GameplayInsights/Private/GameplayProvider.cpp
@@ -56,7 +56,7 @@ static UClass* FindBlueprintClass(const FString& TargetNameRaw)
 
 // This is copied from EditorUtilitySubsystem (where it is not public)
 // Should probably be somewhere shared
-static UClass* FindClassByPathName(const FString& RawTargetPathName, const bool bCheckBlueprintClasses)
+static UStruct* FindTypeByPathName(const FString& RawTargetPathName, const bool bCheckBlueprintClasses)
 {
 	FString TargetName = RawTargetPathName;
 
@@ -90,10 +90,10 @@ static UClass* FindClassByPathName(const FString& RawTargetPathName, const bool
 		}
 	}
 
-	UClass* ResultClass = nullptr;
+	UStruct* ResultClass = nullptr;
 	if (bIsValidClassName)
 	{
-		ResultClass = FindObject<UClass>(nullptr, *TargetName);
+		ResultClass = FindObject<UStruct>(nullptr, *TargetName);
 	}
 
 	// If we still haven't found anything yet, try the asset registry for blueprints that match the requirements
@@ -287,12 +287,12 @@ const FClassInfo* FGameplayProvider::FindClassInfo(uint64 InClassId) const
 	return nullptr;
 }
 
-const UClass* FGameplayProvider::FindClass(uint64 InClassId, bool bSearchBlueprints) const
+const UStruct* FGameplayProvider::FindType(uint64 InTypeId, bool bSearchBlueprints) const
 {
 #if WITH_EDITOR
-	if (const FClassInfo* ClassInfo = FindClassInfo(InClassId))
+	if (const FClassInfo* ClassInfo = FindClassInfo(InTypeId))
 	{
-		return FindClassByPathName(ClassInfo->PathName, bSearchBlueprints);
+		return FindTypeByPathName(ClassInfo->PathName, bSearchBlueprints);
 	}
 	return nullptr;
 #else
@@ -715,9 +715,9 @@ FSlateIcon FGameplayProvider::FindIconForClass(uint64 ClassId) const
 	while (ClassId)
 	{
 		// don't search blueprints because it's very slow and icons are associated with native classes
-		if (const UClass* FoundClass = FindClass(ClassId, false))
+		if (const UStruct* FoundType = FindType(ClassId, false))
 		{
-			Icon = FSlateIconFinder::FindIconForClass(FoundClass);
+			Icon = FSlateIconFinder::FindIconForClass(FoundType);
 			break;
 		}
 
diff --git a/Engine/Plugins/Animation/GameplayInsights/Source/GameplayInsights/Private/GameplayProvider.h b/Engine/Plugins/Animation/GameplayInsights/Source/GameplayInsights/Private/GameplayProvider.h
--- a/Engine/Plugins/Animation/GameplayInsights/Source/GameplayInsights/Private/GameplayProvider.h
+++ b/Engine/Plugins/Animation/GameplayInsights/Source/GameplayInsights/Private/GameplayProvider.h
@@ -73,7 +73,7 @@ public:
 	virtual void EnumerateSubobjects(const RewindDebugger::FObjectId& InObjectId, TFunctionRef<void(const RewindDebugger::FObjectId& SubObjectId)> Callback) const override;
 	virtual const FObjectPropertyValue* FindPropertyValueFromStorageIndex(uint64 InObjectId, int64 InStorageIndex) const override;
 	virtual const FClassInfo* FindClassInfo(uint64 InClassId) const override;
-	virtual const UClass* FindClass(uint64 InClassId, bool bSearchBlueprints = true) const override;
+	virtual const UStruct* FindType(uint64 InTypeId, bool bSearchBlueprints = true) const override;
 	virtual const FClassInfo* FindClassInfo(const TCHAR* InClassPath) const override;
 	virtual bool IsSubClassOf(uint64 InSubClassId, uint64 InParentClassId) const override;
 	virtual const FObjectInfo* FindObjectInfo(const uint64 InObjectId) const override;
diff --git a/Engine/Plugins/Animation/GameplayInsights/Source/GameplayInsights/Private/MontagesTrack.cpp b/Engine/Plugins/Animation/GameplayInsights/Source/GameplayInsights/Private/MontagesTrack.cpp
--- a/Engine/Plugins/Animation/GameplayInsights/Source/GameplayInsights/Private/MontagesTrack.cpp
+++ b/Engine/Plugins/Animation/GameplayInsights/Source/GameplayInsights/Private/MontagesTrack.cpp
@@ -192,9 +192,9 @@ bool FMontageTrack::UpdateInternal()
 				CurveName = FText::FromString(ObjectInfo->Name);
 				bChanged = true;
 
-				if (const UClass* FoundClass = GameplayProvider->FindClass(ObjectInfo->ClassId))
+				if (const UStruct* FoundType = GameplayProvider->FindType(ObjectInfo->ClassId))
 				{
-					Icon = FSlateIconFinder::FindIconForClass(FoundClass);
+					Icon = FSlateIconFinder::FindIconForClass(FoundType);
 				}
 
 				if (Children.Num() == 0)
diff --git a/Engine/Plugins/Animation/GameplayInsights/Source/GameplayInsights/Public/IGameplayProvider.h b/Engine/Plugins/Animation/GameplayInsights/Source/GameplayInsights/Public/IGameplayProvider.h
--- a/Engine/Plugins/Animation/GameplayInsights/Source/GameplayInsights/Public/IGameplayProvider.h
+++ b/Engine/Plugins/Animation/GameplayInsights/Source/GameplayInsights/Public/IGameplayProvider.h
@@ -202,7 +202,12 @@ public:
 	virtual void EnumerateSubobjects(const RewindDebugger::FObjectId& InObjectId, TFunctionRef<void(const RewindDebugger::FObjectId& SubObjectId)> Callback) const = 0;
 	virtual const FObjectPropertyValue * FindPropertyValueFromStorageIndex(uint64 InObjectId, int64 InStorageIndex) const = 0;
 	virtual const FClassInfo* FindClassInfo(uint64 InClassId) const = 0;
-	virtual const UClass* FindClass(uint64 InClassId, bool bSearchBlueprints) const = 0;
+	virtual const UStruct* FindType(uint64 InTypeId, bool bSearchBlueprints) const = 0;
+	UE_DEPRECATED(5.8, "Use FindType returning a UStruct instead.")
+	virtual const UClass* FindClass(const uint64 InClassId, const bool bSearchBlueprints) const final
+	{
+		return Cast<UClass>(FindType(InClassId, bSearchBlueprints));
+	}
 	virtual const FClassInfo* FindClassInfo(const TCHAR* InClassPath) const = 0;
 	virtual bool IsSubClassOf(uint64 InSubClassId, uint64 InParentClassId) const = 0;
 	virtual const FObjectInfo* FindObjectInfo(const uint64 InObjectId) const = 0;
diff --git a/Engine/Plugins/Runtime/StateTree/Source/StateTreeEditorModule/Private/Debugger/StateTreeRewindDebuggerTrack.cpp b/Engine/Plugins/Runtime/StateTree/Source/StateTreeEditorModule/Private/Debugger/StateTreeRewindDebuggerTrack.cpp
--- a/Engine/Plugins/Runtime/StateTree/Source/StateTreeEditorModule/Private/Debugger/StateTreeRewindDebuggerTrack.cpp
+++ b/Engine/Plugins/Runtime/StateTree/Source/StateTreeEditorModule/Private/Debugger/StateTreeRewindDebuggerTrack.cpp
@@ -1,6 +1,7 @@
 // Copyright Epic Games, Inc. All Rights Reserved.
 
 #include "StateTreeRewindDebuggerTrack.h"
+#include "StateTreeEditorSettings.h"
 #include "Debugger/SCompactTreeDebuggerView.h"
 #include "Debugger/SStateTreeFrameEventsView.h"
 #include "Debugger/StateTreeDebugger.h"
@@ -103,9 +104,13 @@ FText FRewindDebuggerTrack::GetDisplayNameInternal() const
 {
 	if (Descriptor.IsValid())
 	{
-		return FText::FromString(FString::Printf(TEXT("%s StateTree (%s)")
-			, *Descriptor->Name
-			, *GetNameSafe(Descriptor->StateTree.Get())));
+		if (UStateTreeEditorSettings::Get().DebuggerTrackNameVerbosity == EStateTreeDebuggerTrackNameVerbosity::Verbose)
+		{
+			return FText::FromString(FString::Printf(TEXT("%s StateTree (%s)")
+				, *Descriptor->Name
+				, *GetNameSafe(Descriptor->StateTree.Get())));
+		}
+		return FText::FromString(GetNameSafe(Descriptor->StateTree.Get()));
 	}
 
 	return LOCTEXT("StateTreeRewindDebuggerTrackName", "StateTreeInstance");
diff --git a/Engine/Plugins/Runtime/StateTree/Source/StateTreeEditorModule/Public/StateTreeEditorSettings.h b/Engine/Plugins/Runtime/StateTree/Source/StateTreeEditorModule/Public/StateTreeEditorSettings.h
--- a/Engine/Plugins/Runtime/StateTree/Source/StateTreeEditorModule/Public/StateTreeEditorSettings.h
+++ b/Engine/Plugins/Runtime/StateTree/Source/StateTreeEditorModule/Public/StateTreeEditorSettings.h
@@ -15,6 +15,16 @@ enum class EStateTreeSaveOnCompile : uint8
 	Always UMETA(DisplayName = "Always"),
 };
 
+UENUM()
+enum class EStateTreeDebuggerTrackNameVerbosity : uint8
+{
+	/** Display only the StateTree asset name (e.g. "MyStateTree"). */
+	Compact UMETA(DisplayName = "Compact"),
+
+	/** Display the instance name and asset name (e.g. "MyInstance StateTree (MyStateTree)"). */
+	Verbose UMETA(DisplayName = "Verbose"),
+};
+
 UCLASS(MinimalAPI, config = EditorPerProjectUserSettings)
 class UStateTreeEditorSettings : public UDeveloperSettings
 {
@@ -29,6 +39,10 @@ class UStateTreeEditorSettings : public UDeveloperSettings
 	UPROPERTY(EditAnywhere, config, Category = "Compiler")
 	EStateTreeSaveOnCompile SaveOnCompile = EStateTreeSaveOnCompile::Never;
 
+	/** Controls how StateTree tracks are labeled in the Rewind Debugger. Compact shows only the asset name; Verbose also includes the instance name. */
+	UPROPERTY(EditAnywhere, config, Category = "Debugger")
+	EStateTreeDebuggerTrackNameVerbosity DebuggerTrackNameVerbosity = EStateTreeDebuggerTrackNameVerbosity::Compact;
+
 	/**
 	 * If enabled, debugger window in the StateTree Asset Editor will display all widgets
 	 * related to the legacy debugger (recording controls, timelines, frame details, etc.).
diff --git a/Engine/Source/Editor/MassEntityEditor/Private/MassEditorStyle.cpp b/Engine/Source/Editor/MassEntityEditor/Private/MassEditorStyle.cpp
--- a/Engine/Source/Editor/MassEntityEditor/Private/MassEditorStyle.cpp
+++ b/Engine/Source/Editor/MassEntityEditor/Private/MassEditorStyle.cpp
@@ -27,6 +27,7 @@ void FMassEntityEditorStyle::Initialize()
 	}
 
 	const FVector2f Icon8x8(8.0f, 8.0f);
+	const FVector2f Icon16x16(16.0f, 16.0f);
 
 	StyleSet = MakeShared<FSlateStyleSet>(GetStyleSetName());
 
@@ -104,6 +105,9 @@ void FMassEntityEditorStyle::Initialize()
 		.SetSelectorFocusedBrush(IMAGE_BRUSH("Common/Selection", Icon8x8, SelectionColor))
 	);
 
+	// Class icons
+	StyleSet->Set("ClassIcon.MassEntityHandle", new IMAGE_BRUSH_SVG("Icons/AssetIcons/MassEntityHandle_16", Icon16x16));
+
 	FSlateStyleRegistry::RegisterSlateStyle(*StyleSet.Get());
 }
 
diff --git a/Engine/Source/Editor/MassEntityEditor/Private/Trace/MassEntityTrack.cpp b/Engine/Source/Editor/MassEntityEditor/Private/Trace/MassEntityTrack.cpp
--- a/Engine/Source/Editor/MassEntityEditor/Private/Trace/MassEntityTrack.cpp
+++ b/Engine/Source/Editor/MassEntityEditor/Private/Trace/MassEntityTrack.cpp
@@ -131,7 +131,7 @@ FEntityTrack::FEntityTrack(const uint64 InEntityId, TSharedRef<FEntityEventCache
 	Icon = FSlateIcon("EditorStyle", "ClassIcon.Object");
 
 	const FMassEntityHandle EntityHandle(FMassEntityHandle::FromNumber(InEntityId));
-	DisplayName = FText::Format(LOCTEXT("EntityName", "Entity {0}"), FText::FromString(LexToString(EntityHandle)));
+	DisplayName = LOCTEXT("EntityTrack", "Composition");
 }
 
 bool FEntityTrack::UpdateInternal()
